Rename SkillSystem module to GameplayAbilities

diff --git a/Engine/Source/Editor/SkillSystemEditor/Private/AttributeDetails.cpp b/Engine/Source/Editor/SkillSystemEditor/Private/AttributeDetails.cpp
--- a/Engine/Source/Editor/SkillSystemEditor/Private/AttributeDetails.cpp
+++ b/Engine/Source/Editor/SkillSystemEditor/Private/AttributeDetails.cpp
@@ -8,7 +8,7 @@
 #include "BlueprintUtilities.h"
 #include "KismetEditorUtilities.h"
 #include "SkillSystemGlobals.h"
-#include "SkillSystemModule.h"
+#include "GameplayAbilitiesModule.h"
 
 #define LOCTEXT_NAMESPACE "AttributeDetailsCustomization"
 
@@ -362,7 +362,7 @@ UCurveTable * FScalableFloatDetails::GetCurveTable()
 
 	if (CurveTable == NULL)
 	{
-		CurveTable = ISkillSystemModule::Get().GetSkillSystemGlobals()->GetGlobalCurveTable();
+		CurveTable = IGameplayAbilitiesModule::Get().GetSkillSystemGlobals()->GetGlobalCurveTable();
 	}
 
 	return CurveTable;
diff --git a/Engine/Source/Editor/SkillSystemEditor/SkillSystemEditor.Build.cs b/Engine/Source/Editor/SkillSystemEditor/SkillSystemEditor.Build.cs
--- a/Engine/Source/Editor/SkillSystemEditor/SkillSystemEditor.Build.cs
+++ b/Engine/Source/Editor/SkillSystemEditor/SkillSystemEditor.Build.cs
@@ -20,7 +20,7 @@ namespace UnrealBuildTool.Rules
 					"CoreUObject",
 					"Engine",
 					"AssetTools",
-					"SkillSystem",
+					"GameplayAbilities",
                     "InputCore",
 					"Slate",
 					"SlateCore",
diff --git a/Engine/Source/Runtime/GameplayAbilities/Classes/AttributeComponent.h b/Engine/Source/Runtime/GameplayAbilities/Classes/AttributeComponent.h
new file mode 100644
--- /dev/null
+++ b/Engine/Source/Runtime/GameplayAbilities/Classes/AttributeComponent.h
@@ -0,0 +1,368 @@
+// Copyright 1998-2013 Epic Games, Inc. All Rights Reserved.
+#pragma once
+
+#include "GameplayEffect.h"
+#include "SkillSystemTypes.h"
+#include "AttributeComponent.generated.h"
+
+GAMEPLAYABILITIES_API DECLARE_LOG_CATEGORY_EXTERN(LogAttributeComponent, Log, All);
+
+DECLARE_MULTICAST_DELEGATE(FAttributeComponentPredictionKeyClear);
+
+USTRUCT()
+struct GAMEPLAYABILITIES_API FAttributeDefaults
+{
+	GENERATED_USTRUCT_BODY()
+
+	UPROPERTY(EditAnywhere, Category="AttributeTest")
+	TSubclassOf<class UAttributeSet> Attributes;
+
+	UPROPERTY(EditAnywhere, BlueprintReadWrite, Category="AttributeTest")
+	class UDataTable*	DefaultStartingTable;
+};
+
+/** 
+ * 
+ */
+UCLASS(ClassGroup=SkillSystem, hidecategories=(Object,LOD,Lighting,Transform,Sockets,TextureStreaming), editinlinenew, meta=(BlueprintSpawnableComponent))
+class GAMEPLAYABILITIES_API UAttributeComponent : public UActorComponent
+{
+	GENERATED_UCLASS_BODY()
+
+	friend FGameplayEffectSpec;
+	friend class ASkillSystemDebugHUD;
+
+	virtual ~UAttributeComponent();
+
+	virtual void InitializeComponent() OVERRIDE;
+
+	/** Finds existing AttributeSet */
+	template <class T >
+	T*	GetSet()
+	{
+		return (T*)GetAttributeSubobject(T::StaticClass());
+	}
+
+	template <class T >
+	T*	GetSetChecked()
+	{
+		return (T*)GetAttributeSubobjectChecked(T::StaticClass());
+	}
+
+	/** Adds a new AttributeSet (initialized to default values) */
+	template <class T >
+	T*  AddSet()
+	{
+		return (T*)GetOrCreateAttributeSubobject(T::StaticClass());
+	}
+
+	/** Adds a new AttributeSet that is a DSO (created by called in their CStor) */
+	template <class T>
+	T*	AddDefaultSubobjectSet(TSubobjectPtr<T> Subobject)
+	{
+		T* Set = Subobject.Get();
+		SpawnedAttributes.Add(Set);
+		return Set;
+	}
+
+	UFUNCTION(BlueprintCallable, Category="Skills")
+	UAttributeSet * InitStats(TSubclassOf<class UAttributeSet> Attributes, const UDataTable* DataTable);
+
+	UFUNCTION(BlueprintCallable, Category="Skills")
+	void ModifyStats(TSubclassOf<class UAttributeSet> Attributes, FDataTableRowHandle ModifierHandle);
+
+	UPROPERTY(EditAnywhere, Category="AttributeTest")
+	TArray<FAttributeDefaults>	DefaultStartingData;
+
+	UPROPERTY(Replicated)
+	TArray<UAttributeSet*>	SpawnedAttributes;
+
+	void SetNumericAttribute(const FGameplayAttribute &Attribute, float NewFloatValue);
+	float GetNumericAttribute(const FGameplayAttribute &Attribute);
+
+	// -- Replication -------------------------------------------------------------------------------------------------
+
+	virtual bool ReplicateSubobjects(class UActorChannel *Channel, class FOutBunch *Bunch, FReplicationFlags *RepFlags) OVERRIDE;
+	
+	virtual void GetSubobjectsWithStableNamesForNetworking(TArray<UObject*>& Objs) OVERRIDE;
+
+	virtual void PostNetReceive() OVERRIDE;
+
+	/**
+	 *	Prediction Keys
+	 *		PredictionKey is a simple system to allow clients to locally predict interactions within the AbilityComponent while having
+	 *		a mechanism for merging results from the server.
+	 *		
+	 *		There are two main cases:
+	 *			-Client predicts something and is wrong. (Mispredictions). He must cleanup stuff that he created or changed.
+	 *			-Client predicts something and is right. 
+	 *			
+	 *			The Client predicts something and asks the server. "Can I do X? PredictionKey=Y".
+	 *			Server can reply immediately and directly with yes or no, and include the PredictionKey with that answer.
+	 *			
+	 *			If No, the client can just immediately cleanup whatever he predicted. This is straightforward.
+	 *			
+	 *			If Yes, there is a small window where the Yes RPC is received but the actor/component properties have not been updated.
+	 *			The client's cannot immediately cleanup his prediction work: he needs to wait for property replication.
+	 *			
+	 *			This is where ReplicatedPredictionKey comes in. 
+	 *			
+	 *			When the server allows a predicted action, he sets ReplicatedPredictionKey to this value. The client then looks
+	 *			at this replicated value. In OnRep_PredictionKey, the client goes through and cleans up all stuff that was predicted
+	 *			and is no longer needed.
+	 *			
+	 *			-Client calls GetNextPredictionKey() at the start of a predicted action. This gives him a key to use for this 'run'.
+	 *			-Client can predict whatever he wants as long as he registers an 'undo' function with GetPredictionKeyDelegate().
+	 *			-This delegate is invoked if the client is denied his prediction, or if he is right and property replication catches up.
+	 *			
+	 *	 
+	 */
+
+	UPROPERTY(ReplicatedUsing=OnRep_PredictionKey)
+	int32	ReplicatedPredictionKey;
+
+	UPROPERTY(transient)
+	int32	LocalPredictionKey;
+
+	int32	GetNextPredictionKey();
+
+	UFUNCTION()
+	void OnRep_PredictionKey();
+
+	TArray<TPair<int32, FAttributeComponentPredictionKeyClear> >	PredictionDelegates;
+	
+
+	FAttributeComponentPredictionKeyClear &	GetPredictionKeyDelegate(int32 PredictionKey);
+
+	/*
+	virtual void OnSubobjectCreatedFromReplication(UObject *NewSubobject) OVERRIDE;
+	
+	virtual void OnSubobjectDestroyFromReplication(UObject *NewSubobject) OVERRIDE;
+	*/
+
+	// ----------------------------------------------------------------------------------------------------------------
+	//
+	//	GameplayEffects
+	//	(maybe should go in a different component?)
+	// ----------------------------------------------------------------------------------------------------------------
+
+	// --------------------------------------------
+	// Primary outward facing API for other systems:
+	// --------------------------------------------
+	FActiveGameplayEffectHandle ApplyGameplayEffectSpecToTarget(OUT FGameplayEffectSpec &GameplayEffect, UAttributeComponent *Target, FModifierQualifier BaseQualifier = FModifierQualifier());
+	FActiveGameplayEffectHandle ApplyGameplayEffectSpecToSelf(const FGameplayEffectSpec &GameplayEffect, FModifierQualifier BaseQualifier = FModifierQualifier());
+
+	UFUNCTION(BlueprintCallable, BlueprintAuthorityOnly, Category = GameplayEffects)
+	bool RemoveActiveGameplayEffect(FActiveGameplayEffectHandle Handle);
+
+	/** This only exists so it can be hooked up to a multicast delegate */
+	void RemoveActiveGameplayEffect_NoReturn(FActiveGameplayEffectHandle Handle)
+	{
+		RemoveActiveGameplayEffect(Handle);
+	}
+
+	float GetGameplayEffectDuration(FActiveGameplayEffectHandle Handle) const;
+
+	float GetGameplayEffectDuration() const;
+
+	// Not happy with this interface but don't see a better way yet. How should outside code (UI, etc) ask things like 'how much is this gameplay effect modifying my damage by'
+	// (most likely we want to catch this on the backend - when damage is applied we can get a full dump/history of how the number got to where it is. But still we may need polling methods like below (how much would my damage be)
+	UFUNCTION(BlueprintCallable, Category = GameplayEffects)
+	float GetGameplayEffectMagnitude(FActiveGameplayEffectHandle Handle, FGameplayAttribute Attribute) const;
+
+	UFUNCTION(BlueprintCallable, Category = GameplayEffects)
+	float GetGameplayEffectMagnitudeByTag(FActiveGameplayEffectHandle InHandle, const FGameplayTag& InTag) const;
+
+	UFUNCTION(BlueprintCallable, Category = GameplayEffects)
+	bool IsGameplayEffectActive(FActiveGameplayEffectHandle InHandle) const;
+
+	// Tags
+	UFUNCTION(BlueprintCallable, Category=GameplayEffects)
+	bool HasAnyTags(FGameplayTagContainer &Tags);
+
+	UFUNCTION(BlueprintCallable, Category = GameplayEffects)
+	bool HasAllTags(FGameplayTagContainer &Tags);
+	
+
+	// --------------------------------------------
+	// Possibly useful but not primary API functions:
+	// --------------------------------------------
+
+	FActiveGameplayEffectHandle ApplyGameplayEffectToTarget(UGameplayEffect *GameplayEffect, UAttributeComponent *Target, float Level = FGameplayEffectLevelSpec::INVALID_LEVEL, FModifierQualifier BaseQualifier = FModifierQualifier());
+
+	UFUNCTION(BlueprintCallable, BlueprintAuthorityOnly, Category = GameplayEffects, meta=(FriendlyName = "ApplyGameplayEffectToSelf"))
+	FActiveGameplayEffectHandle K2_ApplyGameplayEffectToSelf(const UGameplayEffect *GameplayEffect, float Level, AActor *Instigator);
+	
+	FActiveGameplayEffectHandle ApplyGameplayEffectToSelf(const UGameplayEffect *GameplayEffect, float Level, AActor *Instigator, FModifierQualifier BaseQualifier = FModifierQualifier() );
+
+	int32 GetNumActiveGameplayEffect() const;
+
+	void AddDependancyToAttribute(FGameplayAttribute Attribute, const TWeakPtr<FAggregator> InDependant);
+
+	/** Tests if all modifiers in this GameplayEffect will leave the attribute > 0.f */
+	bool CanApplyAttributeModifiers(const UGameplayEffect *GameplayEffect, float Level, AActor *Instigator);
+
+	// Generic 'Get expected magnitude (list) if I was to apply this outgoing or incoming'
+
+	// Get duration or magnitude (list) of active effects
+	//		-Get duration of CD
+	//		-Get magnitude + duration of a movespeed buff
+
+	TArray<float> GetActiveEffectsTimeRemaining(const FActiveGameplayEffectQuery Query) const;
+
+	void OnRestackGameplayEffects();
+
+	// --------------------------------------------
+	// Temp / Debug
+	// --------------------------------------------
+
+	void TEMP_ApplyActiveGameplayEffects();
+	
+	void PrintAllGameplayEffects() const;
+
+	void PushGlobalCurveOveride(UCurveTable *OverrideTable)
+	{
+		if (OverrideTable)
+		{
+			GlobalCurveDataOverride.Overrides.Push(OverrideTable);
+		}
+	}
+
+	// ----------------------------------------------------------------------------------------------------------------
+
+	/**
+	 *	GameplayAbilities
+	 *	
+	 *	The role of the AttributeComponent wrt Abilities is to provide:
+	 *		-Management of abiltiy instances (whether per actor or per execution instance).
+	 *			-Someone *has* to keep track of these instances.
+	 *			-Non instanced abilities *could* be executed without any ability stuff in attribute component.
+	 *				They should be able to operate on an GameplayAbilityActorInfo + GameplayAbility.
+	 *		
+	 *	As convenience it may provide some other features:
+	 *		-Some basic input binding (whether instanced or non instanced abilities).
+	 *		-Concepts like "this component has these abilities
+	 *	
+	 */
+
+	class UGameplayAbility * CreateNewInstanceOfAbility(UGameplayAbility* Ability);
+
+	void CancelAbilitiesWithTags(const FGameplayTagContainer Tags, const FGameplayAbilityActorInfo* ActorInfo, const FGameplayAbilityActivationInfo ActivationInfo, UGameplayAbility* Ignore);
+	
+	/** References to non-replicating abilities that we instanced. We need to keep these refernces to avoid GC */
+	UPROPERTY()
+	TArray<class UGameplayAbility *>	NonReplicatedInstancedAbilities;
+
+	/** References to replicating abilities that we instanced. We need to keep these refernces to avoid GC */
+	UPROPERTY(Replicated)
+	TArray<class UGameplayAbility *>	ReplicatedInstancedAbilities;
+
+	/**
+	 *	The abilities we can activate. 
+	 *		-This will include CDOs for non instanced abilities and per-execution instanced abilities. 
+	 *		-Actor-instanced abilities will be the actual instance (not CDO)
+	 *		
+	 *	This array is not vital for things to work. It is a convenience thing for 'giving abilties to the actor'. But abilities could also work on things
+	 *	without an attribute component. For example an ability could be written to execute on a StaticMeshActor. As long as the ability doesn't require 
+	 *	instancing or anything else that the AttributeComponent would provide, then it doesn't need the component to function.
+	 */ 
+	UPROPERTY(Replicated, BlueprintReadOnly, Category = "Abilities")
+	TArray<class UGameplayAbility *>	ActivatableAbilities;
+	
+
+	UFUNCTION(Server, WithValidation, Reliable)
+	void	ServerTryActivateAbility(class UGameplayAbility* AbilityToActivate, int32 PredictionKey);
+
+	UFUNCTION(Client, Reliable)
+	void	ClientActivateAbilityFailed(class UGameplayAbility* AbilityToActivate, int32 PredictionKey);
+
+	UFUNCTION(Client, Reliable)
+	void	ClientActivateAbilitySucceed(class UGameplayAbility* AbilityToActivate, int32 PredictionKey);
+
+	// ----------------------------------------------------------------------------------------------------------------
+
+	/** This is probably temp for testing. We want to think a bit more on the API for outside stuff activating abilities */
+	UFUNCTION(BlueprintCallable, Category = "Abilities")
+	bool	ActivateAbility(int32 AbilityIdx);
+
+	/** There needs to be a concept of an animating ability. Only one may exist at a time. New requestrs can be queued up, overridden, or ignored. */
+	UPROPERTY()
+	class UGameplayAbility *	AnimatingAbility;
+
+	void	MontageBranchPoint_AbilityDecisionStop();
+
+	void	MontageBranchPoint_AbilityDecisionStart();
+
+	
+
+	// -----------------------------------------------------------------------------
+
+	/** Cached off data about the owning actor that abilities will need to frequently access (movement component, mesh component, anim instance, etc) */
+	TSharedPtr<FGameplayAbilityActorInfo>	AbilityActorInfo;
+
+	void InitAbilityActorInfo();
+
+private:
+
+	// --------------------------------------------
+	// Important internal functions
+	// --------------------------------------------
+
+	void ExecutePeriodicEffect(FActiveGameplayEffectHandle	Handle);
+
+	void ExecuteGameplayEffect(const FGameplayEffectSpec &Spec, const FModifierQualifier &QualifierContext);
+
+	void CheckDurationExpired(FActiveGameplayEffectHandle Handle);
+
+	// --------------------------------------------
+	// Internal Utility / helper functions
+	// --------------------------------------------
+
+	bool AreGameplayEffectApplicationRequirementsSatisfied(const class UGameplayEffect* EffectToAdd, FGameplayEffectInstigatorContext& Instigator) const;
+
+	bool IsOwnerActorAuthoritative() const;
+
+	void OnAttributeGameplayEffectSpecExected(const FGameplayAttribute &Attribute, const struct FGameplayEffectSpec &Spec, struct FGameplayModifierEvaluatedData &Data);
+
+	const FGlobalCurveDataOverride * GetCurveDataOverride() const
+	{
+		// only return data if we have overrides. NULL if we dont.
+		return (GlobalCurveDataOverride.Overrides.Num() > 0 ? &GlobalCurveDataOverride : NULL);
+	}
+
+	FGlobalCurveDataOverride	GlobalCurveDataOverride;
+	
+
+	// --------------------------------------------
+	
+	UPROPERTY(ReplicatedUsing=OnRep_GameplayEffects)
+	FActiveGameplayEffectsContainer	ActiveGameplayEffects;
+
+	UFUNCTION()
+	void OnRep_GameplayEffects();
+
+	// ---------------------------------------------
+
+	UFUNCTION(NetMulticast, unreliable)
+	void NetMulticast_InvokeGameplayCueExecuted(const FGameplayEffectSpec Spec);
+
+	void InvokeGameplayCueExecute(const FGameplayEffectSpec &Spec);
+	void InvokeGameplayCueActivated(const FGameplayEffectSpec &Spec);
+	void InvokeGameplayCueAdded(const FGameplayEffectSpec &Spec);
+	void InvokeGameplayCueRemoved(const FGameplayEffectSpec &Spec);
+
+	// ---------------------------------------------
+
+	
+
+protected:
+
+	virtual void OnRegister() OVERRIDE;
+
+	UAttributeSet*	GetAttributeSubobject(const TSubclassOf<UAttributeSet> AttributeClass) const;
+	UAttributeSet*	GetAttributeSubobjectChecked(const TSubclassOf<UAttributeSet> AttributeClass) const;
+	UAttributeSet*	GetOrCreateAttributeSubobject(const TSubclassOf<UAttributeSet> AttributeClass);
+
+	friend struct FActiveGameplayEffectsContainer;
+	friend struct FActiveGameplayEffect;
+};
diff --git a/Engine/Source/Runtime/GameplayAbilities/Classes/GameplayEffectExtension.h b/Engine/Source/Runtime/GameplayAbilities/Classes/GameplayEffectExtension.h
new file mode 100644
--- /dev/null
+++ b/Engine/Source/Runtime/GameplayAbilities/Classes/GameplayEffectExtension.h
@@ -0,0 +1,45 @@
+// Copyright 1998-2013 Epic Games, Inc. All Rights Reserved.
+
+#pragma once
+
+#include "GameplayTagContainer.h"
+#include "GameplayTagAssetInterface.h"
+#include "AttributeSet.h"
+#include "GameplayEffect.h"
+#include "GameplayEffectExtension.generated.h"
+
+struct FGameplayEffectModCallbackData
+{
+	FGameplayEffectModCallbackData(const FGameplayEffectSpec & InEffectSpec, const FModifierSpec & InModifierSpec, FGameplayModifierEvaluatedData & InEvaluatedData, UAttributeComponent & InTarget)
+		: EffectSpec(InEffectSpec)
+		, ModifierSpec(InModifierSpec)
+		, EvaluatedData(InEvaluatedData)
+		, Target(InTarget)
+	{
+
+	}
+
+	const struct FGameplayEffectSpec &		EffectSpec;		// The spec that the mod came from
+	const struct FModifierSpec &			ModifierSpec;	// The mod we are going to apply
+	struct FGameplayModifierEvaluatedData & EvaluatedData;	// The 'flat'/computed data to be applied to the target
+
+	class UAttributeComponent &Target;		// Target we intend to apply to
+};
+
+UCLASS(BlueprintType)
+class GAMEPLAYABILITIES_API UGameplayEffectExtension: public UObject
+{
+	GENERATED_UCLASS_BODY()
+
+public:
+
+	virtual void PreGameplayEffectExecute(const FGameplayModifierEvaluatedData &SelfData, FGameplayEffectModCallbackData &Data) const
+	{
+
+	}
+
+	virtual void PostGameplayEffectExecute(const FGameplayModifierEvaluatedData &SelfData, const FGameplayEffectModCallbackData &Data) const
+	{
+
+	}
+};
\ No newline at end of file
diff --git a/Engine/Source/Runtime/GameplayAbilities/Classes/SkillSystemGlobals.h b/Engine/Source/Runtime/GameplayAbilities/Classes/SkillSystemGlobals.h
new file mode 100644
--- /dev/null
+++ b/Engine/Source/Runtime/GameplayAbilities/Classes/SkillSystemGlobals.h
@@ -0,0 +1,60 @@
+// Copyright 1998-2014 Epic Games, Inc. All Rights Reserved.
+
+#pragma once
+
+#include "SkillSystemGlobals.generated.h"
+
+class AActor;
+class UAttributeComponent;
+struct FGameplayAbilityActorInfo;
+
+/** Holds global data for the skill system. Can be configured per project via config file */
+UCLASS(config=Game)
+class GAMEPLAYABILITIES_API USkillSystemGlobals : public UObject
+{
+	GENERATED_UCLASS_BODY()
+
+	/** Holds all of the valid gameplay-related tags that can be applied to assets */
+	UPROPERTY(config)
+	FString GlobalCurveTableName;
+
+	/** Holds information about the valid attributes' min and max values and stacking rules */
+	UPROPERTY(config)
+	FString GlobalAttributeDataTableName;
+
+	class UCurveTable *	GetGlobalCurveTable();
+
+	class UDataTable * GetGlobalAttributeDataTable();
+
+	void AutomationTestOnly_SetGlobalCurveTable(class UCurveTable *InTable)
+	{
+		GlobalCurveTable = InTable;
+	}
+
+	void AutomationTestOnly_SetGlobalAttributeDataTable(class UDataTable *InTable)
+	{
+		GlobalAttributeDataTable = InTable;
+	}
+
+	static USkillSystemGlobals&	Get();
+
+	/**
+	 *	Games may want to override this in a UMyProjectSkillSystemsGlobals class and provide
+	 *	a faster lookup for attribute compnents (E.g. Cast<MyProjPawn>(Actor)->AttributeComponent; etc)
+	 */
+	virtual UAttributeComponent * GetAttributeComponentFromActor(AActor *Actor) const;
+
+	virtual FGameplayAbilityActorInfo * AllocAbilityActorInfo(AActor *Actor) const;
+
+private:
+
+	class UCurveTable* GlobalCurveTable;
+
+	class UDataTable* GlobalAttributeDataTable;
+
+#if WITH_EDITOR
+	void OnCurveTableReimported(UObject* InObject);
+	void OnDataTableReimported(UObject* InObject);
+#endif
+
+};
diff --git a/Engine/Source/Runtime/GameplayAbilities/Classes/SkillSystemTestAttributeSet.h b/Engine/Source/Runtime/GameplayAbilities/Classes/SkillSystemTestAttributeSet.h
new file mode 100644
--- /dev/null
+++ b/Engine/Source/Runtime/GameplayAbilities/Classes/SkillSystemTestAttributeSet.h
@@ -0,0 +1,65 @@
+// Copyright 1998-2014 Epic Games, Inc. All Rights Reserved.
+
+#pragma once
+#include "AttributeSet.h"
+#include "SkillSystemTestAttributeSet.generated.h"
+
+UCLASS(Blueprintable, BlueprintType)
+class GAMEPLAYABILITIES_API USkillSystemTestAttributeSet : public UAttributeSet
+{
+	GENERATED_UCLASS_BODY()
+
+	UPROPERTY(EditAnywhere, BlueprintReadWrite, Replicated, Category = "AttributeTest", meta = (HideFromModifiers))			// You can't make a GameplayEffect modify Health Directly (go through)
+	float	MaxHealth;
+
+	UPROPERTY(EditAnywhere, BlueprintReadWrite, Replicated, Category = "AttributeTest", meta = (HideFromModifiers))			// You can't make a GameplayEffect modify Health Directly (go through)
+	float	Health;
+
+	UPROPERTY(EditAnywhere, BlueprintReadWrite, Replicated, Category = "AttributeTest")
+	float	Mana;
+
+	UPROPERTY(EditAnywhere, BlueprintReadWrite, Replicated, Category = "AttributeTest")
+	float	MaxMana;
+
+	/** This Damage is just used for applying negative health mods. Its not a 'persistent' attribute. */
+	UPROPERTY(EditAnywhere, BlueprintReadWrite, Category = "AttributeTest", meta = (HideFromLevelInfos))		// You can't make a GameplayEffect 'powered' by Damage (Its transient)
+	float	Damage;
+
+	/** This Attribute is the actual spell damage for this actor. It will power spell based GameplayEffects */
+	UPROPERTY(EditAnywhere, BlueprintReadWrite, Replicated, Category = "AttributeTest")
+	float	SpellDamage;
+
+	/** This Attribute is the actual physical damage for this actor. It will power physical based GameplayEffects */
+	UPROPERTY(EditAnywhere, BlueprintReadWrite, Replicated, Category = "AttributeTest")
+	float	PhysicalDamage;
+
+	UPROPERTY(EditAnywhere, BlueprintReadWrite, Replicated, Category = "AttributeTest")
+	float	CritChance;
+
+	UPROPERTY(EditAnywhere, BlueprintReadWrite, Replicated, Category = "AttributeTest")
+	float	CritMultiplier;
+
+	UPROPERTY(EditAnywhere, BlueprintReadWrite, Replicated, Category = "AttributeTest")
+	float	ArmorDamageReduction;
+
+	UPROPERTY(EditAnywhere, BlueprintReadWrite, Replicated, Category = "AttributeTest")
+	float	DodgeChance;
+
+	UPROPERTY(EditAnywhere, BlueprintReadWrite, Replicated, Category = "AttributeTest")
+	float	LifeSteal;
+
+	UPROPERTY(EditAnywhere, BlueprintReadWrite, Replicated, Category = "AttributeTest")
+	float	Strength;
+
+	UPROPERTY(EditAnywhere, BlueprintReadWrite, Category = "AttributeTest")
+	float	StackingAttribute1;
+
+	UPROPERTY(EditAnywhere, BlueprintReadWrite, Category = "AttributeTest")
+	float	StackingAttribute2;
+
+	UPROPERTY(EditAnywhere, BlueprintReadWrite, Category = "AttributeTest")
+	float	NoStackAttribute;
+
+	virtual void PreAttributeModify(struct FGameplayEffectModCallbackData &Data) OVERRIDE;
+	virtual void PostAttributeModify(const struct FGameplayEffectModCallbackData &Data) OVERRIDE;
+};
diff --git a/Engine/Source/Runtime/GameplayAbilities/Classes/SkillSystemTestPawn.h b/Engine/Source/Runtime/GameplayAbilities/Classes/SkillSystemTestPawn.h
new file mode 100644
--- /dev/null
+++ b/Engine/Source/Runtime/GameplayAbilities/Classes/SkillSystemTestPawn.h
@@ -0,0 +1,39 @@
+// Copyright 1998-2014 Epic Games, Inc. All Rights Reserved.
+
+#pragma once
+#include "GameplayCueInterface.h"
+#include "GameplayCueView.h"
+#include "SkillSystemTestPawn.generated.h"
+
+UCLASS(Blueprintable, BlueprintType)
+class GAMEPLAYABILITIES_API ASkillSystemTestPawn : public ADefaultPawn, public IGameplayCueInterface
+{
+	GENERATED_UCLASS_BODY()
+
+	virtual void PostInitializeComponents() OVERRIDE;
+
+	UFUNCTION()
+	virtual void GameplayCueActivated(const FGameplayTagContainer & GameplayCueTags, float NormalizedMagnitude, const FGameplayEffectInstigatorContext InstigatorContext) OVERRIDE;
+
+	UFUNCTION()
+	virtual void GameplayCueExecuted(const FGameplayTagContainer & GameplayCueTags, float NormalizedMagnitude, const FGameplayEffectInstigatorContext InstigatorContext) OVERRIDE;
+
+	UFUNCTION()
+	virtual void GameplayCueAdded(const FGameplayTagContainer & GameplayCueTags, float NormalizedMagnitude, const FGameplayEffectInstigatorContext InstigatorContext) OVERRIDE;
+
+	UFUNCTION()
+	virtual void GameplayCueRemoved(const FGameplayTagContainer & GameplayCueTags, float NormalizedMagnitude, const FGameplayEffectInstigatorContext InstigatorContext) OVERRIDE;
+
+
+	/** DefaultPawn collision component */
+	UPROPERTY(Category = SkillSystem, VisibleAnywhere, BlueprintReadOnly)
+	TSubobjectPtr<class UAttributeComponent> AttributeComponent;
+
+	UPROPERTY(EditDefaultsOnly, Category=GameplayEffects)
+	FGameplayCueHandler	GameplayCueHandler;
+
+	//UPROPERTY(EditDefaultsOnly, Category=GameplayEffects)
+	//UGameplayAbilitySet * DefaultAbilitySet;
+
+	static FName AttributeComponentName;
+};
\ No newline at end of file
diff --git a/Engine/Source/Runtime/GameplayAbilities/GameplayAbilities.Build.cs b/Engine/Source/Runtime/GameplayAbilities/GameplayAbilities.Build.cs
new file mode 100644
--- /dev/null
+++ b/Engine/Source/Runtime/GameplayAbilities/GameplayAbilities.Build.cs
@@ -0,0 +1,55 @@
+// Copyright 1998-2013 Epic Games, Inc. All Rights Reserved.
+
+namespace UnrealBuildTool.Rules
+{
+	public class GameplayAbilities : ModuleRules
+	{
+        public GameplayAbilities(TargetInfo Target)
+		{
+            PrivateIncludePaths.Add("../../Engine/Source/Runtime/GameplayAbilities/Private");
+            
+			PublicIncludePaths.AddRange(
+				new string[] {
+					// ... add public include paths required here ...
+				}
+				);
+
+			PrivateIncludePaths.AddRange(
+				new string[] {
+					// ... add other private include paths required here ...
+				}
+				);
+
+			PublicDependencyModuleNames.AddRange(
+				new string[]
+				{
+					"Core",
+					"CoreUObject",
+					"Engine",
+                    "GameplayTags",
+					// ... add other public dependencies that you statically link with here ...
+				}
+				);
+
+			PrivateDependencyModuleNames.AddRange(
+				new string[]
+				{
+					// ... add private dependencies that you statically link with here ...
+				}
+				);
+
+			DynamicallyLoadedModuleNames.AddRange(
+				new string[]
+				{
+					// ... add any modules that your module loads dynamically here ...
+				}
+				);
+
+			if (UEBuildConfiguration.bBuildEditor == true)
+			{
+				PrivateDependencyModuleNames.Add("UnrealEd");
+                PrivateDependencyModuleNames.Add("GameplayTagsEditor");
+			}
+		}
+	}
+}
\ No newline at end of file
diff --git a/Engine/Source/Runtime/GameplayAbilities/Private/AttributeComponent.cpp b/Engine/Source/Runtime/GameplayAbilities/Private/AttributeComponent.cpp
new file mode 100644
--- /dev/null
+++ b/Engine/Source/Runtime/GameplayAbilities/Private/AttributeComponent.cpp
@@ -0,0 +1,738 @@
+// Copyright 1998-2013 Epic Games, Inc. All Rights Reserved.
+// ActorComponent.cpp: Actor component implementation.
+
+#include "SkillSystemModulePrivatePCH.h"
+#include "AttributeComponent.h"
+#include "GameplayCueInterface.h"
+#include "Abilities/GameplayAbility.h"
+
+
+#include "Net/UnrealNetwork.h"
+#include "Engine/ActorChannel.h"
+#include "MessageLog.h"
+#include "UObjectToken.h"
+#include "MapErrors.h"
+
+DEFINE_LOG_CATEGORY(LogAttributeComponent);
+
+#define LOCTEXT_NAMESPACE "AttributeComponent"
+
+
+int32 DebugGameplayCues = 0;
+static FAutoConsoleVariableRef CVarDebugGameplayCues(
+	TEXT("SkillSystem.DebugGameplayCues"),
+	DebugGameplayCues,
+	TEXT("Enables Debugging for GameplayCue events"),
+	ECVF_Default
+	);
+
+/** Enable to log out all render state create, destroy and updatetransform events */
+#define LOG_RENDER_STATE 0
+
+UAttributeComponent::UAttributeComponent(const class FPostConstructInitializeProperties& PCIP)
+	: Super(PCIP)
+{
+	bWantsInitializeComponent = true;
+
+	PrimaryComponentTick.TickGroup = TG_DuringPhysics;
+	PrimaryComponentTick.bStartWithTickEnabled = true; // FIXME! Just temp until timer manager figured out
+	PrimaryComponentTick.bCanEverTick = true;
+	
+	ActiveGameplayEffects.Owner = this;
+
+	SetIsReplicated(true);
+	ReplicatedPredictionKey = 0;
+	LocalPredictionKey = 0;
+}
+
+UAttributeComponent::~UAttributeComponent()
+{
+	ActiveGameplayEffects.PreDestroy();
+}
+
+UAttributeSet* UAttributeComponent::InitStats(TSubclassOf<class UAttributeSet> Attributes, const UDataTable* DataTable)
+{
+	UAttributeSet* AttributeObj = NULL;
+	if (Attributes)
+	{
+		AttributeObj = GetOrCreateAttributeSubobject(Attributes);
+		if (AttributeObj && DataTable)
+		{
+			AttributeObj->InitFromMetaDataTable(DataTable);
+		}
+	}
+	return AttributeObj;
+}
+
+void UAttributeComponent::ModifyStats(TSubclassOf<class UAttributeSet> AttributeClass, FDataTableRowHandle ModifierHandle)
+{
+	UAttributeSet *Attributes = Cast<UAttributeSet>(GetAttributeSubobject(AttributeClass));
+	if (Attributes)
+	{
+		FAttributeModifierTest *ModifierData = ModifierHandle.GetRow<FAttributeModifierTest>();
+		if (ModifierData)
+		{
+			ModifierData->ApplyModifier(Attributes);
+		}
+	}
+}
+
+UAttributeSet*	UAttributeComponent::GetOrCreateAttributeSubobject(const TSubclassOf<UAttributeSet> AttributeClass)
+{
+	AActor *OwningActor = GetOwner();
+	UAttributeSet *MyAttributes  = NULL;
+	if (OwningActor && AttributeClass)
+	{
+		MyAttributes = GetAttributeSubobject(AttributeClass);
+		if (!MyAttributes)
+		{
+			MyAttributes = ConstructObject<UAttributeSet>(AttributeClass, OwningActor);
+			SpawnedAttributes.AddUnique(MyAttributes);
+		}
+	}
+
+	return MyAttributes;
+}
+
+UAttributeSet* UAttributeComponent::GetAttributeSubobjectChecked(const TSubclassOf<UAttributeSet> AttributeClass) const
+{
+	UAttributeSet *Set = GetAttributeSubobject(AttributeClass);
+	check(Set);
+	return Set;
+}
+
+UAttributeSet* UAttributeComponent::GetAttributeSubobject(const TSubclassOf<UAttributeSet> AttributeClass) const
+{
+	for (UAttributeSet * Set : SpawnedAttributes)
+	{
+		if (Set && Set->IsA(AttributeClass))
+		{
+			return Set;
+		}
+	}
+	return NULL;
+}
+
+void UAttributeComponent::OnRegister()
+{
+	Super::OnRegister();
+
+	// Init starting data
+	for (int32 i=0; i < DefaultStartingData.Num(); ++i)
+	{
+		if (DefaultStartingData[i].Attributes && DefaultStartingData[i].DefaultStartingTable)
+		{
+			UAttributeSet* Attributes = GetOrCreateAttributeSubobject(DefaultStartingData[i].Attributes);
+			Attributes->InitFromMetaDataTable(DefaultStartingData[i].DefaultStartingTable);
+		}
+	}
+}
+
+// ---------------------------------------------------------
+
+bool UAttributeComponent::AreGameplayEffectApplicationRequirementsSatisfied(const class UGameplayEffect* EffectToAdd, FGameplayEffectInstigatorContext& Instigator) const
+{
+	bool bReqsSatisfied = false;
+	if (EffectToAdd)
+	{
+		// Collect gameplay tags from instigator and target to see if requirements are satisfied
+		FGameplayTagContainer InstigatorTags;
+		Instigator.GetOwnedGameplayTags(InstigatorTags);
+
+		FGameplayTagContainer TargetTags;
+		IGameplayTagAssetInterface* OwnerGTA = InterfaceCast<IGameplayTagAssetInterface>(GetOwner());
+		if (OwnerGTA)
+		{
+			OwnerGTA->GetOwnedGameplayTags(TargetTags);
+		}
+
+		bReqsSatisfied = EffectToAdd->AreApplicationTagRequirementsSatisfied(InstigatorTags, TargetTags);
+	}
+
+	return bReqsSatisfied;
+}
+
+// ---------------------------------------------------------
+
+bool UAttributeComponent::IsOwnerActorAuthoritative() const
+{
+	return !IsNetSimulating();
+}
+
+void UAttributeComponent::SetNumericAttribute(const FGameplayAttribute &Attribute, float NewFloatValue)
+{
+	UAttributeSet * AttributeSet = GetAttributeSubobjectChecked(Attribute.GetAttributeSetClass());
+	Attribute.SetNumericValueChecked(NewFloatValue, AttributeSet);
+}
+
+float UAttributeComponent::GetNumericAttribute(const FGameplayAttribute &Attribute)
+{
+	UAttributeSet * AttributeSet = GetAttributeSubobjectChecked(Attribute.GetAttributeSetClass());
+	return Attribute.GetNumericValueChecked(AttributeSet);
+}
+
+/** This is a helper function used in automated testing, I'm not sure how useful it will be to gamecode or blueprints */
+FActiveGameplayEffectHandle UAttributeComponent::ApplyGameplayEffectToTarget(UGameplayEffect *GameplayEffect, UAttributeComponent *Target, float Level, FModifierQualifier BaseQualifier)
+{
+	check(GameplayEffect);
+
+	FGameplayEffectSpec	Spec(GameplayEffect, GetOwner(),  Level, GetCurveDataOverride());
+	
+	return ApplyGameplayEffectSpecToTarget(Spec, Target, BaseQualifier);
+}
+
+/** Helper function since we can't have default/optional values for FModifierQualifier in K2 function */
+FActiveGameplayEffectHandle UAttributeComponent::K2_ApplyGameplayEffectToSelf(const UGameplayEffect *GameplayEffect, float Level, AActor *Instigator)
+{
+	return ApplyGameplayEffectToSelf(GameplayEffect, Level, Instigator);
+}
+
+/** This is a helper function - it seems like this will be useful as a blueprint interface at the least, but Level parameter may need to be expanded */
+FActiveGameplayEffectHandle UAttributeComponent::ApplyGameplayEffectToSelf(const UGameplayEffect *GameplayEffect, float Level, AActor *Instigator, FModifierQualifier BaseQualifier)
+{
+	check(GameplayEffect);
+
+	FGameplayEffectSpec	Spec(GameplayEffect, Instigator, Level, GetCurveDataOverride());
+
+	return ApplyGameplayEffectSpecToSelf(Spec, BaseQualifier);
+}
+
+float UAttributeComponent::GetGameplayEffectMagnitudeByTag(FActiveGameplayEffectHandle InHandle, const FGameplayTag& InTag) const
+{
+	return ActiveGameplayEffects.GetGameplayEffectMagnitudeByTag(InHandle, InTag);
+}
+
+int32 UAttributeComponent::GetNumActiveGameplayEffect() const
+{
+	return ActiveGameplayEffects.GetNumGameplayEffects();
+}
+
+bool UAttributeComponent::IsGameplayEffectActive(FActiveGameplayEffectHandle InHandle) const
+{
+	return ActiveGameplayEffects.IsGameplayEffectActive(InHandle);
+}
+
+bool UAttributeComponent::HasAnyTags(FGameplayTagContainer &Tags)
+{
+	// Fixme: we could aggregate our current tags into a map to avoid this type of iteration
+	return ActiveGameplayEffects.HasAnyTags(Tags);
+}
+
+bool UAttributeComponent::HasAllTags(FGameplayTagContainer &Tags)
+{
+	return ActiveGameplayEffects.HasAllTags(Tags);
+}
+
+void UAttributeComponent::TEMP_ApplyActiveGameplayEffects()
+{
+	for (int32 idx=0; idx < ActiveGameplayEffects.GameplayEffects.Num(); ++idx)
+	{
+		FActiveGameplayEffect& ActiveEffect = ActiveGameplayEffects.GameplayEffects[idx];
+
+		ExecuteGameplayEffect(ActiveEffect.Spec, FModifierQualifier().IgnoreHandle(ActiveEffect.Handle));
+
+		SKILL_LOG(Log, TEXT("ActiveEffect[%d] %s - Duration: %.2f]"), idx, *ActiveEffect.Spec.ToSimpleString(), ActiveEffect.Spec.GetDuration());
+	}
+}
+
+FActiveGameplayEffectHandle UAttributeComponent::ApplyGameplayEffectSpecToTarget(OUT FGameplayEffectSpec &Spec, UAttributeComponent *Target, FModifierQualifier BaseQualifier)
+{
+	// Apply outgoing Effects to the Spec.
+	// Outgoing immunity may stop the outgoing effect from being applied to the target
+	if (ActiveGameplayEffects.ApplyActiveEffectsTo(Spec, FModifierQualifier(BaseQualifier).Type(EGameplayMod::OutgoingGE)))
+	{
+		return Target->ApplyGameplayEffectSpecToSelf(Spec, BaseQualifier);
+	}
+
+	return FActiveGameplayEffectHandle();
+}
+
+FActiveGameplayEffectHandle UAttributeComponent::ApplyGameplayEffectSpecToSelf(const FGameplayEffectSpec &Spec, FModifierQualifier BaseQualifier)
+{
+	// Temp, only non instant, non periodic GEs can be predictive
+	// Effects with other effects may be a mix so go with non-predictive
+	check((BaseQualifier.PredictionKey() == 0) || (Spec.GetDuration() != UGameplayEffect::INSTANT_APPLICATION && Spec.GetPeriod() == UGameplayEffect::NO_PERIOD));
+
+	// check if the effect being applied actually succeeds
+	float ChanceToApply = Spec.GetChanceToApplyToTarget();
+	if ((ChanceToApply < 1.f - SMALL_NUMBER) && (FMath::FRand() > ChanceToApply))
+	{
+		return FActiveGameplayEffectHandle();
+	}
+
+	// Make sure we create our copy of the spec in the right place first...
+	FActiveGameplayEffectHandle	MyHandle;
+	bool bInvokeGameplayCueApplied = false;
+
+	FGameplayEffectSpec* OurCopyOfSpec = NULL;
+	TSharedPtr<FGameplayEffectSpec> StackSpec;
+	{
+	
+		float Duration = Spec.GetDuration();
+
+		if (Duration != UGameplayEffect::INSTANT_APPLICATION)
+		{
+			// recalculating stacking needs to come before creating the new effect
+			if (Spec.GetStackingType() != EGameplayEffectStackingPolicy::Unlimited)
+			{
+				ActiveGameplayEffects.StacksNeedToRecalculate();
+			}
+			FActiveGameplayEffect &NewActiveEffect = ActiveGameplayEffects.CreateNewActiveGameplayEffect(Spec, BaseQualifier.PredictionKey());
+			MyHandle = NewActiveEffect.Handle;
+			OurCopyOfSpec = &NewActiveEffect.Spec;
+
+			bInvokeGameplayCueApplied = true;
+		}
+		
+		if (!OurCopyOfSpec)
+		{
+			StackSpec = TSharedPtr<FGameplayEffectSpec>(new FGameplayEffectSpec(Spec));
+			OurCopyOfSpec = StackSpec.Get();
+		}
+
+		// Do a 1st order copy of the spec so that we can modify it
+		// (the one passed in is owned by the caller, we can't apply our incoming GEs to it)
+		// Note that at this point the spec has a bunch of modifiers. Those modifiers may
+		// have other modifiers. THOSE modifiers may or may not be copies of whatever.
+		//
+		// In theory, we don't modify 2nd order modifiers after they are 'attached'
+		// Long complex chains can be created but we never say 'Modify a GE that is modding another GE'
+		OurCopyOfSpec->MakeUnique();
+	}
+
+	// Now that we have our own copy, apply our GEs that modify IncomingGEs
+	if (!ActiveGameplayEffects.ApplyActiveEffectsTo(*OurCopyOfSpec, FModifierQualifier(BaseQualifier).Type(EGameplayMod::IncomingGE).IgnoreHandle(MyHandle)))
+	{
+		// We're immune to this effect
+		return FActiveGameplayEffectHandle();
+	}	
+	
+	// Now that we have the final version of this effect, actually apply it if its going to be hanging around
+	if (Spec.GetDuration() != UGameplayEffect::INSTANT_APPLICATION )
+	{
+		if (Spec.GetPeriod() == UGameplayEffect::NO_PERIOD)
+		{
+			ActiveGameplayEffects.ApplySpecToActiveEffectsAndAttributes(*OurCopyOfSpec, FModifierQualifier(BaseQualifier).IgnoreHandle(MyHandle));
+		}
+	}
+	
+	// We still probably want to apply tags and stuff even if instant?
+	if (bInvokeGameplayCueApplied)
+	{
+		// We both added and activated the GameplayCue here.
+		// On the client, who will invoke the gameplay cue from an OnRep, he will need to look at the StartTime to determine
+		// if the Cue was actually added+activated or just added (due to relevancy)
+
+		// Fixme: what if we wanted to scale Cue magnitude based on damage? E.g, scale an cue effect when the GE is buffed?
+		InvokeGameplayCueAdded(*OurCopyOfSpec);
+		InvokeGameplayCueActivated(*OurCopyOfSpec);
+	}
+	
+	// Execute the GE at least once (if instant, this will execute once and be done. If persistent, it was added to ActiveGameplayEffects above)
+	
+	// Execute if this is an instant application effect
+	if (Spec.GetDuration() == UGameplayEffect::INSTANT_APPLICATION)
+	{
+		ExecuteGameplayEffect(*OurCopyOfSpec, FModifierQualifier(BaseQualifier).IgnoreHandle(MyHandle));
+	}
+
+	if (Spec.GetPeriod() != UGameplayEffect::NO_PERIOD && Spec.TargetEffectSpecs.Num() > 0)
+	{
+		SKILL_LOG(Warning, TEXT("%s is periodic but also applies GameplayEffects to its target. GameplayEffects will only be applied once, not every period."), *Spec.Def->GetPathName());
+	}
+	// todo: this is ignoring the returned handles, should we put them into a TArray and return all of the handles?
+	for (const TSharedRef<FGameplayEffectSpec> TargetSpec : Spec.TargetEffectSpecs)
+	{
+		ApplyGameplayEffectSpecToSelf(TargetSpec.Get(), BaseQualifier);
+	}
+
+	return MyHandle;
+}
+
+void UAttributeComponent::ExecutePeriodicEffect(FActiveGameplayEffectHandle	Handle)
+{
+	ActiveGameplayEffects.ExecutePeriodicGameplayEffect(Handle);
+}
+
+void UAttributeComponent::ExecuteGameplayEffect(const FGameplayEffectSpec &Spec, const FModifierQualifier &QualifierContext)
+{
+	// Should only ever execute effects that are instant application or periodic application
+	// Effects with no period and that aren't instant application should never be executed
+	check( (Spec.GetDuration() == UGameplayEffect::INSTANT_APPLICATION || Spec.GetPeriod() != UGameplayEffect::NO_PERIOD) );
+	
+	ActiveGameplayEffects.ExecuteActiveEffectsFrom(Spec, QualifierContext);
+}
+
+void UAttributeComponent::CheckDurationExpired(FActiveGameplayEffectHandle Handle)
+{
+	ActiveGameplayEffects.CheckDuration(Handle);
+}
+
+bool UAttributeComponent::RemoveActiveGameplayEffect(FActiveGameplayEffectHandle Handle)
+{
+	return ActiveGameplayEffects.RemoveActiveGameplayEffect(Handle);
+}
+
+float UAttributeComponent::GetGameplayEffectDuration(FActiveGameplayEffectHandle Handle) const
+{
+	return ActiveGameplayEffects.GetGameplayEffectDuration(Handle);
+}
+
+float UAttributeComponent::GetGameplayEffectMagnitude(FActiveGameplayEffectHandle Handle, FGameplayAttribute Attribute) const
+{
+	return ActiveGameplayEffects.GetGameplayEffectMagnitude(Handle, Attribute);
+}
+
+void UAttributeComponent::InvokeGameplayCueExecute(const FGameplayEffectSpec &Spec)
+{
+	if (DebugGameplayCues)
+	{
+		SKILL_LOG(Warning, TEXT("InvokeGameplayCueExecute: %s"), *Spec.ToSimpleString());
+	}
+
+
+	AActor *ActorOwner = GetOwner();
+	IGameplayCueInterface * GameplayCueInterface = InterfaceCast<IGameplayCueInterface>(ActorOwner);
+	if (!GameplayCueInterface)
+	{
+		return;
+	}
+	
+	// FIXME: Replication of level not finished
+	float ExecuteLevel = Spec.ModifierLevel.Get()->IsValid() ? Spec.ModifierLevel.Get()->GetLevel() : 1.f;
+	for (FGameplayEffectCue CueInfo : Spec.Def->GameplayCues)
+	{
+		float NormalizedMagnitude = CueInfo.NormalizeLevel(ExecuteLevel);
+		GameplayCueInterface->GameplayCueExecuted(CueInfo.GameplayCueTags, NormalizedMagnitude, Spec.InstigatorContext);
+
+		if (DebugGameplayCues)
+		{
+			DrawDebugSphere(GetWorld(), Spec.InstigatorContext.HitResult->Location, 30.f, 32, FColor(255, 0, 0), true, 30.f);
+			SKILL_LOG(Warning, TEXT("   %s"), *CueInfo.GameplayCueTags.ToString());
+		}
+	}
+}
+
+void UAttributeComponent::InvokeGameplayCueActivated(const FGameplayEffectSpec &Spec)
+{
+	AActor *ActorOwner = GetOwner();
+	IGameplayCueInterface * GameplayCueInterface = InterfaceCast<IGameplayCueInterface>(ActorOwner);
+	if (!GameplayCueInterface)
+	{
+		return;
+	}
+
+	// FIXME: Replication of level not finished
+	float ExecuteLevel = Spec.ModifierLevel.Get()->IsValid() ? Spec.ModifierLevel.Get()->GetLevel() : 1.f; 
+	for (FGameplayEffectCue CueInfo : Spec.Def->GameplayCues)
+	{
+		float NormalizedMagnitude = CueInfo.NormalizeLevel(ExecuteLevel);
+		GameplayCueInterface->GameplayCueActivated(CueInfo.GameplayCueTags, NormalizedMagnitude, Spec.InstigatorContext);
+	}
+}
+
+void UAttributeComponent::NetMulticast_InvokeGameplayCueExecuted_Implementation(const FGameplayEffectSpec Spec)
+{
+	InvokeGameplayCueExecute(Spec);
+}
+
+void UAttributeComponent::InvokeGameplayCueAdded(const FGameplayEffectSpec &Spec)
+{
+	AActor *ActorOwner = GetOwner();
+	IGameplayCueInterface * GameplayCueInterface = InterfaceCast<IGameplayCueInterface>(ActorOwner);
+	if (!GameplayCueInterface)
+	{
+		return;
+	}
+
+	// FIXME: Replication of level not finished
+	float ExecuteLevel = Spec.ModifierLevel.Get()->IsValid() ? Spec.ModifierLevel.Get()->GetLevel() : 1.f;
+	for (FGameplayEffectCue CueInfo : Spec.Def->GameplayCues)
+	{
+		float NormalizedMagnitude = CueInfo.NormalizeLevel(ExecuteLevel);
+		GameplayCueInterface->GameplayCueAdded(CueInfo.GameplayCueTags, NormalizedMagnitude, Spec.InstigatorContext);
+	}
+}
+
+void UAttributeComponent::InvokeGameplayCueRemoved(const FGameplayEffectSpec &Spec)
+{
+	AActor *ActorOwner = GetOwner();
+	IGameplayCueInterface * GameplayCueInterface = InterfaceCast<IGameplayCueInterface>(ActorOwner);
+	if (!GameplayCueInterface)
+	{
+		return;
+	}
+
+	// FIXME: Replication of level not finished
+	float ExecuteLevel = Spec.ModifierLevel.Get()->IsValid() ? Spec.ModifierLevel.Get()->GetLevel() : 1.f;
+	for (FGameplayEffectCue CueInfo : Spec.Def->GameplayCues)
+	{
+		float NormalizedMagnitude = CueInfo.NormalizeLevel(ExecuteLevel);
+		GameplayCueInterface->GameplayCueRemoved(CueInfo.GameplayCueTags, NormalizedMagnitude, Spec.InstigatorContext);
+	}
+}
+
+void UAttributeComponent::AddDependancyToAttribute(FGameplayAttribute Attribute, const TWeakPtr<FAggregator> InDependant)
+{
+	ActiveGameplayEffects.AddDependancyToAttribute(Attribute, InDependant);
+}
+
+bool UAttributeComponent::CanApplyAttributeModifiers(const UGameplayEffect *GameplayEffect, float Level, AActor *Instigator)
+{
+	return ActiveGameplayEffects.CanApplyAttributeModifiers(GameplayEffect, Level, Instigator);
+}
+
+TArray<float> UAttributeComponent::GetActiveEffectsTimeRemaining(const FActiveGameplayEffectQuery Query) const
+{
+	return ActiveGameplayEffects.GetActiveEffectsTimeRemaining(Query);
+}
+
+void UAttributeComponent::OnRestackGameplayEffects()
+{
+	ActiveGameplayEffects.RecalculateStacking();
+}
+
+// ---------------------------------------------------------------------------------------
+
+void UAttributeComponent::GetLifetimeReplicatedProps(TArray< FLifetimeProperty > & OutLifetimeProps) const
+{
+	Super::GetLifetimeReplicatedProps(OutLifetimeProps);
+
+	DOREPLIFETIME(UAttributeComponent, SpawnedAttributes);
+	DOREPLIFETIME(UAttributeComponent, ActiveGameplayEffects);
+
+	DOREPLIFETIME(UAttributeComponent, ReplicatedInstancedAbilities);
+	DOREPLIFETIME(UAttributeComponent, ActivatableAbilities);
+
+	DOREPLIFETIME(UAttributeComponent, ReplicatedPredictionKey);
+}
+
+bool UAttributeComponent::ReplicateSubobjects(class UActorChannel *Channel, class FOutBunch *Bunch, FReplicationFlags *RepFlags)
+{
+	bool WroteSomething = Super::ReplicateSubobjects(Channel, Bunch, RepFlags);
+
+	for (UAttributeSet* Set : SpawnedAttributes)
+	{
+		if (Set)
+		{
+			WroteSomething |= Channel->ReplicateSubobject(Set, *Bunch, *RepFlags);
+		}
+	}
+
+	for (UGameplayAbility* Ability : ReplicatedInstancedAbilities)
+	{
+		if (Ability)
+		{
+			WroteSomething |= Channel->ReplicateSubobject(Ability, *Bunch, *RepFlags);
+		}
+	}
+
+	return WroteSomething;
+}
+
+void UAttributeComponent::GetSubobjectsWithStableNamesForNetworking(TArray<UObject*>& Objs)
+{
+	for (UAttributeSet* Set : SpawnedAttributes)
+	{
+		if (Set && Set->IsNameStableForNetworking())
+		{
+			Objs.Add(Set);
+		}
+	}
+}
+
+void UAttributeComponent::PostNetReceive()
+{
+	Super::PostNetReceive();
+}
+
+void UAttributeComponent::OnRep_GameplayEffects()
+{
+
+}
+
+void UAttributeComponent::OnRep_PredictionKey()
+{
+	// Every predictive action we've done up to and including the current value of ReplicatedPredictionKey needs to be wiped
+	int32 idx = 0;
+	for ( ; idx < PredictionDelegates.Num(); ++idx)
+	{
+		if (PredictionDelegates[idx].Key <= ReplicatedPredictionKey)
+		{
+			PredictionDelegates[idx].Value.Broadcast();
+		}
+		else
+		{
+			break;
+		}
+	}
+	if (idx > 0)
+	{
+		// Remove everyone we called Broadcast on in one call
+		PredictionDelegates.RemoveAt(0, idx);
+	}
+}
+
+FAttributeComponentPredictionKeyClear& UAttributeComponent::GetPredictionKeyDelegate(int32 PredictionKey)
+{
+	// See if we already have one for this key
+	for (TPair<int32, FAttributeComponentPredictionKeyClear> & Item : PredictionDelegates)
+	{
+		if (Item.Key == PredictionKey)
+		{
+			return Item.Value;
+		}
+	}
+
+	// Create a new one
+	TPair<int32, FAttributeComponentPredictionKeyClear> NewPair;
+	NewPair.Key = PredictionKey;
+	
+	PredictionDelegates.Add(NewPair);
+	return PredictionDelegates.Last().Value;
+}
+
+int32 UAttributeComponent::GetNextPredictionKey()
+{
+	return ++LocalPredictionKey;
+}
+
+// ---------------------------------------------------------------------------------------
+
+void UAttributeComponent::PrintAllGameplayEffects() const
+{
+	SKILL_LOG_SCOPE(TEXT("PrintAllGameplayEffects %s"), *GetName());
+	SKILL_LOG(Log, TEXT("Owner: %s"), *GetOwner()->GetName());
+	ActiveGameplayEffects.PrintAllGameplayEffects();
+}
+
+void FActiveGameplayEffectsContainer::PrintAllGameplayEffects() const
+{
+	SKILL_LOG_SCOPE(TEXT("ActiveGameplayEffects. Num: %d"), GameplayEffects.Num());
+	for (const FActiveGameplayEffect& Effect : GameplayEffects)
+	{
+		Effect.PrintAll();
+	}
+}
+
+void FActiveGameplayEffect::PrintAll() const
+{
+	SKILL_LOG(Log, TEXT("Handle: %s"), *Handle.ToString());
+	SKILL_LOG(Log, TEXT("StartWorldTime: %.2f"), StartWorldTime);
+	Spec.PrintAll();
+}
+
+void FGameplayEffectSpec::PrintAll() const
+{
+	SKILL_LOG_SCOPE(TEXT("GameplayEffectSpec"));
+	SKILL_LOG(Log, TEXT("Def: %s"), *Def->GetName());
+	
+	SKILL_LOG(Log, TEXT("Duration: "));
+	Duration.PrintAll();
+
+	SKILL_LOG(Log, TEXT("Period:"));
+	Period.PrintAll();
+
+	SKILL_LOG(Log, TEXT("Modifiers:"));
+	for (const FModifierSpec &Mod : Modifiers)
+	{
+		Mod.PrintAll();
+	}
+}
+
+void FModifierSpec::PrintAll() const
+{
+	SKILL_LOG_SCOPE(TEXT("ModifierSpec"));
+	SKILL_LOG(Log, TEXT("Attribute: %s"), *Info.Attribute.GetName());
+	SKILL_LOG(Log, TEXT("ModifierType: %s"), *EGameplayModToString(Info.ModifierType));
+	SKILL_LOG(Log, TEXT("ModifierOp: %s"), *EGameplayModOpToString(Info.ModifierOp));
+	SKILL_LOG(Log, TEXT("EffectType: %s"), *EGameplayModEffectToString(Info.EffectType));
+	SKILL_LOG(Log, TEXT("RequiredTags: %s"), *Info.RequiredTags.ToString());
+	SKILL_LOG(Log, TEXT("OwnedTags: %s"), *Info.OwnedTags.ToString());
+	SKILL_LOG(Log, TEXT("(Base) Magnitude: %s"), *Info.Magnitude.ToSimpleString());
+
+	Aggregator.PrintAll();
+}
+
+void FAggregatorRef::PrintAll() const
+{
+	if (!WeakPtr.IsValid())
+	{
+		SKILL_LOG(Log, TEXT("Invalid AggregatorRef"));
+		return;
+	}
+
+	if (SharedPtr.IsValid())
+	{
+		SKILL_LOG(Log, TEXT("HardRef AggregatorRef"));
+	}
+	else
+	{
+		SKILL_LOG(Log, TEXT("SoftRef AggregatorRef"));
+
+	}
+	
+	Get()->PrintAll();
+}
+
+void FAggregator::PrintAll() const
+{
+	SKILL_LOG_SCOPE(TEXT("FAggregator 0x%X"), this);
+
+#if SKILL_SYSTEM_AGGREGATOR_DEBUG
+	SKILL_LOG(Log, TEXT("DebugStr: %s"), *DebugString);
+	SKILL_LOG(Log, TEXT("Copies (of me): %d"), CopiesMade);
+#endif
+
+	if (Level.IsValid())
+	{
+		SKILL_LOG_SCOPE(TEXT("LevelInfo"));
+		Level->PrintAll();
+	}
+	else
+	{
+		SKILL_LOG(Log, TEXT("No Level Data"));
+	}
+
+	{
+		SKILL_LOG_SCOPE(TEXT("BaseData"));
+		BaseData.PrintAll();
+	}
+
+	{
+		SKILL_LOG_SCOPE(TEXT("CachedData"));
+		CachedData.PrintAll();
+	}
+	
+	for (int32 i=0; i < EGameplayModOp::Max; ++i)
+	{
+		if (Mods[i].Num() > 0)
+		{
+			SKILL_LOG_SCOPE(TEXT("%s Mods"), *EGameplayModOpToString(i));
+			for (const FAggregatorRef &Ref : Mods[i])
+			{
+				Ref.PrintAll();
+			}
+		}
+	}
+}
+
+void FGameplayModifierData::PrintAll() const
+{
+	SKILL_LOG(Log, TEXT("Magnitude: %s"), *Magnitude.ToSimpleString());
+	SKILL_LOG(Log, TEXT("Tags: %s"), *Tags.ToString());
+}
+
+void FGameplayModifierEvaluatedData::PrintAll() const
+{
+	SKILL_LOG(Log, TEXT("IsValid: %d"), IsValid); 
+	SKILL_LOG(Log, TEXT("Magnitude: %.2f"), Magnitude);
+	SKILL_LOG(Log, TEXT("Tags: %s"), *Tags.ToString());
+}
+
+void FGameplayEffectLevelSpec::PrintAll() const
+{
+	SKILL_LOG(Log, TEXT("ConstantLevel: %.2f"), ConstantLevel);
+}
+
+#undef LOCTEXT_NAMESPACE
diff --git a/Engine/Source/Runtime/GameplayAbilities/Private/AttributeSet.cpp b/Engine/Source/Runtime/GameplayAbilities/Private/AttributeSet.cpp
new file mode 100644
--- /dev/null
+++ b/Engine/Source/Runtime/GameplayAbilities/Private/AttributeSet.cpp
@@ -0,0 +1,242 @@
+// Copyright 1998-2013 Epic Games, Inc. All Rights Reserved.
+
+#include "SkillSystemModulePrivatePCH.h"
+#include "AttributeSet.h"
+
+void FGameplayAttribute::SetNumericValueChecked(const float NewValue, class UAttributeSet* Dest) const
+{
+	UNumericProperty *NumericProperty = CastChecked<UNumericProperty>(Attribute);
+	void * ValuePtr = NumericProperty->ContainerPtrToValuePtr<void>(Dest);
+	NumericProperty->SetFloatingPointPropertyValue(ValuePtr, NewValue);
+}
+
+float FGameplayAttribute::GetNumericValueChecked(class UAttributeSet* Src) const
+{
+	UNumericProperty *NumericProperty = CastChecked<UNumericProperty>(Attribute);
+	void * ValuePtr = NumericProperty->ContainerPtrToValuePtr<void>(Src);
+	return NumericProperty->GetFloatingPointPropertyValue(ValuePtr);
+}
+
+UAttributeSet::UAttributeSet(const class FPostConstructInitializeProperties& PCIP)
+	: Super(PCIP)
+{
+
+}
+
+void UAttributeSet::InitFromMetaDataTable(const UDataTable* DataTable)
+{
+	static const FString Context = FString(TEXT("UAttribute::BindToMetaDataTable"));
+
+	for( TFieldIterator<UProperty> It(GetClass(), EFieldIteratorFlags::IncludeSuper) ; It ; ++It )
+	{
+		UProperty *Property = *It;
+		UNumericProperty *NumericProperty = Cast<UNumericProperty>(Property);
+		if (NumericProperty)
+		{
+			FString RowNameStr = FString::Printf(TEXT("%s.%s"), *Property->GetOuter()->GetName(), *Property->GetName());
+		
+			FAttributeMetaData * MetaData = DataTable->FindRow<FAttributeMetaData>(FName(*RowNameStr), Context, false);
+			if (MetaData)
+			{
+				void *Data = NumericProperty->ContainerPtrToValuePtr<void>(this);
+				NumericProperty->SetFloatingPointPropertyValue(Data, MetaData->BaseValue);
+			}
+		}
+	}
+
+	PrintDebug();
+}
+
+void UAttributeSet::ClampFromMetaDataTable(const UDataTable *DataTable)
+{
+	// Looking up a datatable row everytime a value changes in order to apply clamping pretty much sucks.
+	static const FString Context = FString(TEXT("UAttribute::BindToMetaDataTable"));
+
+	UE_LOG(LogActorComponent, Warning, TEXT("UAttribute::ClampFromMetaDataTable. UAttribute %s from UDataTable %s"), *GetName(), *DataTable->GetName());
+
+	for( TFieldIterator<UProperty> It(GetClass(), EFieldIteratorFlags::IncludeSuper) ; It ; ++It )
+	{
+		UProperty *Property = *It;
+		UNumericProperty *NumericProperty = Cast<UNumericProperty>(Property);
+		if (NumericProperty)
+		{
+			FString RowNameStr = FString::Printf(TEXT("%s.%s"), *Property->GetOuter()->GetName(), *Property->GetName());
+			FAttributeMetaData * MetaData = DataTable->FindRow<FAttributeMetaData>(FName(*RowNameStr), Context, false);
+			if (MetaData)
+			{
+				float CurrentValue = NumericProperty->GetFloatingPointPropertyValue(this);
+				CurrentValue = FMath::Clamp(CurrentValue, MetaData->MinValue, MetaData->MaxValue);
+
+				UE_LOG(LogActorComponent, Warning, TEXT("   Found Row for %s. Clamping: %.2f <%.2f, %.2f>"), *RowNameStr, CurrentValue, MetaData->MinValue, MetaData->MaxValue );
+
+				NumericProperty->SetFloatingPointPropertyValue(NumericProperty->ContainerPtrToValuePtr<void>(this), CurrentValue);
+			}
+		}
+	}
+
+	PrintDebug();
+}
+
+void UAttributeSet::PrintDebug()
+{
+	
+}
+
+void UAttributeSet::OnAttributeChange(UProperty *Property)
+{
+
+}
+
+FAttributeMetaData::FAttributeMetaData() :
+	MinValue(0.f),
+	MaxValue(1.f)
+{
+
+}
+
+UProperty* FAttributeModifierTest::GetUProperty()
+{
+	if (!CachedUProperty)
+	{
+		FString ClassName;
+		FString ThisPropertyName;
+
+		PropertyName.Split( TEXT("."), &ClassName, &ThisPropertyName);
+
+		UClass *FoundClass = FindObject<UClass>(ANY_PACKAGE, *ClassName);
+		if (FoundClass)
+		{
+			CachedUProperty = FindField<UProperty>(FoundClass, *ThisPropertyName);
+		}
+	}
+
+	return CachedUProperty;
+}
+
+void FAttributeModifierTest::ApplyModifier(UAttributeSet *Object)
+{
+	UProperty *TheProperty = GetUProperty();
+	if (TheProperty)
+	{
+		void * ValuePtr = TheProperty->ContainerPtrToValuePtr<void>(Object);
+		switch( ModifierType)
+		{
+			case EAttributeModifierType::EATTRMOD_Add:
+			{
+				UNumericProperty *NumericProperty = Cast<UNumericProperty>(TheProperty);
+				if (NumericProperty)
+				{
+					float CurrentValue = NumericProperty->GetFloatingPointPropertyValue(ValuePtr);
+					float AdditionalValue = 0.f;
+					TTypeFromString<float>::FromString(AdditionalValue, *this->NewValue);
+					
+					float FinalValue = CurrentValue + AdditionalValue;
+					NumericProperty->SetFloatingPointPropertyValue(ValuePtr, FinalValue);
+
+					UE_LOG(LogActorComponent, Warning, TEXT("FAttributeModifierTest::ApplyModifier - %s. EATTRMOD_Add - CurrentValue: %.2f AdditionalValue: %.2f -> FinalValue: %.2f"), *TheProperty->GetName(),
+						CurrentValue, AdditionalValue, FinalValue );
+				}
+			}
+			break;
+
+			case EAttributeModifierType::EATTRMOD_Multiple:
+			{
+				UNumericProperty *NumericProperty = Cast<UNumericProperty>(TheProperty);
+				if (NumericProperty)
+				{
+					float CurrentValue = NumericProperty->GetFloatingPointPropertyValue(ValuePtr);
+					float Factor = 0.f;
+					TTypeFromString<float>::FromString(Factor, *this->NewValue);
+
+					float FinalValue = CurrentValue * Factor;
+					NumericProperty->SetFloatingPointPropertyValue(ValuePtr, FinalValue);
+
+					UE_LOG(LogActorComponent, Warning, TEXT("FAttributeModifierTest::ApplyModifier - %s. EATTRMOD_Multiple - CurrentValue: %.2f Factor: %.2f -> FinalValue: %.2f"), *TheProperty->GetName(),
+						CurrentValue, Factor, FinalValue );
+				}
+			}
+			break;
+
+			case EAttributeModifierType::EATTRMOD_Override:
+				TheProperty->ImportText(*this->NewValue, ValuePtr, 0, Object); // Not sure if the 4th parameter is right
+
+				UE_LOG(LogActorComponent, Warning, TEXT("FAttributeModifierTest::ApplyModifier - %s. EATTRMOD_Override - %s"), *TheProperty->GetName(), *this->NewValue );
+				break;
+		}
+	}
+
+	Object->PrintDebug();
+}
+
+void FScalableFloat::FinalizeCurveData(const FGlobalCurveDataOverride *GlobalOverrides)
+{
+	static const FString ContextString = TEXT("FScalableFloat::FinalizeCurveData");
+
+	// We are a static value, so do nothing.
+	if (Curve.RowName == NAME_None)
+	{
+		return;
+	}
+
+	// Tied to an explicit table, so bind now.
+	if (Curve.CurveTable != NULL)
+	{
+		FinalCurve = Curve.GetCurve(ContextString);
+		return;
+	}
+
+	// We have overrides
+	if (GlobalOverrides)
+	{
+		for (UCurveTable* OverrideTable : GlobalOverrides->Overrides)
+		{
+			FinalCurve = OverrideTable->FindCurve(Curve.RowName, ContextString, false);
+			if (FinalCurve)
+			{
+				return;
+			}
+		}
+	}
+
+	// Look at global defaults
+	const UCurveTable * GlobalTable = IGameplayAbilitiesModule::Get().GetSkillSystemGlobals()->GetGlobalCurveTable();
+	if (GlobalTable)
+	{
+		FinalCurve = GlobalTable->FindCurve(Curve.RowName, ContextString, false);
+	}
+
+	if (!FinalCurve)
+	{
+		SKILL_LOG(Warning, TEXT("Unable to find RowName: %s for FScalableFloat."), *Curve.RowName.ToString());
+	}
+}
+
+float FScalableFloat::GetValueAtLevel(float Level) const
+{
+	if (FinalCurve)
+	{
+		return Value * FinalCurve->Eval(Level);
+	}
+
+	return Value;
+}
+
+bool FGameplayAttribute::operator==(const FGameplayAttribute& Other) const
+{
+	return ((Other.Attribute == Attribute));
+}
+
+bool FGameplayAttribute::operator!=(const FGameplayAttribute& Other) const
+{
+	return ((Other.Attribute != Attribute));
+}
+
+bool FScalableFloat::operator==(const FScalableFloat& Other) const
+{
+	return ((Other.Curve == Curve) && (Other.Value == Value));
+}
+
+bool FScalableFloat::operator!=(const FScalableFloat& Other) const
+{
+	return ((Other.Curve != Curve) || (Other.Value != Value));
+}
diff --git a/Engine/Source/Runtime/GameplayAbilities/Private/GameplayAbilitiesModule.cpp b/Engine/Source/Runtime/GameplayAbilities/Private/GameplayAbilitiesModule.cpp
new file mode 100644
--- /dev/null
+++ b/Engine/Source/Runtime/GameplayAbilities/Private/GameplayAbilitiesModule.cpp
@@ -0,0 +1,40 @@
+// Copyright 1998-2013 Epic Games, Inc. All Rights Reserved.
+
+#include "SkillSystemModulePrivatePCH.h"
+#include "SkillSystemGlobals.h"
+
+
+class FGameplayAbilitiesModule : public IGameplayAbilitiesModule
+{
+	// Begin IModuleInterface
+	virtual void StartupModule() OVERRIDE;
+	virtual void ShutdownModule() OVERRIDE;
+	// End IModuleInterface
+
+	virtual USkillSystemGlobals* GetSkillSystemGlobals()
+	{
+		check(SkillSystemGlobals);
+		return SkillSystemGlobals;
+	}
+
+	USkillSystemGlobals *SkillSystemGlobals;
+
+private:
+	
+};
+
+IMPLEMENT_MODULE(FGameplayAbilitiesModule, GameplayAbilities)
+
+void FGameplayAbilitiesModule::StartupModule()
+{	
+	// This code will execute after your module is loaded into memory (but after global variables are initialized, of course.)
+	SkillSystemGlobals = ConstructObject<USkillSystemGlobals>(USkillSystemGlobals::StaticClass(), GetTransientPackage(), NAME_None, RF_RootSet);
+}
+
+void FGameplayAbilitiesModule::ShutdownModule()
+{
+	// This function may be called during shutdown to clean up your module.  For modules that support dynamic reloading,
+	// we call this function before unloading the module.
+
+	SkillSystemGlobals = NULL;
+}
diff --git a/Engine/Source/Runtime/GameplayAbilities/Private/GameplayCueInterface.cpp b/Engine/Source/Runtime/GameplayAbilities/Private/GameplayCueInterface.cpp
new file mode 100644
--- /dev/null
+++ b/Engine/Source/Runtime/GameplayAbilities/Private/GameplayCueInterface.cpp
@@ -0,0 +1,9 @@
+// Copyright 1998-2014 Epic Games, Inc. All Rights Reserved.
+
+#include "SkillSystemModulePrivatePCH.h"
+#include "GameplayCueInterface.h"
+
+UGameplayCueInterface::UGameplayCueInterface(const FPostConstructInitializeProperties& PCIP)
+: Super(PCIP)
+{
+}
diff --git a/Engine/Source/Runtime/GameplayAbilities/Private/GameplayCueView.cpp b/Engine/Source/Runtime/GameplayAbilities/Private/GameplayCueView.cpp
new file mode 100644
--- /dev/null
+++ b/Engine/Source/Runtime/GameplayAbilities/Private/GameplayCueView.cpp
@@ -0,0 +1,220 @@
+// Copyright 1998-2014 Epic Games, Inc. All Rights Reserved.
+
+#pragma once
+
+#include "SkillSystemModulePrivatePCH.h"
+#include "GameplayEffect.h"
+#include "GameplayCueView.h"
+#include "ParticleDefinitions.h"
+#include "SoundDefinitions.h"
+#include "Particles/ParticleLODLevel.h"
+#include "Particles/ParticleModuleRequired.h"
+#include "Particles/ParticleSystemComponent.h"
+#include "GameplayTagsModule.h"
+
+UGameplayCueView::UGameplayCueView(const class FPostConstructInitializeProperties& PCIP)
+: Super(PCIP)
+{
+
+}
+
+FGameplayCueViewInfo * FGameplayCueHandler::GetBestMatchingView(EGameplayCueEvent::Type Type, const FGameplayTag BaseTag)
+{
+	IGameplayTagsModule& GameplayTagsModule = IGameplayTagsModule::Get();
+	FGameplayTagContainer TagAndParentsContainer = GameplayTagsModule.GetGameplayTagsManager().RequestGameplayTagParents(BaseTag);
+
+	for (auto InnerTagIt = TagAndParentsContainer.CreateConstIterator(); InnerTagIt; ++InnerTagIt)
+	{
+		FGameplayTag Tag = *InnerTagIt;
+
+		for (UGameplayCueView * Def : Definitions)
+		{
+			for (FGameplayCueViewInfo & View : Def->Views)
+			{
+				if (View.CueType == Type && View.Tags.HasTag(Tag, EGameplayTagMatchType::IncludeParentTags, EGameplayTagMatchType::Explicit))
+				{
+					return &View;
+				}
+			}
+		}
+	}
+
+	return NULL;
+}
+
+void FGameplayCueHandler::GameplayCueActivated(const FGameplayTagContainer & GameplayCueTags, float NormalizedMagnitude, const FGameplayEffectInstigatorContext InstigatorContext)
+{
+	check(Owner);
+	for (auto TagIt = GameplayCueTags.CreateConstIterator(); TagIt; ++TagIt)
+	{
+		if (FGameplayCueViewInfo * View = GetBestMatchingView(EGameplayCueEvent::Applied, *TagIt))
+		{
+			View->SpawnViewEffects(Owner, NULL, InstigatorContext);
+		}
+	}
+}
+
+void FGameplayCueHandler::GameplayCueExecuted(const FGameplayTagContainer & GameplayCueTags, float NormalizedMagnitude, const FGameplayEffectInstigatorContext InstigatorContext)
+{
+	check(Owner);
+	for (auto TagIt = GameplayCueTags.CreateConstIterator(); TagIt; ++TagIt)
+	{
+		if (FGameplayCueViewInfo * View = GetBestMatchingView(EGameplayCueEvent::Executed, *TagIt))
+		{
+			View->SpawnViewEffects(Owner, NULL, InstigatorContext);
+		}
+	}
+}
+
+void FGameplayCueHandler::GameplayCueAdded(const FGameplayTagContainer & GameplayCueTags, float NormalizedMagnitude, const FGameplayEffectInstigatorContext InstigatorContext)
+{
+	check(Owner);
+
+	for (auto TagIt = GameplayCueTags.CreateConstIterator(); TagIt; ++TagIt)
+	{
+		TArray<TSharedPtr<FGameplayCueViewEffects> > &Effects = SpawnedViewEffects.FindOrAdd(*TagIt);
+
+		// Clear old effects if they existed? This will vary case to case. Removing old effects is the easiest approach
+		ClearEffects(Effects);
+		check(Effects.Num() == 0);
+
+		if (FGameplayCueViewInfo * View = GetBestMatchingView(EGameplayCueEvent::Added, *TagIt))
+		{
+			View->SpawnViewEffects(Owner, NULL, InstigatorContext);
+
+			TSharedPtr<FGameplayCueViewEffects> SpawnedEffects = View->SpawnViewEffects(Owner, &SpawnedObjects, InstigatorContext);
+			Effects.Add(SpawnedEffects);
+		}
+	}
+}
+
+void FGameplayCueHandler::GameplayCueRemoved(const FGameplayTagContainer & GameplayCueTags, float NormalizedMagnitude, const FGameplayEffectInstigatorContext InstigatorContext)
+{
+	check(Owner);
+
+	for (auto TagIt = GameplayCueTags.CreateConstIterator(); TagIt; ++TagIt)
+	{
+		TArray<TSharedPtr<FGameplayCueViewEffects> > *Effects = SpawnedViewEffects.Find(*TagIt);
+		if (Effects)
+		{
+			ClearEffects(*Effects);
+			SpawnedViewEffects.Remove(*TagIt);
+		}	
+		
+		// Remove old effects
+		
+		if (FGameplayCueViewInfo * View = GetBestMatchingView(EGameplayCueEvent::Removed, *TagIt))
+		{
+			View->SpawnViewEffects(Owner, NULL, InstigatorContext);
+		}
+	}
+}
+
+void FGameplayCueHandler::ClearEffects(TArray< TSharedPtr<FGameplayCueViewEffects> > &Effects)
+{
+	bool RemovedSomething = false;
+	for (TSharedPtr<FGameplayCueViewEffects> &EffectPtr : Effects)
+	{
+		if (!EffectPtr.IsValid())
+		{
+			continue;
+		}
+
+		FGameplayCueViewEffects &Effect = *EffectPtr.Get();
+
+		if (Effect.SpawnedActor.IsValid())
+		{
+			Effect.SpawnedActor->Destroy();
+			RemovedSomething = true;
+		}
+
+		if (Effect.AudioComponent.IsValid())
+		{
+			Effect.AudioComponent->Stop();
+			RemovedSomething = true;
+		}
+
+		if (Effect.ParticleSystemComponent.IsValid())
+		{
+			Effect.ParticleSystemComponent->DestroyComponent();
+			RemovedSomething = true;
+		}
+	}
+
+	Effects.Empty();
+	
+	// Cleanup flat array
+	if (RemovedSomething)
+	{
+		for (int32 idx=0; idx < SpawnedObjects.Num(); ++idx)
+		{
+			UObject *Obj = SpawnedObjects[idx];
+			if (!Obj || Obj->IsPendingKill())
+			{
+				SpawnedObjects.RemoveAtSwap(idx);
+				idx--;
+			}
+		}
+	}
+}
+
+TSharedPtr<FGameplayCueViewEffects> FGameplayCueViewInfo::SpawnViewEffects(AActor *Owner, TArray<UObject*> *SpawnedObjects, const FGameplayEffectInstigatorContext InstigatorContext) const
+{
+	check(Owner);
+
+	TSharedPtr<FGameplayCueViewEffects> SpawnedEffects = TSharedPtr<FGameplayCueViewEffects>(new FGameplayCueViewEffects());
+
+	if (Sound)
+	{
+		SpawnedEffects->AudioComponent = UGameplayStatics::PlaySoundAttached(Sound, Owner->GetRootComponent());
+		if (SpawnedObjects)
+		{
+			SpawnedObjects->Add(SpawnedEffects->AudioComponent.Get());
+		}
+	}
+	if (ParticleSystem)
+	{
+		if (InstigatorContext.HitResult.IsValid())
+		{
+			FHitResult &HitResult = *InstigatorContext.HitResult.Get();
+
+			SpawnedEffects->ParticleSystemComponent = UGameplayStatics::SpawnEmitterAttached(ParticleSystem, Owner->GetRootComponent(), NAME_None, HitResult.Location,
+				HitResult.Normal.Rotation(), EAttachLocation::KeepWorldPosition);
+		}
+		else
+		{
+			SpawnedEffects->ParticleSystemComponent = UGameplayStatics::SpawnEmitterAttached(ParticleSystem, Owner->GetRootComponent());
+		}		
+
+		if (SpawnedObjects)
+		{
+			SpawnedObjects->Add(SpawnedEffects->ParticleSystemComponent.Get());
+		}
+		else if (SpawnedEffects->ParticleSystemComponent.IsValid())
+		{
+			for (int32 EmitterIndx = 0; EmitterIndx < SpawnedEffects->ParticleSystemComponent->EmitterInstances.Num(); EmitterIndx++)
+			{
+				if (SpawnedEffects->ParticleSystemComponent->EmitterInstances[EmitterIndx] &&
+					SpawnedEffects->ParticleSystemComponent->EmitterInstances[EmitterIndx]->CurrentLODLevel &&
+					SpawnedEffects->ParticleSystemComponent->EmitterInstances[EmitterIndx]->CurrentLODLevel->RequiredModule &&
+					SpawnedEffects->ParticleSystemComponent->EmitterInstances[EmitterIndx]->CurrentLODLevel->RequiredModule->EmitterLoops == 0)
+				{
+					SKILL_LOG(Warning, TEXT("%s - particle system has a looping emitter. This should not be used in a executed GameplayCue!"), *SpawnedEffects->ParticleSystemComponent->GetName());
+					break;
+				}
+			}
+		}
+	}
+	if (ActorClass)
+	{
+		FVector Location = Owner->GetActorLocation();
+		FRotator Rotation = Owner->GetActorRotation();
+		SpawnedEffects->SpawnedActor = Owner->GetWorld()->SpawnActor(ActorClass, &Location, &Rotation);
+		if (SpawnedObjects)
+		{
+			SpawnedObjects->Add(SpawnedEffects->SpawnedActor.Get());
+		}
+	}
+
+	return SpawnedEffects;
+}
