Rename SkillSystem module to GameplayAbilities

diff --git a/Engine/Source/Runtime/GameplayAbilities/Private/GameplayEffect.cpp b/Engine/Source/Runtime/GameplayAbilities/Private/GameplayEffect.cpp
new file mode 100644
--- /dev/null
+++ b/Engine/Source/Runtime/GameplayAbilities/Private/GameplayEffect.cpp
@@ -0,0 +1,2009 @@
+// Copyright 1998-2013 Epic Games, Inc. All Rights Reserved.
+
+#include "SkillSystemModulePrivatePCH.h"
+#include "GameplayEffect.h"
+#include "GameplayEffectExtension.h"
+#include "GameplayEffectStackingExtension.h"
+#include "AttributeComponent.h"
+
+const float UGameplayEffect::INFINITE_DURATION = -1.f;
+const float UGameplayEffect::INSTANT_APPLICATION = 0.f;
+const float UGameplayEffect::NO_PERIOD = 0.f;
+const float FGameplayEffectLevelSpec::INVALID_LEVEL = -1.f;
+
+#if SKILL_SYSTEM_AGGREGATOR_DEBUG
+FAggregator::FAllocationStats FAggregator::AllocationStats;
+#endif
+
+// --------------------------------------------------------------------------------------------------------------------------------------------------------
+//
+//	UGameplayEffect
+//
+// --------------------------------------------------------------------------------------------------------------------------------------------------------
+
+
+UGameplayEffect::UGameplayEffect(const class FPostConstructInitializeProperties& PCIP)
+	: Super(PCIP)
+{
+	ChanceToApplyToTarget.SetValue(1.f);
+	ChanceToExecuteOnGameplayEffect.SetValue(1.f);
+	StackingPolicy = EGameplayEffectStackingPolicy::Unlimited;
+	StackedAttribName = NAME_None;
+}
+
+void UGameplayEffect::GetOwnedGameplayTags(FGameplayTagContainer& TagContainer) const
+{
+	TagContainer.AppendTags(OwnedTagsContainer);
+}
+
+void UGameplayEffect::GetClearGameplayTags(FGameplayTagContainer& TagContainer) const
+{
+	TagContainer.AppendTags(ClearTagsContainer);
+}
+
+bool UGameplayEffect::AreApplicationTagRequirementsSatisfied(const FGameplayTagContainer& InstigatorTags, const FGameplayTagContainer& TargetTags) const
+{
+	return (InstigatorTags.MatchesAll(ApplicationRequiredInstigatorTags, true) && TargetTags.MatchesAll(ApplicationRequiredTargetTags, true));
+}
+
+void UGameplayEffect::PostLoad()
+{
+	Super::PostLoad();
+
+	ValidateGameplayEffect();
+}
+
+void UGameplayEffect::ValidateGameplayEffect()
+{
+	ValidateStacking();
+
+	// todo: add a check here for instant effects that modify incoming/outgoing GEs
+}
+
+void UGameplayEffect::ValidateStacking()
+{
+	if (StackingPolicy != EGameplayEffectStackingPolicy::Unlimited)
+	{
+		StackingPolicy = StackingPolicy;
+
+		if (StackedAttribName == NAME_None)
+		{
+			SKILL_LOG(Warning, TEXT("%s has a stacking rule but does not have an attribute to apply it to. Removing stacking rule."), *GetPathName());
+			StackingPolicy = EGameplayEffectStackingPolicy::Unlimited;
+		}
+
+		UDataTable* DataTable = IGameplayAbilitiesModule::Get().GetSkillSystemGlobals()->GetGlobalAttributeDataTable();
+
+		if (DataTable && DataTable->RowMap.Contains(StackedAttribName))
+		{
+			FAttributeMetaData* Row = (FAttributeMetaData*)DataTable->RowMap[StackedAttribName];
+			if (Row && Row->bCanStack == false)
+			{
+				SKILL_LOG(Warning, TEXT("%s has a stacking rule but modifies attribute %s. %s is not allowed to stack. Removing stacking rule."), *GetPathName(), *Modifiers[0].Attribute.GetName(), *Modifiers[0].Attribute.GetName());
+				StackingPolicy = EGameplayEffectStackingPolicy::Unlimited;
+			}
+		}
+		else
+		{
+			SKILL_LOG(Warning, TEXT("%s has a stacking rule but modifies attribute %s. %s was not found in the global attribute data table. Removing stacking rule."), *GetPathName(), *Modifiers[0].Attribute.GetName(), *Modifiers[0].Attribute.GetName());
+			StackingPolicy = EGameplayEffectStackingPolicy::Unlimited;
+		}
+
+		if (Modifiers.Num() == 0)
+		{
+			SKILL_LOG(Warning, TEXT("%s has a stacking rule but does not have any modifiers to apply it to. Removing stacking rule."), *GetPathName());
+			StackingPolicy = EGameplayEffectStackingPolicy::Unlimited;
+		}
+		else
+		{
+			// look for modifiers to attributes that are not the stacking attribute and warn that they may not apply in a predictable fashion
+			for (int32 Idx = 0; Idx < Modifiers.Num(); ++Idx)
+			{
+				if (Modifiers[Idx].ModifierType == EGameplayMod::Attribute)
+				{
+					FName AttributeName(*Modifiers[Idx].Attribute.GetName());
+
+					if (AttributeName != StackedAttribName)
+					{
+						SKILL_LOG(Warning, TEXT("%s has a stacking rule and modifies attribute %s but one of the modifiers modifies attribute %s. Changes to %s may not behave in a predictable fashion."), *GetPathName(), *StackedAttribName.ToString(), *AttributeName.ToString(), *AttributeName.ToString());
+					}
+				}
+				else if (Modifiers[Idx].ModifierType == EGameplayMod::ActiveGE || Modifiers[Idx].ModifierType == EGameplayMod::IncomingGE || Modifiers[Idx].ModifierType == EGameplayMod::OutgoingGE)
+				{
+					SKILL_LOG(Warning, TEXT("%s has a stacking rule and modifies attribute %s but one of its modifiers modifies gameplay effects. Changes to other gameplay effects may not behave in a predictable fashion."), *GetPathName(), *StackedAttribName.ToString());
+				}
+			}
+		}
+	}
+}
+
+// --------------------------------------------------------------------------------------------------------------------------------------------------------
+//
+//	FGameplayEffectSpec
+//
+// --------------------------------------------------------------------------------------------------------------------------------------------------------
+
+FGameplayEffectSpec::FGameplayEffectSpec(const UGameplayEffect *InDef, AActor * Instigator, float Level, const FGlobalCurveDataOverride *CurveData)
+: Def(InDef)
+, ModifierLevel(TSharedPtr<FGameplayEffectLevelSpec>(new FGameplayEffectLevelSpec(Level, InDef->LevelInfo, Instigator)))
+, Duration(new FAggregator(InDef->Duration.MakeFinalizedCopy(CurveData), ModifierLevel, SKILL_AGG_DEBUG(TEXT("%s Duration"), *InDef->GetName())))
+, Period(new FAggregator(InDef->Period.MakeFinalizedCopy(CurveData), ModifierLevel, SKILL_AGG_DEBUG(TEXT("%s Period"), *InDef->GetName())))
+, ChanceToApplyToTarget(new FAggregator(InDef->ChanceToApplyToTarget.MakeFinalizedCopy(CurveData), ModifierLevel, SKILL_AGG_DEBUG(TEXT("%s ChanceToApplyToTarget"), *InDef->GetName())))
+, ChanceToExecuteOnGameplayEffect(new FAggregator(InDef->ChanceToExecuteOnGameplayEffect.MakeFinalizedCopy(CurveData), ModifierLevel, SKILL_AGG_DEBUG(TEXT("%s ChanceToExecuteOnGameplayEffect"), *InDef->GetName())))
+{
+	Duration.Get()->RegisterLevelDependancies();
+	Period.Get()->RegisterLevelDependancies();
+	ChanceToApplyToTarget.Get()->RegisterLevelDependancies();
+	ChanceToExecuteOnGameplayEffect.Get()->RegisterLevelDependancies();
+
+	InitModifiers(CurveData, Instigator, Level);
+
+	if (InDef)
+	{
+		for (UGameplayEffect* TargetDef : InDef->TargetEffects)
+		{
+			TargetEffectSpecs.Add(TSharedRef<FGameplayEffectSpec>(new FGameplayEffectSpec(TargetDef, Instigator, Level, CurveData)));
+		}
+	}
+
+	InstigatorContext.AddInstigator(Instigator);
+}
+
+void FGameplayEffectSpec::InitModifiers(const FGlobalCurveDataOverride *CurveData, AActor *Owner, float Level)
+{
+	check(Def);
+
+	Modifiers.Reserve(Def->Modifiers.Num());
+	
+	for (const FGameplayModifierInfo &ModInfo : Def->Modifiers)
+	{
+		// Pass down the LevelInfo into this Modifier. 
+		// ModifierSpec may want to override how leveling will work (for this modifier only).
+		// Or it may use the GameplayEffectSpec's level. 
+		// ApplyNewDef will update NewLevelSpec appropriately.
+
+		TSharedPtr<FGameplayEffectLevelSpec> NewLevelSpec = ModifierLevel;
+		ModifierLevel->ApplyNewDef(ModInfo.LevelInfo, NewLevelSpec);
+
+		// This creates a new FModifierSpec that we own.
+		Modifiers.Emplace(FModifierSpec(ModInfo, NewLevelSpec, CurveData, Owner, Level));
+	}	
+}
+
+void FGameplayEffectSpec::MakeUnique()
+{
+	for (FModifierSpec &ModSpec : Modifiers)
+	{
+		ModSpec.Aggregator.MakeUnique();
+	}
+}
+
+int32 FGameplayEffectSpec::ApplyModifiersFrom(FGameplayEffectSpec &InSpec, const FModifierQualifier &QualifierContext)
+{
+	SKILL_LOG_SCOPE(TEXT("FGameplayEffectSpec::ApplyModifiersFrom %s. InSpec: %s"), *this->ToSimpleString(), *InSpec.ToSimpleString());
+
+	int32 NumApplied = 0;
+	bool ShouldSnapshot = InSpec.ShouldApplyAsSnapshot(QualifierContext);
+
+	if (!InSpec.Def || !InSpec.Def->AreGameplayEffectTagRequirementsSatisfied(Def))
+	{
+		// InSpec doesn't match this gameplay effect but if InSpec provides immunity we also need to check the modifiers because they can be canceled independent of the gameplay effect
+		if ((int32)InSpec.Def->AppliesImmunityTo == (int32)QualifierContext.Type())
+		{
+			for (int ii = 0; ii < Modifiers.Num(); ++ii)
+			{
+				const FGameplayModifierEvaluatedData& Data = Modifiers[ii].Aggregator.Get()->Evaluate();
+				if (InSpec.Def->AreGameplayEffectTagRequirementsSatisfied(Data.Tags))
+				{
+					Modifiers.RemoveAtSwap(ii);
+					--ii;
+				}
+			}
+		}
+
+		return 0;
+	}
+
+	if ((int32)InSpec.Def->AppliesImmunityTo == (int32)QualifierContext.Type())
+	{
+		return -1;
+	}
+
+	// Fixme: Use acceleration structures to speed up these types of lookups
+	// The called functions below are reliant on the InSpecs evaluated data. We should ideally only call evaluate once per mod.
+
+	for (const FModifierSpec &InMod : InSpec.Modifiers)
+	{
+		if (!InMod.CanModifyInContext(QualifierContext))
+		{
+			continue;
+		}
+
+		switch (InMod.Info.EffectType)
+		{
+			case EGameplayModEffect::Magnitude:
+			{
+				for (FModifierSpec &MyMod : Modifiers)
+				{
+					if (InMod.CanModifyModifier(MyMod, QualifierContext))
+					{
+						InMod.ApplyModTo(MyMod, ShouldSnapshot);
+						NumApplied++;
+					}
+				}
+				break;
+			}
+
+			// Fixme: Duration mods aren't checking that attributes match - do we care?
+			// eg - "I mod duration of all GEs that modify Health" would need to check to see if this mod modifies health before doing the stuff below.
+			// Or can we get by with just tags?
+
+			case EGameplayModEffect::Duration:
+			{
+				// don't modify infinite duration unless we're overriding it
+				if (GetDuration() > 0.f || InMod.Info.ModifierOp == EGameplayModOp::Override)
+				{
+					Duration.Get()->ApplyMod(InMod.Info.ModifierOp, InMod.Aggregator, ShouldSnapshot);
+					NumApplied++;
+				}
+				
+				break;
+			}
+
+			case EGameplayModEffect::ChanceApplyTarget:
+			{
+				ChanceToApplyToTarget.Get()->ApplyMod(InMod.Info.ModifierOp, InMod.Aggregator, ShouldSnapshot);
+				NumApplied++;
+				break;
+			}
+
+			case EGameplayModEffect::ChanceExecuteEffect:
+			{
+				ChanceToExecuteOnGameplayEffect.Get()->ApplyMod(InMod.Info.ModifierOp, InMod.Aggregator, ShouldSnapshot);
+				NumApplied++;
+				break;
+			}
+
+			case EGameplayModEffect::LinkedGameplayEffect:
+			{
+				TargetEffectSpecs.Add(InMod.TargetEffectSpec.ToSharedRef());
+				NumApplied++;
+				break;
+			}
+		}
+	}
+
+	return NumApplied;
+}
+
+int32 FGameplayEffectSpec::ExecuteModifiersFrom(const FGameplayEffectSpec &InSpec, const FModifierQualifier &QualifierContext)
+{
+	int32 NumExecuted = 0;
+
+	// Fixme: Use acceleration structures to speed up these types of lookups
+	for (FModifierSpec &MyMod : Modifiers)
+	{
+		for (const FModifierSpec &InMod : InSpec.Modifiers)
+		{
+			if (InMod.CanModifyModifier(MyMod, QualifierContext))
+			{
+				InMod.ExecuteModOn(MyMod);
+				NumExecuted++;
+			}
+		}
+	}
+
+	return NumExecuted;
+}
+
+float FGameplayEffectSpec::GetDuration() const
+{
+	return Duration.Get()->Evaluate().Magnitude;
+}
+
+float FGameplayEffectSpec::GetPeriod() const
+{
+	return Period.Get()->Evaluate().Magnitude;
+}
+
+float FGameplayEffectSpec::GetChanceToApplyToTarget() const
+{
+	return ChanceToApplyToTarget.Get()->Evaluate().Magnitude;
+}
+
+float FGameplayEffectSpec::GetChanceToExecuteOnGameplayEffect() const
+{
+	return ChanceToExecuteOnGameplayEffect.Get()->Evaluate().Magnitude;
+}
+
+float FGameplayEffectSpec::GetMagnitude(const FGameplayAttribute &Attribute) const
+{
+	float CurrentMagnitude = 0.f;
+
+	for (const FModifierSpec &Mod : Modifiers)
+	{
+		if (Mod.Info.ModifierType == EGameplayMod::Attribute)
+		{
+			if (Mod.Info.Attribute != Attribute)
+			{
+				continue;
+			}
+
+			// Todo: Tags/application checks here - make sure we can still apply
+
+			// First, evaluate all of our data 
+
+			FGameplayModifierEvaluatedData EvaluatedData = Mod.Aggregator.Get()->Evaluate();
+
+			FAggregator	Aggregator(CurrentMagnitude, SKILL_AGG_DEBUG(TEXT("Magnitude of Attribute %s"), *Mod.Info.Attribute.GetName()));
+
+			Aggregator.ExecuteMod(Mod.Info.ModifierOp, EvaluatedData);
+
+			CurrentMagnitude = Aggregator.Evaluate().Magnitude;
+		}
+	}
+
+	return CurrentMagnitude;
+}
+
+bool FGameplayEffectSpec::ShouldApplyAsSnapshot(const FModifierQualifier &QualifierContext) const
+{
+	bool ShouldSnapshot;
+	switch(Def->CopyPolicy)
+	{
+		case EGameplayEffectCopyPolicy::AlwaysSnapshot:
+			ShouldSnapshot = true;
+			break;
+			
+		case EGameplayEffectCopyPolicy::AlwaysLink:
+			ShouldSnapshot = false;
+			break;
+			
+		default:
+			ShouldSnapshot = (QualifierContext.Type() == EGameplayMod::OutgoingGE);
+			break;
+	}
+	
+	return ShouldSnapshot;
+}
+
+EGameplayEffectStackingPolicy::Type FGameplayEffectSpec::GetStackingType() const
+{
+	return Def->StackingPolicy;
+}
+
+// --------------------------------------------------------------------------------------------------------------------------------------------------------
+//
+//	FModifierSpec
+//
+// --------------------------------------------------------------------------------------------------------------------------------------------------------
+
+bool FModifierSpec::CanModifyInContext(const FModifierQualifier &QualifierContext) const
+{
+	// Can only modify if are valid within this Qualifier Context
+	// (E.g, if I am an OutgoingGE mod, I cannot modify during a IncomingGE context)
+	if (Info.ModifierType != QualifierContext.Type())
+	{
+		return false;
+	}	
+
+	return true;
+}
+
+bool FModifierSpec::CanModifyModifier(FModifierSpec &Other, const FModifierQualifier &QualifierContext) const
+{
+	// Attribute is essentially a key. This isn't 100% necessary - we could just rely on tags
+	if (Info.Attribute != Other.Info.Attribute)
+	{
+		return false;
+	}
+
+	// Tag checking is done at the FAggregator level. So all we do here is the attribute check.
+
+	return true;
+}
+
+FModifierSpec::FModifierSpec(const FGameplayModifierInfo &InInfo, TSharedPtr<FGameplayEffectLevelSpec> InLevel, const FGlobalCurveDataOverride *CurveData, AActor *Owner, float Level)
+: Info(InInfo)
+, Aggregator(new FAggregator(FGameplayModifierData(InInfo, CurveData), InLevel, SKILL_AGG_DEBUG(TEXT("FModifierSpec: %s "), *InInfo.ToSimpleString())))
+{
+	Aggregator.Get()->RegisterLevelDependancies();
+	if (InInfo.TargetEffect)
+	{
+		TargetEffectSpec = TSharedPtr<FGameplayEffectSpec>(new FGameplayEffectSpec(InInfo.TargetEffect, Owner, Level, CurveData));
+	}
+}
+
+void FModifierSpec::ApplyModTo(FModifierSpec &Other, bool TakeSnapshot) const
+{
+	Other.Aggregator.Get()->ApplyMod(this->Info.ModifierOp, this->Aggregator, TakeSnapshot);
+}
+
+void FModifierSpec::ExecuteModOn(FModifierSpec &Other) const
+{
+	SKILL_LOG_SCOPE(TEXT("Executing %s on %s"), *ToSimpleString(), *Other.ToSimpleString() );
+	Other.Aggregator.Get()->ExecuteModAggr(this->Info.ModifierOp, this->Aggregator);
+}
+
+bool FModifierSpec::AreTagRequirementsSatisfied(const FModifierSpec &ModifierToBeModified) const
+{
+	const FGameplayModifierEvaluatedData & ToBeModifiedData = ModifierToBeModified.Aggregator.Get()->Evaluate();
+
+	bool HasRequired = ToBeModifiedData.Tags.MatchesAll(this->Info.RequiredTags, true);
+	bool HasIgnored = ToBeModifiedData.Tags.MatchesAny(this->Info.IgnoreTags, false);
+
+	return HasRequired && !HasIgnored;
+}
+
+// --------------------------------------------------------------------------------------------------------------------------------------------------------
+//
+//	FGameplayEffectInstigatorContext
+//
+// --------------------------------------------------------------------------------------------------------------------------------------------------------
+
+void FGameplayEffectInstigatorContext::AddInstigator(class AActor *InInstigator)
+{
+	Instigator = InInstigator;
+	InstigatorAttributeComponent = NULL;
+
+	// Cache off his attribute component.
+	if (Instigator)
+	{
+		TArray<UObject*> ChildObjects;
+		GetObjectsWithOuter(Instigator, ChildObjects);
+
+		for (UObject * Obj : ChildObjects)
+		{
+			if (Obj && Obj->GetClass()->IsChildOf(UAttributeComponent::StaticClass()))
+			{
+				InstigatorAttributeComponent = Cast<UAttributeComponent>(Obj);
+				break;
+			}
+		}
+	}
+}
+
+void FGameplayEffectInstigatorContext::AddHitResult(const FHitResult InHitResult)
+{
+	check(!HitResult.IsValid());
+	HitResult = TSharedPtr<FHitResult>(new FHitResult(InHitResult));
+}
+
+bool FGameplayEffectInstigatorContext::NetSerialize(FArchive& Ar, class UPackageMap* Map, bool& bOutSuccess)
+{
+	Ar << Instigator;
+
+	bool HasHitResults = HitResult.IsValid();
+	Ar << HasHitResults;
+	if (Ar.IsLoading())
+	{
+		if (HasHitResults)
+		{
+			if (!HitResult.IsValid())
+			{
+				HitResult = TSharedPtr<FHitResult>(new FHitResult());
+			}
+			AddInstigator(Instigator); // Just to initialize InstigatorAttributeComponent
+		}
+	}
+
+	if (HasHitResults == 1)
+	{
+		Ar << HitResult->BoneName;
+		HitResult->Location.NetSerialize(Ar, Map, bOutSuccess);
+		HitResult->Normal.NetSerialize(Ar, Map, bOutSuccess);
+	}
+
+	bOutSuccess = true;
+	return true;
+}
+
+// --------------------------------------------------------------------------------------------------------------------------------------------------------
+//
+//	FAggregatorRef
+//
+// --------------------------------------------------------------------------------------------------------------------------------------------------------
+
+void FAggregatorRef::MakeUnique()
+{
+	SKILL_LOG(Log, TEXT("MakeUnique %s"), *ToString());
+
+	// Make a hard ref copy of our FAggregator
+	MakeHardRef();
+	SharedPtr = TSharedPtr<FAggregator>(new FAggregator(*SharedPtr.Get()));
+	WeakPtr = SharedPtr;
+
+	// Update dependancy chain so that the copy we just made is updated if any of the applied modifiers change
+	SharedPtr->RefreshDependencies();
+}
+
+void FAggregatorRef::MakeUniqueDeep()
+{
+	SKILL_LOG_SCOPE(TEXT("MakeUniqueDeep %s"), *ToString());
+
+	// Make a hard ref copy of our FAggregator
+	MakeHardRef();
+	SharedPtr = TSharedPtr<FAggregator>(new FAggregator(*SharedPtr.Get()));
+	WeakPtr = SharedPtr;
+
+	// Make all of our mods UniqueDeep
+	Get()->MakeUniqueDeep();
+
+	// Update dependancy chain so that the copy we just made is updated if any of the applied modifiers change
+	SharedPtr->RefreshDependencies();
+}
+
+FString FAggregatorRef::ToString() const
+{
+	if (SharedPtr.IsValid())
+	{
+		return FString::Printf(TEXT("[HardRef to: %s]"), *SharedPtr->ToSimpleString());
+	}
+	if (WeakPtr.IsValid())
+	{
+		return FString::Printf(TEXT("[SoftRef to: %s]"), *WeakPtr.Pin()->ToSimpleString());
+	}
+	return FString(TEXT("Invalid"));
+}
+
+bool FAggregatorRef::NetSerialize(FArchive& Ar, class UPackageMap* Map, bool& bOutSuccess)
+{
+	float EvaluatedMagnitude=0.f;
+
+	if (Ar.IsSaving())
+	{
+		if (IsValid())
+		{
+			EvaluatedMagnitude = Get()->Evaluate().Magnitude;
+		}
+	}
+
+	Ar << EvaluatedMagnitude;
+
+	if (Ar.IsLoading())
+	{
+		if (!IsValid())
+		{
+			SharedPtr = TSharedPtr<struct FAggregator>(new FAggregator);
+			WeakPtr = SharedPtr;
+		}
+
+		Get()->SetFromNetSerialize(EvaluatedMagnitude);
+	}
+	return true;
+}
+
+void FAggregator::SetFromNetSerialize(float SerializedValue)
+{
+	BaseData = FGameplayModifierData(SerializedValue, NULL);
+	CachedData = FGameplayModifierEvaluatedData(SerializedValue);
+}
+
+// --------------------------------------------------------------------------------------------------------------------------------------------------------
+//
+//	FAggregator
+//
+// --------------------------------------------------------------------------------------------------------------------------------------------------------
+
+FAggregator::FAggregator()
+{
+#if SKILL_SYSTEM_AGGREGATOR_DEBUG
+	CopiesMade = 0;
+	FAggregator::AllocationStats.DefaultCStor++;
+#endif
+}
+
+FAggregator::FAggregator(const FGameplayModifierData &InBaseData, TSharedPtr<FGameplayEffectLevelSpec> InLevel, const TCHAR *InDebugStr)
+: Level(InLevel)
+, BaseData(InBaseData)
+{
+#if SKILL_SYSTEM_AGGREGATOR_DEBUG
+	if (InDebugStr)
+	{
+		DebugString = FString(InDebugStr);
+	}
+	CopiesMade = 0;
+	FAggregator::AllocationStats.ModifierCStor++;
+#endif
+}
+
+FAggregator::FAggregator(const FScalableFloat &InBaseMagnitude, TSharedPtr<FGameplayEffectLevelSpec> LevelInfo, const TCHAR *InDebugStr)
+: Level(LevelInfo)
+, BaseData(InBaseMagnitude)
+{
+#if SKILL_SYSTEM_AGGREGATOR_DEBUG
+	if (InDebugStr)
+	{
+		DebugString = FString(InDebugStr);
+	}
+	CopiesMade = 0;
+	FAggregator::AllocationStats.ScalableFloatCstor++;
+#endif
+}
+
+FAggregator::FAggregator(const FGameplayModifierEvaluatedData &InEvalData, const TCHAR *InDebugStr)
+: Level(TSharedPtr<FGameplayEffectLevelSpec>(new FGameplayEffectLevelSpec()))
+, BaseData(InEvalData.Magnitude, InEvalData.Callbacks)
+{
+#if SKILL_SYSTEM_AGGREGATOR_DEBUG
+	if (InDebugStr)
+	{
+		DebugString = FString(InDebugStr);
+	}
+	CopiesMade = 0;
+	FAggregator::AllocationStats.FloatCstor++;
+#endif
+}
+
+FAggregator::FAggregator(const FAggregator &In)
+{
+	*this = In;
+#if SKILL_SYSTEM_AGGREGATOR_DEBUG
+	FAggregator::AllocationStats.CopyCstor++;
+	CopiesMade = 0;
+	DebugString = FString::Printf(TEXT("Copy %d of [%s]"), ++In.CopiesMade, *In.DebugString);
+#endif
+} 
+
+FAggregator::~FAggregator()
+{
+	for (TWeakPtr<FAggregator> WeakPtr : Dependants)
+	{
+		if (WeakPtr.IsValid())
+		{
+			SKILL_LOG(Log, TEXT("%s Marking Dependant %s Dirty on Destroy"), *ToSimpleString(), *WeakPtr.Pin()->ToSimpleString());
+
+			WeakPtr.Pin()->MarkDirty();
+#if SKILL_SYSTEM_AGGREGATOR_DEBUG
+			FAggregator::AllocationStats.DependantsUpdated++;
+#endif
+		}
+	}
+}
+
+void FAggregator::RefreshDependencies()
+{
+	TWeakPtr<FAggregator> LocalWeakPtr(SharedThis(this));
+	for (int32 i = 0; i < EGameplayModOp::Max; ++i)
+	{
+		for (FAggregatorRef &Ref : Mods[i])
+		{
+			if (Ref.IsValid())
+			{
+				Ref.Get()->AddDependantAggregator(LocalWeakPtr);
+			}
+		}
+	}
+	
+	RegisterLevelDependancies();
+}
+
+void FAggregator::ApplyMod(EGameplayModOp::Type ModType, FAggregatorRef Ref, bool TakeSnapshot)
+{
+	if (TakeSnapshot)
+	{
+		Ref.MakeUniqueDeep();
+	}
+	else
+	{
+		Ref.MakeSoftRef();
+	}
+	Mods[ModType].Push(Ref);
+
+	// Make the ref tell us if it changes
+	if (HasBeenAlreadyMadeSharable())
+	{
+		Ref.Get()->AddDependantAggregator(TWeakPtr< FAggregator >(SharedThis(this)));
+	}
+	MarkDirty();
+}
+
+// Execute is intended to directly modify the base value of an aggregator.
+// However if the aggregator's base value is not static (scales with some outside influence)
+// we will treat this execute as an apply, but make our own copy of the passed in aggregator.
+// (so that it will essentially be permanent).
+void FAggregator::ExecuteModAggr(EGameplayModOp::Type ModType, FAggregatorRef Ref)
+{
+	const FGameplayModifierEvaluatedData& EvaluatedData = Ref.Get()->Evaluate();
+	ExecuteMod(ModType, EvaluatedData);	
+}
+
+void FAggregator::ExecuteMod(EGameplayModOp::Type ModType, const FGameplayModifierEvaluatedData& EvaluatedData)
+{
+	// Am I allowed to even be executed on? If my base data scales then all I can do is apply this to myself
+	if (!BaseData.Magnitude.IsStatic())
+	{
+		SKILL_LOG(Log, TEXT("Treating Execute as an Apply since our Level is not static!"));
+		FAggregatorRef NewAggregator(new FAggregator(EvaluatedData, SKILL_AGG_DEBUG(TEXT("ExecutedMod"))));
+		ApplyMod(ModType, NewAggregator, false);
+		return;
+	}
+
+	switch (ModType)
+	{
+		case EGameplayModOp::Override:
+		{
+			BaseData.Magnitude.SetValue(EvaluatedData.Magnitude);
+			BaseData.Tags = EvaluatedData.Tags;
+			break;
+		}
+		case EGameplayModOp::Additive:
+		{
+			BaseData.Magnitude.Value += EvaluatedData.Magnitude;
+			BaseData.Tags.AppendTags(EvaluatedData.Tags);
+			break;
+		}
+		case EGameplayModOp::Multiplicitive:
+		{
+			BaseData.Magnitude.Value *= EvaluatedData.Magnitude;
+			BaseData.Tags.AppendTags(EvaluatedData.Tags);
+			break;
+		}
+		case EGameplayModOp::Division:
+		{
+			BaseData.Magnitude.Value /= EvaluatedData.Magnitude;
+			BaseData.Tags.AppendTags(EvaluatedData.Tags);
+			break;
+		}
+		case EGameplayModOp::Callback:
+		{
+			// Executing a callback mod on another aggregator, what is expected here?
+			check(false);
+			break;
+		}
+	}
+
+	SKILL_LOG(Log, TEXT("ExecuteMod: %s new BaseData.Magnitude: %s"), *ToSimpleString(), *BaseData.Magnitude.ToSimpleString());
+	MarkDirty();
+}
+
+void FAggregator::AddDependantAggregator(TWeakPtr<FAggregator> InDependant)
+{
+	check(InDependant.IsValid());
+	check(!Dependants.Contains(InDependant));
+
+	SKILL_LOG(Log, TEXT("AddDependantAggregator: %s is a dependant of %s"), *InDependant.Pin()->ToSimpleString(), *ToSimpleString());
+	Dependants.Add(InDependant);
+}
+
+void FAggregator::TakeSnapshotOfLevel()
+{
+	check(Level.IsValid());
+	Level = TSharedPtr<FGameplayEffectLevelSpec>( new FGameplayEffectLevelSpec( *Level.Get()) );
+}
+
+void FAggregator::RegisterLevelDependancies()
+{
+	if (!BaseData.Magnitude.IsStatic())
+	{
+		TWeakPtr<FAggregator> LocalWeakPtr(SharedThis(this));
+		Level->RegisterLevelDependancy(LocalWeakPtr);
+	}
+}
+
+// Please try really hard to never add a "force full re-evaluate" flag to this function!
+// We want to strive to make this system dirty the cached data when its actually dirtied, 
+// and never do full catch all rebuilds.
+
+const FGameplayModifierEvaluatedData& FAggregator::Evaluate() const
+{
+	SKILL_LOG_SCOPE(TEXT("Aggregator Evaluate %s"), *ToSimpleString());
+
+	if (!CachedData.IsValid)
+	{
+		// ------------------------------------------------------------------------
+		// If there are any overrides, then just take the first valid one.
+		// ------------------------------------------------------------------------
+		for (const FAggregatorRef &Agg : Mods[EGameplayModOp::Override])
+		{
+			SKILL_LOG_SCOPE(TEXT("EGameplayModOp::Override"));
+			if (Agg.IsValid())
+			{
+				CachedData = Agg.Get()->Evaluate();
+				return CachedData;
+			}
+		}
+
+		// If we are going to do math, we need to lock our base value in.
+		// Calculate our magnitude at our level
+		// (We may not have a level, in that case, ensure we also don't have a leveling table)
+
+		float EvaluatedMagnitude = 0.f;
+		if (Level->IsValid())
+		{
+			EvaluatedMagnitude = BaseData.Magnitude.GetValueAtLevel(Level->GetLevel());
+		}
+		else
+		{
+			EvaluatedMagnitude = BaseData.Magnitude.GetValueChecked();
+		}
+
+		CachedData = FGameplayModifierEvaluatedData(EvaluatedMagnitude, BaseData.Callbacks, ActiveHandle, &BaseData.Tags);
+
+		int32 TotalModCount = Mods[EGameplayModOp::Additive].Num() + Mods[EGameplayModOp::Multiplicitive].Num() + Mods[EGameplayModOp::Division].Num();
+		
+		// Early out if no mods.
+		// We need to calculate num of mods anyways to do ModLIst.Reserve.
+		if (TotalModCount <= 0)
+		{
+			SKILL_LOG(Log, TEXT("Final Magnitude: %.2f"), CachedData.Magnitude);
+			return CachedData;
+		}
+
+		// ------------------------------------------------------------------------
+		//	Apply Numeric Modifiers
+		//		This is convoluted due to tagging. We have modifiers that require we have or don't have certain tags.
+		//		These mods can also give us new tags. Its possible the 1st mod will require a tag that the 2nd mod gives us.
+		//
+		//	The basic approach here is create an ordered, linear list of all modifiers:
+		//		[Additive][Multipliticitive][Division] mods
+		//
+		//	We make a pass through the list, aggregating as we go. During a pass we keep track of what what tags we've added
+		//	and if there were any mods that we rejected due to not having tags. When the pass is over, we check if we added
+		//	any that we needed. If so, we make another pass. (Once a modifier aggregated, we remove it from the list).
+		//
+		//	Paradoxes are still possible. ModX gives tag A, requires we don't have tag B. ModY gives tag B, requires we don't have tag A.
+		//	We detect this in a single pass. In the above example we would aggregate ModX, but warn loudly when we found ModY.
+		//	(we expect content to solve this via stacking rules, or just not making these type of requirements).
+		//
+		//
+		//
+		//	
+		// ------------------------------------------------------------------------
+		
+		// We have to do tag aggregation to figure out what we can and can't apply.
+		FGameplayTagContainer CommittedIgnoreTags;
+
+		TArray<const FAggregator*> ModList;
+		ModList.Reserve(TotalModCount + (EGameplayModOp::Override - EGameplayModOp::Additive));
+
+		// Build linear list of what we will aggregate
+		for (int32 OpIdx = (EGameplayModOp::Additive); OpIdx < EGameplayModOp::Override; ++OpIdx)
+		{
+			for (const FAggregatorRef &Agg : Mods[OpIdx])
+			{	
+				ModList.Add(Agg.Get());
+			}
+			ModList.Add(this);	// Sentinel value to signify "NextOp"
+		}
+
+		static const float OpBias[EGameplayModOp::Override] = {
+			0.f,	// EGameplayModOp::Additive
+			1.f,	// EGameplayModOp::Multiplicitive
+			1.f,	// EGameplayModOp::Division
+		};
+
+		float OpAggregation[EGameplayModOp::Override] = {
+			0.f,	// EGameplayModOp::Additive
+			1.f,	// EGameplayModOp::Multiplicitive
+			1.f,	// EGameplayModOp::Division
+		};
+
+		// Make multiple passes to do tagging
+		while (true)
+		{
+			FGameplayTagContainer	NewGiveTags;
+			FGameplayTagContainer	NewIgnoreTags;
+			FGameplayTagContainer	MissingTags;
+
+			EGameplayModOp::Type	ModOp = (EGameplayModOp::Additive);
+			
+			for (int32 ModIdx = 0; ModIdx < ModList.Num()-1; ++ModIdx)
+			{
+				const FAggregator * Agg = ModList[ModIdx];
+				if (Agg == NULL)
+				{
+					continue;
+				}
+				if (Agg == this)
+				{
+					ModOp = static_cast<EGameplayModOp::Type>(static_cast<int32>(ModOp)+1);
+					check(ModOp < EGameplayModOp::Override);
+					continue;
+				}
+				
+				const FGameplayModifierData &ModBaseData = ModList[ModIdx]->BaseData;
+				const FGameplayTagContainer &ModIgnoreTags = ModBaseData.IgnoreTags;
+				const FGameplayTagContainer &ModRequireTags = ModBaseData.RequireTags;
+
+				// This mod requires we have certain tags
+				if (ModRequireTags.Num() > 0 && !CachedData.Tags.MatchesAll(ModRequireTags, false) && !NewGiveTags.MatchesAll(ModRequireTags, false))
+				{
+					// But something else could give us this tag! So keep track of it and don't remove this Mod from the ModList.
+					MissingTags.AppendTags(ModRequireTags);
+					continue;
+				}
+
+				// This mod is now either accepted or rejected. It will not be needed for subsequent passes in this Evaluate, so NULL It out now.
+				ModList[ModIdx] = NULL;
+				
+				// This mod requires we don't have certain tags
+				if (CachedData.Tags.MatchesAny(ModIgnoreTags, false))
+				{
+					continue;
+				}
+
+				// Check for conflicts within this pass
+				if (ModIgnoreTags.Num() > 0 && NewGiveTags.Num() > 0 && ModIgnoreTags.MatchesAny(NewGiveTags, false))
+				{
+					// Pass problem!
+					SKILL_LOG(Warning, TEXT("Tagging conflicts during Aggregate! Use Stacking rules to avoid this"));
+					SKILL_LOG(Warning, TEXT("   While Evaluating: %s"), *ToSimpleString());
+					SKILL_LOG(Warning, TEXT("   Applying Mod: %s"), *Agg->ToSimpleString());
+					SKILL_LOG(Warning, TEXT("   ModIgnoreTags: %s"), *ModIgnoreTags.ToString() );
+					SKILL_LOG(Warning, TEXT("   NewGiveTags: %s"), *NewGiveTags.ToString());
+					continue;
+				}
+
+				// Our requirements on the mod
+				const FGameplayModifierEvaluatedData& ModEvaluatedData = Agg->Evaluate();
+
+				// We have already committed during this aggregation to not have certain tags
+				if (CommittedIgnoreTags.Num() > 0 && CommittedIgnoreTags.MatchesAny(ModEvaluatedData.Tags, false))
+				{
+					continue;
+				}
+
+				if (NewIgnoreTags.Num() > 0 && ModEvaluatedData.Tags.Num() > 0 && NewIgnoreTags.MatchesAny(ModEvaluatedData.Tags, false))
+				{
+					// Pass problem!
+					SKILL_LOG(Warning, TEXT("Tagging conflicts during Aggregate! Use Stacking rules to avoid this"));
+					SKILL_LOG(Warning, TEXT("   While Evaluating: %s"), *ToSimpleString());
+					SKILL_LOG(Warning, TEXT("   Applying Mod: %s"), *Agg->ToSimpleString());
+					SKILL_LOG(Warning, TEXT("   Mod Tags: %s"), *ModEvaluatedData.Tags.ToString());
+					SKILL_LOG(Warning, TEXT("   NewIgnoreTags: %s"), *NewIgnoreTags.ToString());
+					continue;
+				}
+
+				// Commit this Mod
+				NewIgnoreTags.AppendTags(ModIgnoreTags);
+				ModEvaluatedData.Aggregate(NewGiveTags, OpAggregation[ModOp], OpBias[ModOp]);
+			}
+
+			// Commit this pass's tags
+			CachedData.Tags.AppendTags(NewGiveTags);
+
+			// Keep doing passes until we don't add any new tags or don't have any missing tags
+			if (MissingTags.Num() <= 0 || !MissingTags.MatchesAny( NewGiveTags, false ))
+			{
+				break;
+			}
+		}
+
+		float Division = OpAggregation[EGameplayModOp::Division] > 0.f ? OpAggregation[EGameplayModOp::Division] : 1.f;
+
+		CachedData.Magnitude = ((CachedData.Magnitude + OpAggregation[EGameplayModOp::Additive]) * OpAggregation[EGameplayModOp::Multiplicitive]) / Division;
+
+		SKILL_LOG(Log, TEXT("Final Magnitude: %.2f"), CachedData.Magnitude);
+		
+	}
+	else
+	{
+		SKILL_LOG(Log, TEXT("CachedData was valid. Magnitude: %.2f"), CachedData.Magnitude);
+	}
+	return CachedData;
+}
+
+void FAggregator::PreEvaluate(FGameplayEffectModCallbackData &Data) const
+{
+	check(Data.ModifierSpec.Aggregator.Get() == this);
+	for (const FAggregatorRef &Agg : Mods[EGameplayModOp::Callback])
+	{
+		if (Agg.IsValid())
+		{
+			Agg.Get()->Evaluate().InvokePreExecute(Data);
+		}
+	}
+}
+
+void FAggregator::PostEvaluate(const struct FGameplayEffectModCallbackData &Data) const
+{
+	check(Data.ModifierSpec.Aggregator.Get() == this);
+	for (const FAggregatorRef &Agg : Mods[EGameplayModOp::Callback])
+	{
+		if (Agg.IsValid())
+		{
+			Agg.Get()->Evaluate().InvokePostExecute(Data);
+		}
+	}
+}
+
+FAggregator & FAggregator::MarkDirty()
+{
+	CachedData.IsValid = false;
+
+	// Execute OnDirty callbacks first. This may do things like update the actual uproperty value of an attribute.
+	OnDirty.ExecuteIfBound(this);
+
+	// Now tell people who depend on my value that I have changed. Important to do this after the OnDirty callback has been called.
+	for (int32 i=0; i < Dependants.Num(); ++i)
+	{
+		TWeakPtr<FAggregator> WeakPtr = Dependants[i];
+		if (WeakPtr.IsValid())
+		{
+			SKILL_LOG(Log, TEXT("%s Marking Dependant %s Dirty (from ::MarkDirty())"), *ToSimpleString(), *WeakPtr.Pin()->ToSimpleString());
+#if SKILL_SYSTEM_AGGREGATOR_DEBUG
+			FAggregator::AllocationStats.DependantsUpdated++;
+#endif
+			WeakPtr.Pin()->MarkDirty();
+		}
+		else
+		{
+			Dependants.RemoveAtSwap(i);
+			--i;
+		}
+	}
+
+	return *this;
+}
+
+void FAggregator::MakeUniqueDeep()
+{
+	for (int32 i = 0; i < EGameplayModOp::Max; ++i)
+	{
+		for (FAggregatorRef &Ref : Mods[i])
+		{
+			if (Ref.IsValid())
+			{
+				Ref.MakeUniqueDeep();
+			}
+		}
+	}
+}
+
+void FAggregator::ClearAllDependancies()
+{
+	Dependants.SetNum(0);
+	OnDirty.Unbind();
+}
+
+// --------------------------------------------------------------------------------------------------------------------------------------------------------
+//
+//	FGameplayModifierEvaluatedData
+//
+// --------------------------------------------------------------------------------------------------------------------------------------------------------
+
+void FGameplayModifierEvaluatedData::Aggregate(FGameplayTagContainer &OutTags, float &OutMagnitude, const float Bias) const
+{
+	OutMagnitude += (Magnitude - Bias);
+	OutTags.AppendTags(Tags);
+}
+
+void FGameplayModifierEvaluatedData::InvokePreExecute(FGameplayEffectModCallbackData &Data) const
+{
+	if (Callbacks)
+	{
+		for (TSubclassOf<class UGameplayEffectExtension> ExtClass : Callbacks->ExtensionClasses)
+		{
+			if (ExtClass)
+			{
+				UGameplayEffectExtension * Ext = ExtClass->GetDefaultObject<UGameplayEffectExtension>();
+				Ext->PreGameplayEffectExecute(*this, Data);
+			}
+		}
+	}
+}
+
+void FGameplayModifierEvaluatedData::InvokePostExecute(const FGameplayEffectModCallbackData &Data) const
+{
+	if (Callbacks)
+	{
+		for (TSubclassOf<class UGameplayEffectExtension> ExtClass : Callbacks->ExtensionClasses)
+		{
+			if (ExtClass)
+			{
+				UGameplayEffectExtension * Ext = ExtClass->GetDefaultObject<UGameplayEffectExtension>();
+				Ext->PostGameplayEffectExecute(*this, Data);
+			}
+		}
+	}
+}
+
+// --------------------------------------------------------------------------------------------------------------------------------------------------------
+//
+//	FActiveGameplayEffect
+//
+// --------------------------------------------------------------------------------------------------------------------------------------------------------
+
+bool FActiveGameplayEffect::CanBeStacked(const FActiveGameplayEffect& Other) const
+{
+	return (Handle != Other.Handle) && Spec.GetStackingType() == Other.Spec.GetStackingType() && Spec.Def->Modifiers[0].Attribute == Other.Spec.Def->Modifiers[0].Attribute;
+}
+
+void FActiveGameplayEffect::PreReplicatedRemove(const struct FActiveGameplayEffectsContainer &InArray)
+{
+	InArray.Owner->InvokeGameplayCueRemoved(Spec);
+}
+
+void FActiveGameplayEffect::PostReplicatedAdd(const struct FActiveGameplayEffectsContainer &InArray)
+{
+	static const int32 MAX_DELTA_TIME = 3;
+
+	InArray.Owner->InvokeGameplayCueAdded(Spec);
+	// Was this actually just activated, or are we just finding out about it due to relevancy/join in progress?
+	float WorldTimeSeconds = InArray.GetWorldTime();
+	int32 GameStateTime = InArray.GetGameStateTime();
+
+	int32 DeltaGameStateTime = GameStateTime - StartGameStateTime;	// How long we think the effect has been playing (only 1 second precision!)
+
+	if (GameStateTime > 0 && FMath::Abs(DeltaGameStateTime) < MAX_DELTA_TIME)
+	{
+		InArray.Owner->InvokeGameplayCueActivated(Spec);
+	}
+
+	// Set our local start time accordingly
+
+	StartWorldTime = WorldTimeSeconds - static_cast<float>(DeltaGameStateTime);
+}
+
+// --------------------------------------------------------------------------------------------------------------------------------------------------------
+//
+//	FActiveGameplayEffectsContainer
+//
+// --------------------------------------------------------------------------------------------------------------------------------------------------------
+
+bool FActiveGameplayEffectsContainer::ApplyActiveEffectsTo(OUT FGameplayEffectSpec &Spec, const FModifierQualifier &QualifierContext)
+{
+	FActiveGameplayEffectHandle().IsValid();
+
+	SKILL_LOG_SCOPE(TEXT("ApplyActiveEffectsTo: %s %s"), *Spec.ToSimpleString(), *QualifierContext.ToString());
+	for (FActiveGameplayEffect & ActiveEffect : GameplayEffects)
+	{
+		// We dont want to use FModifierQualifier::TestTarget here, since we aren't the 'target'. We are applying stuff to Spec which will be applied to a target.
+		if (QualifierContext.IgnoreHandle().IsValid() && QualifierContext.IgnoreHandle() == ActiveEffect.Handle)
+		{
+			continue;
+		}
+		// should we try to apply ActiveEffect to Spec?
+		float ChanceToExecute = ActiveEffect.Spec.GetChanceToExecuteOnGameplayEffect();
+		if ((ChanceToExecute < 1.f - SMALL_NUMBER) && (FMath::FRand() > ChanceToExecute))
+		{
+			continue;
+		}
+		if (Spec.ApplyModifiersFrom(ActiveEffect.Spec, QualifierContext) < 0)
+		{
+			// ActiveEffect provides immunity to Spec. No need to look at other active effects
+			return false;
+		}
+	}
+
+	return true;
+}
+
+/** This is the main function that applies/attaches a GameplayEffect on Attributes and ActiveGameplayEffects */
+void FActiveGameplayEffectsContainer::ApplySpecToActiveEffectsAndAttributes(FGameplayEffectSpec &Spec, const FModifierQualifier &QualifierContext)
+{
+	for (const FModifierSpec &Mod : Spec.Modifiers)
+	{
+		if (Mod.Info.ModifierType == EGameplayMod::Attribute)
+		{
+			// Todo: Tag/application checks here
+
+			SKILL_LOG_SCOPE(TEXT("Applying Attribute Mod %s to property"), *Mod.ToSimpleString());
+
+			FAggregator * Aggregator = FindOrCreateAttributeAggregator(Mod.Info.Attribute).Get();
+
+			// Add the modifier to the property value aggregator
+			// Note that this will immediately invoke the callback to update the attribute's current value, so we don't have to explicitly do it here.
+			Aggregator->ApplyMod(Mod.Info.ModifierOp, Mod.Aggregator, Spec.ShouldApplyAsSnapshot(QualifierContext));
+		}
+
+		if (Mod.Info.ModifierType == EGameplayMod::ActiveGE)
+		{
+			SKILL_LOG_SCOPE(TEXT("Applying Mod %s to ActiveEffects"),  *Mod.ToSimpleString());
+
+			bool bUpdateStackedEffects = false;
+
+			// TODO: Tag checks here
+
+			// This modifies GEs that are currently active, so apply this to them...
+			for (FActiveGameplayEffect & ActiveEffect : GameplayEffects)
+			{
+				if (!QualifierContext.TestTarget(ActiveEffect.Handle))
+				{
+					continue;
+				}
+
+				ActiveEffect.Spec.ApplyModifiersFrom( Spec, FModifierQualifier().Type(EGameplayMod::ActiveGE) );
+				bUpdateStackedEffects = true;
+			}
+
+			if (bUpdateStackedEffects)
+			{
+				StacksNeedToRecalculate();
+			}
+		}
+	}
+}
+
+/** This is the main function that executes a GameplayEffect on Attributes and ActiveGameplayEffects */
+void FActiveGameplayEffectsContainer::ExecuteActiveEffectsFrom(const FGameplayEffectSpec &Spec, const FModifierQualifier &QualifierContext)
+{
+	bool InvokeGameplayCueExecute = false;
+
+	// check if this is a stacking effect and if it is the active stacking effect.
+	if (Spec.GetStackingType() != EGameplayEffectStackingPolicy::Unlimited)
+	{
+		// we're a stacking attribute but not active
+		if (Spec.bTopOfStack == false)
+		{
+			return;
+		}
+	}
+
+	float ChanceToExecute = Spec.GetChanceToExecuteOnGameplayEffect();
+
+	// check if the new effect removes or modifies existing effects
+	if (Spec.Def->AppliesImmunityTo == EGameplayImmunity::ActiveGE)
+	{
+		for (int32 ii = 0; ii < GameplayEffects.Num(); ++ii)
+		{
+			// should we try to apply this to the current effect?
+			if ((ChanceToExecute < 1.f - SMALL_NUMBER) && (FMath::FRand() > ChanceToExecute))
+			{
+				continue;
+			}
+
+			if (Spec.Def->AreGameplayEffectTagRequirementsSatisfied(GameplayEffects[ii].Spec.Def))
+			{
+				GameplayEffects.RemoveAtSwap(ii);
+				--ii;
+				continue;
+			}
+
+			// we kept the effect, now check its mods to see if we should remove any of them
+			for (int32 jj = 0; jj < GameplayEffects[ii].Spec.Modifiers.Num(); ++jj)
+			{
+				const FGameplayModifierEvaluatedData& Data = GameplayEffects[ii].Spec.Modifiers[jj].Aggregator.Get()->Evaluate();
+				if (GameplayEffects[ii].Spec.Def->AreGameplayEffectTagRequirementsSatisfied(Data.Tags))
+				{
+					GameplayEffects[ii].Spec.Modifiers.RemoveAtSwap(jj);
+					--jj;
+				}
+			}
+		}
+	}
+
+	bool bModifiesActiveGEs = false;
+
+	for (const FModifierSpec &Mod : Spec.Modifiers)
+	{
+		if (Mod.Info.ModifierType == EGameplayMod::Attribute)
+		{
+			UAttributeSet * AttributeSet = Owner->GetAttributeSubobject(Mod.Info.Attribute.GetAttributeSetClass());
+			if (AttributeSet == NULL)
+			{
+				// Our owner doesn't have this attribute, so we can't do anything
+				SKILL_LOG(Log, TEXT("%s does not have attribute %s. Skipping modifier"), *Owner->GetPathName(), *Mod.Info.Attribute.GetName());
+				continue;
+			}
+
+			// Todo: Tags/application checks here - make sure we can still apply
+			InvokeGameplayCueExecute = true;
+
+			SKILL_LOG_SCOPE(TEXT("Executing Attribute Mod %s"), *Mod.ToSimpleString());
+
+			// First, evaluate all of our data 
+
+			FGameplayModifierEvaluatedData EvaluatedData = Mod.Aggregator.Get()->Evaluate();
+
+			FGameplayEffectModCallbackData ExecuteData(Spec, Mod, EvaluatedData, *Owner);
+						
+			/** This should apply 'gameplay effect specific' rules, such as life steal, shields, etc */
+			Mod.Aggregator.Get()->PreEvaluate(ExecuteData);
+
+			/** This should apply 'gamewide' rules. Such as clamping Health to MaxHealth or granting +3 health for every point of strength, etc */
+			AttributeSet->PreAttributeModify(ExecuteData);
+
+			// Do we have active GE's that are already modifying this?
+			FAggregatorRef *RefPtr = OngoingPropertyEffects.Find(Mod.Info.Attribute);
+			if (RefPtr)
+			{
+				SKILL_LOG(Log, TEXT("Property %s has active mods. Adding to Aggregator."), *Mod.Info.Attribute.GetName());
+				RefPtr->Get()->ExecuteMod(Mod.Info.ModifierOp, EvaluatedData);
+			}
+			else
+			{
+				// Modify the property inplace, without putting it in the OngoingPropertyEffects map
+				float		CurrentValueOfProperty = Owner->GetNumericAttribute(Mod.Info.Attribute);
+				FAggregator	Aggregator(CurrentValueOfProperty, SKILL_AGG_DEBUG(TEXT("Inplace Attribute %s"), *Mod.Info.Attribute.GetName()));
+
+				Aggregator.ExecuteMod(Mod.Info.ModifierOp, EvaluatedData);
+				
+				const float NewPropertyValue = Aggregator.Evaluate().Magnitude;
+
+				SKILL_LOG(Log, TEXT("Property %s new value is: %.2f [was: %.2f]"), *Mod.Info.Attribute.GetName(), NewPropertyValue, CurrentValueOfProperty);
+				Owner->SetNumericAttribute(Mod.Info.Attribute, NewPropertyValue);
+			}
+
+			/** This should apply 'gameplay effect specific' rules, such as life steal, shields, etc */
+			Mod.Aggregator.Get()->PostEvaluate(ExecuteData);
+
+			/** This should apply 'gamewide' rules. Such as clamping Health to MaxHealth or granting +3 health for every point of strength, etc */
+			AttributeSet->PostAttributeModify(ExecuteData);
+			
+		}
+		else if(Mod.Info.ModifierType == EGameplayMod::ActiveGE)
+		{
+			bModifiesActiveGEs = true;
+		}
+	}
+
+	// If any of the mods apply to active gameplay effects try to apply them now.
+	// We need to do this here because all of the modifiers either need to apply or not apply to each active gameplay effect based on their chance to execute on gameplay effects.
+	if (bModifiesActiveGEs)
+	{
+		for (FActiveGameplayEffect & ActiveEffect : GameplayEffects)
+		{
+			// Don't apply spec to itself
+			if (!QualifierContext.TestTarget(ActiveEffect.Handle))
+			{
+				continue;
+			}
+
+			// should we try to apply this to the current effect?
+			if ((ChanceToExecute < 1.f - SMALL_NUMBER) && (FMath::FRand() > ChanceToExecute))
+			{
+				continue;
+			}
+
+			if (ActiveEffect.Spec.ExecuteModifiersFrom(Spec, FModifierQualifier().Type(EGameplayMod::ActiveGE)) > 0)
+			{
+				InvokeGameplayCueExecute = true;
+			}
+		}
+	}
+
+	if (InvokeGameplayCueExecute)
+	{
+		// TODO: check replication policy. Right now we will replicate every execute via a multicast RPC
+
+		SKILL_LOG(Log, TEXT("Invoking Execute GameplayCue for %s"), *Spec.ToSimpleString() );
+		Owner->NetMulticast_InvokeGameplayCueExecuted(Spec);
+	}
+
+}
+
+void FActiveGameplayEffectsContainer::ExecutePeriodicGameplayEffect(FActiveGameplayEffectHandle Handle)
+{
+	// Could make this a map for quicker lookup
+	for (FActiveGameplayEffect& Effect : GameplayEffects)
+	{
+		if (Effect.Handle == Handle)
+		{
+			// Execute
+			ExecuteActiveEffectsFrom(Effect.Spec, FModifierQualifier().IgnoreHandle(Effect.Handle));
+
+			return;
+		}
+	}
+}
+
+void FActiveGameplayEffectsContainer::AddDependancyToAttribute(FGameplayAttribute Attribute, const TWeakPtr<FAggregator> InDependant)
+{
+	FAggregator * Aggregator = FindOrCreateAttributeAggregator(Attribute).Get();
+	Aggregator->AddDependantAggregator(InDependant);
+}
+
+FAggregatorRef & FActiveGameplayEffectsContainer::FindOrCreateAttributeAggregator(FGameplayAttribute Attribute)
+{
+	FAggregatorRef *RefPtr = OngoingPropertyEffects.Find(Attribute);
+	if (RefPtr)
+	{
+		return *RefPtr;
+	}
+
+	// Create a new aggregator for this attribute.
+	float CurrentValueOfProperty = Owner->GetNumericAttribute(Attribute);
+	SKILL_LOG(Log, TEXT("Creating new entry in OngoingPropertyEffect map for AddDependancyToAttribute. CurrentValue: %.2f"), CurrentValueOfProperty);
+
+	FAggregator *NewPropertyAggregator = new FAggregator(FGameplayModifierEvaluatedData(CurrentValueOfProperty), SKILL_AGG_DEBUG(TEXT("Attribute %s Aggregator"), *Attribute.GetName()));
+
+	NewPropertyAggregator->OnDirty = FAggregator::FOnDirty::CreateRaw(this, &FActiveGameplayEffectsContainer::OnPropertyAggregatorDirty, Attribute);
+
+	return OngoingPropertyEffects.Add(Attribute, FAggregatorRef(NewPropertyAggregator));
+}
+
+float FActiveGameplayEffectsContainer::GetGameplayEffectDuration(FActiveGameplayEffectHandle Handle) const
+{
+	// Could make this a map for quicker lookup
+	for (int32 idx = 0; idx < GameplayEffects.Num(); ++idx)
+	{
+		if (GameplayEffects[idx].Handle == Handle)
+		{
+			return GameplayEffects[idx].GetDuration();
+		}
+	}
+	
+	SKILL_LOG(Warning, TEXT("GetGameplayEffectDuration called with invalid Handle: %s"), *Handle.ToString() );
+	return UGameplayEffect::INFINITE_DURATION;
+}
+
+void FActiveGameplayEffectsContainer::OnDurationAggregatorDirty(FAggregator* Aggregator, UAttributeComponent* Owner, FActiveGameplayEffectHandle Handle)
+{
+	for (FActiveGameplayEffect& Effect : GameplayEffects)
+	{
+		if (Effect.Handle == Handle)
+		{
+			FTimerManager& TimerManager = Owner->GetWorld()->GetTimerManager();
+			float CurrDuration = Effect.GetDuration();
+			float TimeRemaining = CurrDuration - TimerManager.GetTimerElapsed(Effect.DurationHandle);
+			FTimerDelegate Delegate = FTimerDelegate::CreateUObject(Owner, &UAttributeComponent::CheckDurationExpired, Handle);
+			TimerManager.SetTimer(Effect.DurationHandle, Delegate, CurrDuration, false, FMath::Max(TimeRemaining, 0.f));
+
+			return;
+		}
+	}
+}
+
+float FActiveGameplayEffectsContainer::GetGameplayEffectMagnitude(FActiveGameplayEffectHandle Handle, FGameplayAttribute Attribute) const
+{
+	// Could make this a map for quicker lookup
+	for (FActiveGameplayEffect Effect : GameplayEffects)
+	{
+		if (Effect.Handle == Handle)
+		{
+			for (const FModifierSpec &Mod : Effect.Spec.Modifiers)
+			{
+				if (Mod.Info.Attribute == Attribute)
+				{
+					return Mod.Aggregator.Get()->Evaluate().Magnitude;
+				}
+			}
+		}
+	}
+
+	SKILL_LOG(Warning, TEXT("GetGameplayEffectMagnitude called with invalid Handle: %s"), *Handle.ToString());
+	return -1.f;
+}
+
+bool FActiveGameplayEffectsContainer::IsGameplayEffectActive(FActiveGameplayEffectHandle Handle) const
+{
+	// Could make this a map for quicker lookup
+	for (const FActiveGameplayEffect& Effect : GameplayEffects)
+	{
+		if (Effect.Handle == Handle)
+		{
+			// stacking effects may return false
+			if (Effect.Spec.GetStackingType() != EGameplayEffectStackingPolicy::Unlimited &&
+				Effect.Spec.bTopOfStack == false)
+			{
+				return false;
+			}
+			return true;
+		}
+	}
+
+	return false;
+}
+
+float FActiveGameplayEffectsContainer::GetGameplayEffectMagnitudeByTag(FActiveGameplayEffectHandle Handle, const FGameplayTag& InTag) const
+{
+	// Could make this a map for quicker lookup
+	for (FActiveGameplayEffect Effect : GameplayEffects)
+	{
+		if (Effect.Handle == Handle)
+		{
+			for (const FModifierSpec &Mod : Effect.Spec.Modifiers)
+			{
+				if (Mod.Info.OwnedTags.HasTag(InTag, EGameplayTagMatchType::IncludeParentTags, EGameplayTagMatchType::Explicit))
+				{
+					return Mod.Aggregator.Get()->Evaluate().Magnitude;
+				}
+			}
+		}
+	}
+
+	SKILL_LOG(Warning, TEXT("GetGameplayEffectMagnitude called with invalid Handle: %s"), *Handle.ToString());
+	return -1.f;
+}
+
+void FActiveGameplayEffectsContainer::OnPropertyAggregatorDirty(FAggregator* Aggregator, FGameplayAttribute Attribute)
+{
+	check(OngoingPropertyEffects.FindChecked(Attribute).Get() == Aggregator);
+
+	SKILL_LOG_SCOPE(TEXT("FActiveGameplayEffectsContainer::OnPropertyAggregatorDirty"));
+
+	// Immediately calculate the newest value of the property			
+	float NewPropertyValue = Aggregator->Evaluate().Magnitude;
+
+	SKILL_LOG(Log, TEXT("Property %s new value is: %.2f"), *Attribute.GetName(), NewPropertyValue);
+	
+	Owner->SetNumericAttribute(Attribute, NewPropertyValue);
+}
+
+void FActiveGameplayEffectsContainer::StacksNeedToRecalculate()
+{
+	if (!bNeedToRecalculateStacks)
+	{
+		bNeedToRecalculateStacks = true;
+		FTimerManager& TimerManager = Owner->GetWorld()->GetTimerManager();
+		FTimerDelegate Delegate = FTimerDelegate::CreateUObject(Owner, &UAttributeComponent::OnRestackGameplayEffects);
+
+		TimerManager.SetTimerForNextTick(Delegate);
+	}
+}
+
+FActiveGameplayEffect & FActiveGameplayEffectsContainer::CreateNewActiveGameplayEffect(const FGameplayEffectSpec &Spec, int32 InPredictionKey)
+{
+	LastAssignedHandle = LastAssignedHandle.GetNextHandle();
+	FActiveGameplayEffect & NewEffect = *new (GameplayEffects)FActiveGameplayEffect(LastAssignedHandle, Spec, GetWorldTime(), GetGameStateTime(), InPredictionKey);
+
+	// register callbacks with the timer manager
+	if (Owner)
+	{
+		if (Spec.GetDuration() > 0.f)
+		{
+			FTimerManager& TimerManager = Owner->GetWorld()->GetTimerManager();
+			FTimerDelegate Delegate = FTimerDelegate::CreateUObject(Owner, &UAttributeComponent::CheckDurationExpired, LastAssignedHandle);
+			TimerManager.SetTimer(NewEffect.DurationHandle, Delegate, Spec.GetDuration(), false, Spec.GetDuration());
+		}
+
+		// The timer manager moves things from the pending list to the active list after checking the active list on the first tick so we need to execute here
+		if (Spec.GetPeriod() != UGameplayEffect::NO_PERIOD)
+		{
+			FTimerManager& TimerManager = Owner->GetWorld()->GetTimerManager();
+			FTimerDelegate Delegate = FTimerDelegate::CreateUObject(Owner, &UAttributeComponent::ExecutePeriodicEffect, LastAssignedHandle);
+
+			// If this is a periodic stacking effect, make sure that it's in sync with the others.
+			float FirstDelay = -1.f;
+			bool bApplyToNextTick = true;
+			if (Spec.GetStackingType() != EGameplayEffectStackingPolicy::Unlimited)
+			{
+				for (FActiveGameplayEffect& Effect : GameplayEffects)
+				{
+					if (Effect.CanBeStacked(NewEffect))
+					{
+						FirstDelay = TimerManager.GetTimerRemaining(Effect.PeriodHandle);
+						bApplyToNextTick = TimerManager.IsTimerPending(Effect.PeriodHandle);
+						break;
+					}
+				}
+			}
+			
+			if (bApplyToNextTick)
+			{
+				TimerManager.SetTimerForNextTick(Delegate); // not part of a stack or the first item on the stack, execute once right away
+			}
+
+			TimerManager.SetTimer(NewEffect.PeriodHandle, Delegate, Spec.GetPeriod(), true, FirstDelay); // this is going to be off by a frame for stacking because of the pending list
+		}
+
+		NewEffect.Spec.Duration.Get()->OnDirty = FAggregator::FOnDirty::CreateRaw(this, &FActiveGameplayEffectsContainer::OnDurationAggregatorDirty, Owner, LastAssignedHandle);
+	}
+	
+	if (InPredictionKey == 0 || IsNetAuthority())	// Clients predicting a GameplayEffect must not call MarkItemDirty
+	{
+		MarkItemDirty(NewEffect);
+	}
+	else
+	{
+		// Clients predicting should call MarkArrayDirty to force the internal replication map to be rebuilt.
+		MarkArrayDirty();
+
+		Owner->GetPredictionKeyDelegate(InPredictionKey).AddUObject(Owner, &UAttributeComponent::RemoveActiveGameplayEffect_NoReturn, NewEffect.Handle);
+	}
+	return NewEffect;
+}
+
+bool FActiveGameplayEffectsContainer::RemoveActiveGameplayEffect(FActiveGameplayEffectHandle Handle)
+{
+	// Could make this a map for quicker lookup
+	for (int32 Idx = 0; Idx < GameplayEffects.Num(); ++Idx)
+	{
+		if (GameplayEffects[Idx].Handle == Handle)
+		{
+			return InternalRemoveActiveGameplayEffect(Idx);
+		}
+	}
+
+	SKILL_LOG(Warning, TEXT("RemoveActiveGameplayEffect called with invalid Handle: %s"), *Handle.ToString());
+	return false;
+}
+
+bool FActiveGameplayEffectsContainer::InternalRemoveActiveGameplayEffect(int32 Idx)
+{
+	if (ensure(Idx < GameplayEffects.Num()))
+	{
+		FActiveGameplayEffect& Effect = GameplayEffects[Idx];
+		Owner->InvokeGameplayCueRemoved(GameplayEffects[Idx].Spec);
+
+		if (GameplayEffects[Idx].DurationHandle.IsValid())
+		{
+			Owner->GetWorld()->GetTimerManager().ClearTimer(GameplayEffects[Idx].DurationHandle);
+		}
+		if (GameplayEffects[Idx].PeriodHandle.IsValid())
+		{
+			Owner->GetWorld()->GetTimerManager().ClearTimer(GameplayEffects[Idx].PeriodHandle);
+		}
+
+		GameplayEffects.RemoveAtSwap(Idx);
+		MarkArrayDirty();
+
+		// Hack: force netupdate on owner. This isn't really necessary in real gameplay but is nice
+		// during debugging where breakpoints or pausing can messup network update times. Open issue
+		// with network team.
+		Owner->GetOwner()->ForceNetUpdate();
+
+		StacksNeedToRecalculate();
+
+		return true;
+	}
+
+	SKILL_LOG(Warning, TEXT("InternalRemoveActiveGameplayEffect called with invalid index: %d"), Idx);
+	return false;
+}
+
+bool FActiveGameplayEffectsContainer::IsNetAuthority() const
+{
+	check(Owner);
+	return Owner->IsOwnerActorAuthoritative();
+}
+
+void FActiveGameplayEffectsContainer::PreDestroy()
+{
+	// Prior to destruction we need to clear all dependancies
+	for (auto It : OngoingPropertyEffects)
+	{
+		if (It.Value.IsValid())
+		{
+			It.Value.Get()->ClearAllDependancies();
+		}
+	}
+}
+
+int32 FActiveGameplayEffectsContainer::GetGameStateTime() const
+{
+	UWorld *World = Owner->GetWorld();
+	AGameState * GameState = World->GetGameState<AGameState>();
+	if (GameState)
+	{
+		return GameState->ElapsedTime;
+	}
+
+	return static_cast<int32>(World->GetTimeSeconds());
+}
+
+float FActiveGameplayEffectsContainer::GetWorldTime() const
+{
+	UWorld *World = Owner->GetWorld();
+	return World->GetTimeSeconds();
+}
+
+void FActiveGameplayEffectsContainer::CheckDuration(FActiveGameplayEffectHandle Handle)
+{
+	for (int32 Idx = 0; Idx < GameplayEffects.Num(); ++Idx)
+	{
+		FActiveGameplayEffect& Effect = GameplayEffects[Idx];
+		if (Effect.Handle == Handle)
+		{
+			FTimerManager& TimerManager = Owner->GetWorld()->GetTimerManager();
+
+			// The duration may have changed since we registered this callback with the timer manager.
+			// Make sure that this effect should really be destroyed now
+			float Duration = Effect.GetDuration();
+			float CurrentTime = GetWorldTime();
+			if (Duration > 0.f && (Effect.StartWorldTime + Duration) < CurrentTime)
+			{
+				// This gameplay effect has hit its duration. Check if it needs to execute one last time before removing it.
+				if (Effect.PeriodHandle.IsValid())
+				{
+					float PeriodTimeRemaining = TimerManager.GetTimerRemaining(Effect.PeriodHandle);
+					if (FMath::IsNearlyEqual(PeriodTimeRemaining, 0.f))
+					{
+						ExecuteActiveEffectsFrom(Effect.Spec, FModifierQualifier().IgnoreHandle(Effect.Handle));
+					}
+				}
+
+				InternalRemoveActiveGameplayEffect(Idx);
+			}
+			else
+			{
+				// check the time remaining for the current gameplay effect duration timer
+				// if it is less than zero create a new callback with the correct time remaining
+				float TimeRemaining = TimerManager.GetTimerRemaining(Effect.DurationHandle);
+				if (TimeRemaining <= 0.f)
+				{
+					FTimerDelegate Delegate = FTimerDelegate::CreateUObject(Owner, &UAttributeComponent::CheckDurationExpired, Effect.Handle);
+					TimerManager.SetTimer(Effect.DurationHandle, Delegate, (Effect.StartWorldTime + Duration) - CurrentTime, false);
+				}
+			}
+
+			break;
+		}
+	}
+}
+
+void FActiveGameplayEffectsContainer::RecalculateStacking()
+{
+	bNeedToRecalculateStacks = false;
+
+	// one or more gameplay effects has been removed or added so we need to go through all of them to see what the best elements are in each stack
+	TArray<FActiveGameplayEffect*> StackedEffects;
+	TArray<FActiveGameplayEffect*> CustomStackedEffects;
+
+	for (FActiveGameplayEffect& Effect : GameplayEffects)
+	{
+		// ignore effects that don't stack and effects that replace when they stack
+		// effects that replace when they stack only need to be handled when they are added
+		if (Effect.Spec.GetStackingType() != EGameplayEffectStackingPolicy::Unlimited &&
+			Effect.Spec.GetStackingType() != EGameplayEffectStackingPolicy::Replaces)
+		{
+			if (Effect.Spec.Def->Modifiers.Num() > 0)
+			{
+				Effect.Spec.bTopOfStack = false;
+				bool bFoundStack = false;
+
+				// group all of the custom stacking effects and deal with them after
+				if (Effect.Spec.GetStackingType() == EGameplayEffectStackingPolicy::Callback)
+				{
+					CustomStackedEffects.Add(&Effect);
+					continue;
+				}
+
+				for (FActiveGameplayEffect*& StackedEffect : StackedEffects)
+				{
+					if (StackedEffect->Spec.GetStackingType() == Effect.Spec.GetStackingType() && StackedEffect->Spec.Def->Modifiers[0].Attribute == Effect.Spec.Def->Modifiers[0].Attribute)
+					{
+						switch (Effect.Spec.GetStackingType())
+						{
+						case EGameplayEffectStackingPolicy::Highest:
+						{
+							float BestSpecMagnitude = StackedEffect->Spec.GetMagnitude(StackedEffect->Spec.Modifiers[0].Info.Attribute);
+							float CurrSpecMagnitude = Effect.Spec.GetMagnitude(Effect.Spec.Modifiers[0].Info.Attribute);
+							if (BestSpecMagnitude < CurrSpecMagnitude)
+							{
+								StackedEffect = &Effect;
+							}
+							break;
+						}
+						case EGameplayEffectStackingPolicy::Lowest:
+						{
+							float BestSpecMagnitude = StackedEffect->Spec.GetMagnitude(StackedEffect->Spec.Modifiers[0].Info.Attribute);
+							float CurrSpecMagnitude = Effect.Spec.GetMagnitude(Effect.Spec.Modifiers[0].Info.Attribute);
+							if (BestSpecMagnitude > CurrSpecMagnitude)
+							{
+								StackedEffect = &Effect;
+							}
+							break;
+						}
+						default:
+							// we shouldn't ever be here
+							SKILL_LOG(Warning, TEXT("%s uses unhandled stacking rule %s."), *Effect.Spec.Def->GetPathName(), *EGameplayEffectStackingPolicyToString(Effect.Spec.GetStackingType()));
+							break;
+						};
+
+						bFoundStack = true;
+						break;
+					}
+				}
+
+				// if we didn't find a stack matching this effect we need to add it to our list
+				if (!bFoundStack)
+				{
+					StackedEffects.Add(&Effect);
+				}
+			}
+		}
+	}
+
+	// now deal with the custom effects
+	while (CustomStackedEffects.Num() > 0)
+	{
+		int32 StackedIdx = StackedEffects.Add(CustomStackedEffects[0]);
+		CustomStackedEffects.RemoveAtSwap(0);
+
+		UGameplayEffectStackingExtension * StackedExt = StackedEffects[StackedIdx]->Spec.Def->StackingExtension->GetDefaultObject<UGameplayEffectStackingExtension>();
+
+		TArray<FActiveGameplayEffect*> Effects;
+		
+		for (int32 Idx = 0; Idx < CustomStackedEffects.Num(); Idx++)
+		{
+			UGameplayEffectStackingExtension * CurrentExt = CustomStackedEffects[Idx]->Spec.Def->StackingExtension->GetDefaultObject<UGameplayEffectStackingExtension>();
+			if (CurrentExt && StackedExt && CurrentExt->Handle == StackedExt->Handle && 
+				StackedEffects[StackedIdx]->Spec.Def->Modifiers[0].Attribute == CustomStackedEffects[Idx]->Spec.Def->Modifiers[0].Attribute)
+			{
+				Effects.Add(CustomStackedEffects[Idx]);
+				CustomStackedEffects.RemoveAtSwap(Idx);
+				--Idx;
+			}
+		}
+
+		UGameplayEffectStackingExtension * Ext = StackedEffects[StackedIdx]->Spec.Def->StackingExtension->GetDefaultObject<UGameplayEffectStackingExtension>();
+		Ext->CalculateStack(Effects, *this, *StackedEffects[StackedIdx]);
+	}
+
+	// we've found all of the best stacking effects, mark them so that they updated properly
+	for (FActiveGameplayEffect* const StackedEffect : StackedEffects)
+	{
+		StackedEffect->Spec.bTopOfStack = true;
+	}
+}
+
+bool FActiveGameplayEffectsContainer::HasAnyTags(FGameplayTagContainer &Tags)
+{
+	SCOPE_CYCLE_COUNTER(STAT_GameplayEffectsHasAnyTag);
+
+	for (FActiveGameplayEffect &ActiveEffect : GameplayEffects)
+	{
+		// Fixme: check stacking rules!
+		if (ActiveEffect.Spec.Def->OwnedTagsContainer.MatchesAny(Tags, false))
+		{
+			return true;
+		}
+	}
+
+	return false;
+}
+
+bool FActiveGameplayEffectsContainer::HasAllTags(FGameplayTagContainer &Tags)
+{
+	SCOPE_CYCLE_COUNTER(STAT_GameplayEffectsHasAllTags);
+
+	for (auto It = Tags.CreateConstIterator(); It;  ++It)
+	{
+		bool Found = false;
+		for (FActiveGameplayEffect &ActiveEffect : GameplayEffects)
+		{
+			// Fixme: check stacking rules!
+
+			/**
+			 *	The desired behavior here is:
+			 *		"Do we have tag a.b.c and the GameplayEffect has a.b.c.d -> match!"
+			 *		"Do we have tag a.b.c.d. and the GameplayEffect has a.b.c -> no match!"
+			 */
+
+			if (ActiveEffect.Spec.Def->OwnedTagsContainer.HasTag(*It, EGameplayTagMatchType::IncludeParentTags, EGameplayTagMatchType::Explicit))
+			{
+				Found = true;
+				break;
+			}
+		}
+		if (!Found)
+		{
+			return false;
+		}
+	}
+
+	return true;
+}
+
+bool FActiveGameplayEffectsContainer::CanApplyAttributeModifiers(const UGameplayEffect *GameplayEffect, float Level, AActor *Instigator)
+{
+	SCOPE_CYCLE_COUNTER(STAT_GameplayEffectsCanApplyAttributeModifiers);
+
+	FGameplayEffectSpec	Spec(GameplayEffect, Instigator, Level, Owner->GetCurveDataOverride());
+
+	ApplyActiveEffectsTo(Spec, FModifierQualifier().Type(EGameplayMod::IncomingGE));
+	
+	for (const FModifierSpec & Mod : Spec.Modifiers)
+	{
+		// It only makes sense to check additive operators
+		if (Mod.Info.ModifierOp == EGameplayModOp::Additive)
+		{
+			UAttributeSet * Set = Owner->GetAttributeSubobject(Mod.Info.Attribute.GetAttributeSetClass());
+			float CurrentValue = Mod.Info.Attribute.GetNumericValueChecked(Set);
+			float CostValue = Mod.Aggregator.Get()->Evaluate().Magnitude;
+
+			if (CurrentValue + CostValue < 0.f)
+			{
+				return false;
+			}
+		}
+	}
+	return true;
+}
+
+TArray<float> FActiveGameplayEffectsContainer::GetActiveEffectsTimeRemaining(const FActiveGameplayEffectQuery Query) const
+{
+	SCOPE_CYCLE_COUNTER(STAT_GameplayEffectsGetActiveEffectsData);
+
+	float CurrentTime = GetWorldTime();
+
+	TArray<float>	ReturnList;
+
+	for (const FActiveGameplayEffect &Effect : GameplayEffects)
+	{
+		if (Query.TagContainer)
+		{
+			if (!Effect.Spec.Def->OwnedTagsContainer.MatchesAny(*Query.TagContainer, false))
+			{
+				continue;
+			}
+		}
+
+		float Elapsed = CurrentTime - Effect.StartWorldTime;
+		float Duration = Effect.GetDuration();
+
+		ReturnList.Add(Duration - Elapsed);
+	}
+
+	// Note: keep one return location to avoid copy operation.
+	return ReturnList;
+}
+
+// --------------------------------------------------------------------------------------------------------------------------------------------------------
+//
+//	FGameplayEffectLevelSpec
+//
+// --------------------------------------------------------------------------------------------------------------------------------------------------------
+
+float FGameplayEffectLevelSpec::GetLevel() const
+{
+	if (ConstantLevel != INVALID_LEVEL)
+	{
+		return ConstantLevel;
+	}
+
+	if (Source.IsValid() && Attribute.GetUProperty() != NULL)
+	{
+		UClass * SetClass = Attribute.GetAttributeSetClass();
+		check(SetClass);
+
+		TArray<UObject*> ChildObjects;
+		GetObjectsWithOuter(Source.Get(), ChildObjects);
+
+		for (int32 ChildIndex = 0; ChildIndex < ChildObjects.Num(); ++ChildIndex)
+		{
+			UObject *Obj = ChildObjects[ChildIndex];
+			if (Obj->GetClass()->IsChildOf(SetClass))
+			{
+				CachedLevel = Attribute.GetNumericValueChecked(CastChecked<UAttributeSet>(Obj));
+				return CachedLevel;
+			}
+		}
+	}
+
+	// Our source is invalid. 
+	ConstantLevel = CachedLevel;
+	return CachedLevel;
+}
+
+void FGameplayEffectLevelSpec::RegisterLevelDependancy(TWeakPtr<FAggregator> OwningAggregator)
+{
+	if (Source.IsValid() && Attribute.GetUProperty() != NULL)
+	{
+		UAttributeComponent * SourceComponent = Source->FindComponentByClass<UAttributeComponent>();
+		if (SourceComponent)
+		{
+			SourceComponent->AddDependancyToAttribute(Attribute, OwningAggregator);
+		}
+	}
+}
+
+// --------------------------------------------------------------------------------------------------------------------------------------------------------
+//
+//	Misc
+//
+// --------------------------------------------------------------------------------------------------------------------------------------------------------
+
+FString EGameplayModOpToString(int32 Type)
+{
+	static UEnum *e = FindObject<UEnum>(ANY_PACKAGE, TEXT("EGameplayModOp"));
+	return e->GetEnum(Type).ToString();
+}
+
+FString EGameplayModToString(int32 Type)
+{
+	static UEnum *e = FindObject<UEnum>(ANY_PACKAGE, TEXT("EGameplayMod"));
+	return e->GetEnum(Type).ToString();
+}
+
+FString EGameplayModEffectToString(int32 Type)
+{
+	static UEnum *e = FindObject<UEnum>(ANY_PACKAGE, TEXT("EGameplayModEffect"));
+	return e->GetEnum(Type).ToString();
+}
+
+FString EGameplayEffectCopyPolicyToString(int32 Type)
+{
+	static UEnum *e = FindObject<UEnum>(ANY_PACKAGE, TEXT("EGameplayEffectCopyPolicy"));
+	return e->GetEnum(Type).ToString();
+}
+
+FString EGameplayEffectStackingPolicyToString(int32 Type)
+{
+	static UEnum *e = FindObject<UEnum>(ANY_PACKAGE, TEXT("EGameplayEffectStackingPolicy"));
+	return e->GetEnum(Type).ToString();
+}
+
diff --git a/Engine/Source/Runtime/GameplayAbilities/Private/GameplayEffectExtension.cpp b/Engine/Source/Runtime/GameplayAbilities/Private/GameplayEffectExtension.cpp
new file mode 100644
--- /dev/null
+++ b/Engine/Source/Runtime/GameplayAbilities/Private/GameplayEffectExtension.cpp
@@ -0,0 +1,10 @@
+// Copyright 1998-2013 Epic Games, Inc. All Rights Reserved.
+
+#include "SkillSystemModulePrivatePCH.h"
+#include "GameplayEffectExtension.h"
+
+UGameplayEffectExtension::UGameplayEffectExtension(const class FPostConstructInitializeProperties& PCIP)
+: Super(PCIP)
+{
+
+}
diff --git a/Engine/Source/Runtime/GameplayAbilities/Private/GameplayEffectExtension_LifestealTest.cpp b/Engine/Source/Runtime/GameplayAbilities/Private/GameplayEffectExtension_LifestealTest.cpp
new file mode 100644
--- /dev/null
+++ b/Engine/Source/Runtime/GameplayAbilities/Private/GameplayEffectExtension_LifestealTest.cpp
@@ -0,0 +1,61 @@
+// Copyright 1998-2013 Epic Games, Inc. All Rights Reserved.
+
+#include "SkillSystemModulePrivatePCH.h"
+#include "GameplayEffectExtension_LifestealTest.h"
+#include "SkillSystemTestAttributeSet.h"
+#include "GameplayTagsModule.h"
+#include "AttributeComponent.h"
+
+UGameplayEffectExtension_LifestealTest::UGameplayEffectExtension_LifestealTest(const class FPostConstructInitializeProperties& PCIP)
+: Super(PCIP)
+{
+	HealthRestoreGameplayEffect = NULL;
+}
+
+
+void UGameplayEffectExtension_LifestealTest::PreGameplayEffectExecute(const FGameplayModifierEvaluatedData &SelfData, FGameplayEffectModCallbackData &Data) const
+{
+
+}
+
+void UGameplayEffectExtension_LifestealTest::PostGameplayEffectExecute(const FGameplayModifierEvaluatedData &SelfData, const FGameplayEffectModCallbackData &Data) const
+{
+	IGameplayTagsModule& GameplayTagsModule = IGameplayTagsModule::Get();
+
+	float DamageDone = Data.EvaluatedData.Magnitude;
+	float LifestealPCT = SelfData.Magnitude;
+
+	float HealthToRestore = -DamageDone * LifestealPCT;
+	if (HealthToRestore > 0.f)
+	{
+		UAttributeComponent *Source = Data.EffectSpec.InstigatorContext.GetOriginalInstigatorAttributeComponent();
+
+		UGameplayEffect * LocalHealthRestore = HealthRestoreGameplayEffect;
+		if (!LocalHealthRestore)
+		{
+			UProperty *HealthProperty = FindFieldChecked<UProperty>(USkillSystemTestAttributeSet::StaticClass(), GET_MEMBER_NAME_CHECKED(USkillSystemTestAttributeSet, Health));
+
+			// Since this is a test class and we don't want to tie it any actual content assets, just construct a GameplayEffect here.
+			LocalHealthRestore = Cast<UGameplayEffect>(StaticConstructObject(UGameplayEffect::StaticClass(), GetTransientPackage(), FName(TEXT("LifestealHealthRestore"))));
+			LocalHealthRestore->Modifiers.SetNum(1);
+			LocalHealthRestore->Modifiers[0].Magnitude.SetValue(HealthToRestore);
+			LocalHealthRestore->Modifiers[0].ModifierType = EGameplayMod::Attribute;
+			LocalHealthRestore->Modifiers[0].ModifierOp = EGameplayModOp::Additive;
+			LocalHealthRestore->Modifiers[0].Attribute.SetUProperty(HealthProperty);
+			LocalHealthRestore->Modifiers[0].OwnedTags.AddTag(GameplayTagsModule.GetGameplayTagsManager().RequestGameplayTag(FName(TEXT("Lifesteal"))));
+			LocalHealthRestore->Duration.Value = UGameplayEffect::INSTANT_APPLICATION;
+			LocalHealthRestore->Period.Value = UGameplayEffect::NO_PERIOD;
+		}
+
+		if (SelfData.Handle.IsValid())
+		{
+			// We are coming from an active gameplay effect
+			check(SelfData.Handle.IsValid());
+		}
+
+		// Apply a GameplayEffect to restore health
+		// We make the GameplayEffect's level = the health restored. This is one approach. We could also
+		// try a basic restore 1 health item but apply a 2nd GE to modify that - but that seems like too many levels of indirection
+		Source->ApplyGameplayEffectToSelf(LocalHealthRestore, HealthToRestore, Source->GetOwner());
+	}
+}
\ No newline at end of file
diff --git a/Engine/Source/Runtime/GameplayAbilities/Private/GameplayEffectExtension_ShieldTest.cpp b/Engine/Source/Runtime/GameplayAbilities/Private/GameplayEffectExtension_ShieldTest.cpp
new file mode 100644
--- /dev/null
+++ b/Engine/Source/Runtime/GameplayAbilities/Private/GameplayEffectExtension_ShieldTest.cpp
@@ -0,0 +1,74 @@
+// Copyright 1998-2013 Epic Games, Inc. All Rights Reserved.
+
+#include "SkillSystemModulePrivatePCH.h"
+#include "GameplayEffectExtension_ShieldTest.h"
+#include "GameplayTagsModule.h"
+#include "AttributeComponent.h"
+#include "SkillSystemTestAttributeSet.h"
+
+UGameplayEffectExtension_ShieldTest::UGameplayEffectExtension_ShieldTest(const class FPostConstructInitializeProperties& PCIP)
+: Super(PCIP)
+{
+
+}
+
+
+void UGameplayEffectExtension_ShieldTest::PreGameplayEffectExecute(const FGameplayModifierEvaluatedData &SelfData, FGameplayEffectModCallbackData &Data) const
+{
+	IGameplayTagsModule& GameplayTagsModule = IGameplayTagsModule::Get();
+	UAttributeComponent *Source = Data.EffectSpec.InstigatorContext.GetOriginalInstigatorAttributeComponent();
+
+	// FIXME: some annoyances here: Damage about to be applied = Data.EvaluatedData.Magnitude = negative. Do some sign flipping here that would make more sense if we were dealing with
+	// 'damage' (positive) instead of 'health' (negative)
+
+	// How much damage we will absorb
+	float ShieldedDamage = FMath::Min(-Data.EvaluatedData.Magnitude, SelfData.Magnitude);
+
+	// Decrease the magnitude of damage done
+	Data.EvaluatedData.Magnitude += ShieldedDamage;
+
+	if (ShieldedDamage >= SelfData.Magnitude)
+	{
+		check(SelfData.Handle.IsValid());
+
+		// The shield is now depleted, so remove the gameplay effect altogether
+		bool RemovedSomething = Data.Target.RemoveActiveGameplayEffect(SelfData.Handle);
+		check(RemovedSomething);
+
+		// Note we could get in trouble with the above checks, depending on how the GameplayEffect that called us was setup. 
+		// For example if we were applied to a periodic damage effect, and were copied as a snapshot, our shield could expire
+		// but our aggregator on the DOT gameplayeffect would still exist with a now invalid handle. This is ok, and possibly powerful
+		// but just means that the way an effect like this is written would have to take it into account. This one doesn't.
+		//
+		// It even opens up interesting possibilities. You could have a shield specifically designed for dots: something that would absorb
+		// the first X damage from each DOT applied to you. Practicaly? Maybe not, but demonstrates power of this setup.
+	}
+	else if (ShieldedDamage > 0.f)
+	{
+		// Apply a GameplayEffect to remove some of our shield
+		UGameplayEffect * LocalShieldRemoval = ShieldRemoveGameplayEffect;
+
+		if (!LocalShieldRemoval)
+		{
+			UProperty *HealthProperty = FindFieldChecked<UProperty>(USkillSystemTestAttributeSet::StaticClass(), GET_MEMBER_NAME_CHECKED(USkillSystemTestAttributeSet, Health));
+
+			// Since this is a test class and we don't want to tie it any actual content assets, just construct a GameplayEffect here.
+			LocalShieldRemoval = Cast<UGameplayEffect>(StaticConstructObject(UGameplayEffect::StaticClass(), GetTransientPackage(), FName(TEXT("ShieldAbsorbRemoval"))));
+			LocalShieldRemoval->Modifiers.SetNum(1);
+			LocalShieldRemoval->Modifiers[0].Magnitude.SetValue(-ShieldedDamage);
+			LocalShieldRemoval->Modifiers[0].ModifierType = EGameplayMod::ActiveGE;
+			LocalShieldRemoval->Modifiers[0].ModifierOp = EGameplayModOp::Additive;
+			LocalShieldRemoval->Modifiers[0].Attribute.SetUProperty(HealthProperty);
+			LocalShieldRemoval->Modifiers[0].OwnedTags.AddTag(GameplayTagsModule.GetGameplayTagsManager().RequestGameplayTag(FName(TEXT("ShieldAbsorb"))));
+			LocalShieldRemoval->Duration.Value = UGameplayEffect::INSTANT_APPLICATION;
+			LocalShieldRemoval->Period.Value = UGameplayEffect::NO_PERIOD;
+		}
+
+		Data.Target.ApplyGameplayEffectToSelf(LocalShieldRemoval, ShieldedDamage, Source->GetOwner(), FModifierQualifier().ExclusiveTarget(SelfData.Handle));
+	}
+}
+
+void UGameplayEffectExtension_ShieldTest::PostGameplayEffectExecute(const FGameplayModifierEvaluatedData &SelfData, const FGameplayEffectModCallbackData &Data) const
+{
+
+}
\ No newline at end of file
diff --git a/Engine/Source/Runtime/GameplayAbilities/Private/SkillSystemGlobals.cpp b/Engine/Source/Runtime/GameplayAbilities/Private/SkillSystemGlobals.cpp
new file mode 100644
--- /dev/null
+++ b/Engine/Source/Runtime/GameplayAbilities/Private/SkillSystemGlobals.cpp
@@ -0,0 +1,93 @@
+// Copyright 1998-2014 Epic Games, Inc. All Rights Reserved.
+
+#include "SkillSystemModulePrivatePCH.h"
+#include "AttributeComponent.h"
+#include "SkillSystemTypes.h"
+
+USkillSystemGlobals::USkillSystemGlobals(const class FPostConstructInitializeProperties& PCIP)
+: Super(PCIP)
+{
+
+}
+
+UCurveTable * USkillSystemGlobals::GetGlobalCurveTable()
+{
+	if (!GlobalCurveTable && !GlobalCurveTableName.IsEmpty())
+	{
+		GlobalCurveTable = LoadObject<UCurveTable>(NULL, *GlobalCurveTableName, NULL, LOAD_None, NULL);
+#if WITH_EDITOR
+		// Hook into notifications for object re-imports so that the gameplay tag tree can be reconstructed if the table changes
+		if (GIsEditor && GlobalCurveTable)
+		{
+			GEditor->OnObjectReimported().AddUObject(this, &USkillSystemGlobals::OnCurveTableReimported);
+		}
+#endif
+	}
+	
+	return GlobalCurveTable;
+}
+
+UDataTable * USkillSystemGlobals::GetGlobalAttributeDataTable()
+{
+	if (!GlobalAttributeDataTable && !GlobalAttributeDataTableName.IsEmpty())
+	{
+		GlobalAttributeDataTable = LoadObject<UDataTable>(NULL, *GlobalAttributeDataTableName, NULL, LOAD_None, NULL);
+#if WITH_EDITOR
+		// Hook into notifications for object re-imports so that the gameplay tag tree can be reconstructed if the table changes
+		if (GIsEditor && GlobalAttributeDataTable)
+		{
+			GEditor->OnObjectReimported().AddUObject(this, &USkillSystemGlobals::OnDataTableReimported);
+		}
+#endif
+	}
+
+	return GlobalAttributeDataTable;
+}
+
+#if WITH_EDITOR
+
+void USkillSystemGlobals::OnCurveTableReimported(UObject* InObject)
+{
+	if (GIsEditor && !IsRunningCommandlet() && InObject && InObject == GlobalCurveTable)
+	{
+	}
+}
+
+void USkillSystemGlobals::OnDataTableReimported(UObject* InObject)
+{
+	if (GIsEditor && !IsRunningCommandlet() && InObject && InObject == GlobalAttributeDataTable)
+	{
+	}
+}
+
+#endif
+
+
+UAttributeComponent * USkillSystemGlobals::GetAttributeComponentFromActor(AActor *Actor) const
+{
+	if (Actor)
+	{
+		return Actor->FindComponentByClass<UAttributeComponent>();
+	}
+
+	// Caller should check this
+	ensure(false);
+	return NULL;
+}
+
+FGameplayAbilityActorInfo * USkillSystemGlobals::AllocAbilityActorInfo(AActor *Actor) const
+{
+	FGameplayAbilityActorInfo * Info = new FGameplayAbilityActorInfo();
+	Info->InitFromActor(Actor);
+	return Info;
+}
+
+/** This is just some syntax sugar to avoid calling gode to have to IGameplayAbilitiesModule::Get() */
+USkillSystemGlobals& USkillSystemGlobals::Get()
+{
+	static USkillSystemGlobals * GlobalPtr = IGameplayAbilitiesModule::Get().GetSkillSystemGlobals();
+	check(GlobalPtr == IGameplayAbilitiesModule::Get().GetSkillSystemGlobals());
+	check(GlobalPtr);
+
+	return *GlobalPtr;
+}
diff --git a/Engine/Source/Runtime/GameplayAbilities/Private/SkillSystemLog.cpp b/Engine/Source/Runtime/GameplayAbilities/Private/SkillSystemLog.cpp
new file mode 100644
--- /dev/null
+++ b/Engine/Source/Runtime/GameplayAbilities/Private/SkillSystemLog.cpp
@@ -0,0 +1,83 @@
+// Copyright 1998-2014 Epic Games, Inc. All Rights Reserved.
+#include "SkillSystemModulePrivatePCH.h"
+
+DEFINE_LOG_CATEGORY(LogSkillSystem);
+
+SkillSystemLogScope::~SkillSystemLogScope()
+{
+	SkillSystemLog::PopScope();
+}
+
+void SkillSystemLogScope::Init()
+{
+	PrintedInThisScope = false;
+	SkillSystemLog::PushScope(this);
+}
+
+// ----------------------------------------------
+
+SkillSystemLog * GetInstance()
+{
+	static SkillSystemLog Instance;
+	return &Instance;
+}
+
+FString SkillSystemLog::Log(ELogVerbosity::Type Verbosity, FString Log)
+{
+	SkillSystemLog * Instance = GetInstance();
+
+#if !NO_LOGGING
+	if (!LogSkillSystem.IsSuppressed(Verbosity))
+	{
+		for (int32 idx=0; idx < Instance->ScopeStack.Num(); ++idx)
+		{
+			SkillSystemLogScope *Scope = Instance->ScopeStack[idx];
+			if (!Scope->PrintedInThisScope && !Scope->ScopeName.IsEmpty())
+			{
+				if (Instance->NeedNewLine)
+				{
+					UE_LOG(LogSkillSystem, Log, TEXT(""));
+					Instance->NeedNewLine = false;
+				}
+
+				Scope->PrintedInThisScope = true;
+				int32 ident = (2 * idx);
+				FString IndentStrX = FString::Printf(TEXT("%*s"), ident, TEXT(""));
+				UE_LOG(LogSkillSystem, Log, TEXT("%s<%s>"), *IndentStrX, *Scope->ScopeName);
+			}
+		}
+	}
+#endif
+
+	FString IndentStr = FString::Printf(TEXT("%*s"), Instance->Indent, TEXT(""));
+	
+	return IndentStr + Log;
+}
+
+void SkillSystemLog::PushScope(SkillSystemLogScope * Scope)
+{
+	SkillSystemLog * Instance = GetInstance();
+
+	Instance->Indent += 2;
+	Instance->ScopeStack.Push(Scope);
+}
+
+void SkillSystemLog::PopScope()
+{
+	SkillSystemLog * Instance = GetInstance();
+
+	Instance->NeedNewLine = true;
+	Instance->Indent -= 2;
+	check(Instance->Indent >= 0);
+
+	if (Instance->ScopeStack.Top()->PrintedInThisScope)
+	{
+		if (!Instance->ScopeStack.Top()->ScopeName.IsEmpty())
+		{
+			FString IndentStrX = FString::Printf(TEXT("%*s"), Instance->Indent, TEXT(""));
+			UE_LOG(LogSkillSystem, Log, TEXT("%s</%s>"), *IndentStrX, *Instance->ScopeStack.Top()->ScopeName);
+		}
+	}
+	
+	Instance->ScopeStack.Pop();
+}
\ No newline at end of file
diff --git a/Engine/Source/Runtime/GameplayAbilities/Private/SkillSystemModulePrivatePCH.h b/Engine/Source/Runtime/GameplayAbilities/Private/SkillSystemModulePrivatePCH.h
new file mode 100644
--- /dev/null
+++ b/Engine/Source/Runtime/GameplayAbilities/Private/SkillSystemModulePrivatePCH.h
@@ -0,0 +1,21 @@
+// Copyright 1998-2014 Epic Games, Inc. All Rights Reserved.
+
+#pragma once
+
+#if WITH_EDITOR
+#include "UnrealEd.h"
+#endif
+
+
+// You should place include statements to your module's private header files here.  You only need to
+// add includes for headers that are used in most of your module's source files though.
+
+#include "Core.h"
+#include "Engine.h"
+#include "Engine/DataTable.h"
+
+#include "SkillSystemTypes.h"
+
+#include "SkillSystem.h"
+#include "GameplayAbilitiesModule.h"
+#include "SkillSystemLog.h"
\ No newline at end of file
diff --git a/Engine/Source/Runtime/GameplayAbilities/Private/SkillSystemTestAttributeSet.cpp b/Engine/Source/Runtime/GameplayAbilities/Private/SkillSystemTestAttributeSet.cpp
new file mode 100644
--- /dev/null
+++ b/Engine/Source/Runtime/GameplayAbilities/Private/SkillSystemTestAttributeSet.cpp
@@ -0,0 +1,131 @@
+// Copyright 1998-2014 Epic Games, Inc. All Rights Reserved.
+
+#include "SkillSystemModulePrivatePCH.h"
+#include "SkillSystemTestAttributeSet.h"
+#include "GameplayTagsModule.h"
+#include "GameplayEffectExtension.h"
+#include "AttributeComponent.h"
+
+
+USkillSystemTestAttributeSet::USkillSystemTestAttributeSet(const class FPostConstructInitializeProperties& PCIP)
+: Super(PCIP)
+{
+	Health = MaxHealth = 100.f;
+	Mana = MaxMana = 100.f;
+	
+	Damage = 0.f;
+	CritChance = 0.f;
+	SpellDamage = 0.f;
+	PhysicalDamage = 0.f;
+	Strength = 0.f;
+	StackingAttribute1 = 0.f;
+	StackingAttribute2 = 0.f;
+	NoStackAttribute = 0.f;
+}
+
+void USkillSystemTestAttributeSet::PreAttributeModify(struct FGameplayEffectModCallbackData &Data)
+{
+	static UProperty *HealthProperty = FindFieldChecked<UProperty>(USkillSystemTestAttributeSet::StaticClass(), GET_MEMBER_NAME_CHECKED(USkillSystemTestAttributeSet, Health));
+	static UProperty *DamageProperty = FindFieldChecked<UProperty>(USkillSystemTestAttributeSet::StaticClass(), GET_MEMBER_NAME_CHECKED(USkillSystemTestAttributeSet, Damage));
+
+	// In this function, our GameplayEffect mod has been evaluated. We have a magnitude and a Tags collection that we can still modify before it is applied.
+	// We also still have the Aggregation data that calculated Data.EvaluatedData. If we really needed to, we could look at this, remove or change things at the aggregator level, and reevaluate ourselves.
+	// But that would be considered very advanced/rare.
+
+	UProperty *ModifiedProperty = Data.ModifierSpec.Info.Attribute.GetUProperty();
+
+	// Is Damage about to be applied?
+	if (DamageProperty == ModifiedProperty)
+	{
+		// Can the target dodge this completely?
+		if (DodgeChance > 0.f)
+		{
+			if (FMath::FRand() <= DodgeChance)
+			{
+				// Dodge!
+				Data.EvaluatedData.Magnitude = 0.f;
+				Data.EvaluatedData.Tags.AddTag(IGameplayTagsModule::RequestGameplayTag(FName(TEXT("Dodged"))));
+
+				// How dodge is handled will be game dependant. There are a few options I think of:
+				// -We still apply 0 damage, but tag it as Dodged. The GameplayCue system could pick up on this and play a visual effect. The combat log could pick up in and print it out too.
+				// -We throw out this GameplayEffect right here, and apply another GameplayEffect for 'Dodge' it wouldn't modify an attribute but could trigger gameplay cues, it could serve as a 'cooldown' for dodge
+				//		if the game wanted rules like 'you can't dodge more than once every .5 seconds', etc.
+			}
+		}		
+		
+		if (Data.EvaluatedData.Magnitude > 0.f)
+		{
+			// Check the source - does he have Crit?
+			USkillSystemTestAttributeSet * SourceAttributes = Data.EffectSpec.InstigatorContext.GetOriginalInstigatorAttributeComponent()->GetSet<USkillSystemTestAttributeSet>();
+			if (SourceAttributes && SourceAttributes->CritChance > 0.f)
+			{
+				if (FMath::FRand() <= SourceAttributes->CritChance)
+				{
+					// Crit!
+					Data.EvaluatedData.Magnitude *= SourceAttributes->CritMultiplier;
+					Data.EvaluatedData.Tags.AddTag(IGameplayTagsModule::RequestGameplayTag(FName(TEXT("Damage.Crit"))));
+				}
+			}
+
+			// Now apply armor reduction
+			if (Data.EvaluatedData.Tags.HasTag(IGameplayTagsModule::RequestGameplayTag(FName(TEXT("Damage.Physical"))), EGameplayTagMatchType::IncludeParentTags, EGameplayTagMatchType::Explicit))
+			{
+				// This is a trivial/naive implementation of armor. It assumes the rmorDamageReduction is an actual % to reduce physics damage by.
+				// Real games would probably use armor rating as an attribute and calculate a % here based on the damage source's level, etc.
+				Data.EvaluatedData.Magnitude *= (1.f - ArmorDamageReduction);
+				Data.EvaluatedData.Tags.AddTag(IGameplayTagsModule::RequestGameplayTag(FName(TEXT("Damage.Mitigated.Armor"))));
+			}
+		}
+
+		// At this point, the Magnitude of the applied damage may have been modified by us. We still do the translation to Health in USkillSystemTestAttributeSet::PostAttributeModify.
+	}
+}
+
+void USkillSystemTestAttributeSet::PostAttributeModify(const struct FGameplayEffectModCallbackData &Data)
+{
+	static UProperty *HealthProperty = FindFieldChecked<UProperty>(USkillSystemTestAttributeSet::StaticClass(), GET_MEMBER_NAME_CHECKED(USkillSystemTestAttributeSet, Health));
+	static UProperty *DamageProperty = FindFieldChecked<UProperty>(USkillSystemTestAttributeSet::StaticClass(), GET_MEMBER_NAME_CHECKED(USkillSystemTestAttributeSet, Damage));
+
+	UProperty *ModifiedProperty = Data.ModifierSpec.Info.Attribute.GetUProperty();
+
+	// What property was modified?
+	if (DamageProperty == ModifiedProperty)
+	{
+		// Anytime Damage is applied with 'Damage.Fire' tag, there is a chance to apply a burning DOT
+		if (Data.EvaluatedData.Tags.HasTag( IGameplayTagsModule::RequestGameplayTag(FName(TEXT("FireDamage"))), EGameplayTagMatchType::IncludeParentTags, EGameplayTagMatchType::Explicit))
+		{
+			// Logic to rand() a burning DOT, if successful, apply DOT GameplayEffect to the target
+		}
+
+		// Treat damage as minus health
+		Health -= Damage;
+		Damage = 0.f;
+
+		// Check for Death?
+		//  -This could be defined here or at the actor level.
+		//  -Doing it here makes a lot of sense to me, but we have legacy code in ::TakeDamage function, so some games may just want to punt to that pipeline from here.
+	}
+}
+
+
+void USkillSystemTestAttributeSet::GetLifetimeReplicatedProps(TArray< FLifetimeProperty > & OutLifetimeProps) const
+{
+	/*
+	DOREPLIFETIME( USkillSystemTestAttributeSet, MaxHealth);
+	DOREPLIFETIME( USkillSystemTestAttributeSet, Health);
+	DOREPLIFETIME( USkillSystemTestAttributeSet, Mana);
+	DOREPLIFETIME( USkillSystemTestAttributeSet, MaxMana);
+
+	DOREPLIFETIME( USkillSystemTestAttributeSet, SpellDamage);
+	DOREPLIFETIME( USkillSystemTestAttributeSet, PhysicalDamage);
+
+	DOREPLIFETIME( USkillSystemTestAttributeSet, CritChance);
+	DOREPLIFETIME( USkillSystemTestAttributeSet, CritMultiplier);
+	DOREPLIFETIME( USkillSystemTestAttributeSet, ArmorDamageReduction);
+
+	DOREPLIFETIME( USkillSystemTestAttributeSet, DodgeChance);
+	DOREPLIFETIME( USkillSystemTestAttributeSet, LifeSteal);
+
+	DOREPLIFETIME( USkillSystemTestAttributeSet, Strength);
+	*/
+}
\ No newline at end of file
diff --git a/Engine/Source/Runtime/GameplayAbilities/Private/SkillSystemTestPawn.cpp b/Engine/Source/Runtime/GameplayAbilities/Private/SkillSystemTestPawn.cpp
new file mode 100644
--- /dev/null
+++ b/Engine/Source/Runtime/GameplayAbilities/Private/SkillSystemTestPawn.cpp
@@ -0,0 +1,57 @@
+// Copyright 1998-2014 Epic Games, Inc. All Rights Reserved.
+
+#include "SkillSystemModulePrivatePCH.h"
+#include "SkillSystemTestPawn.h"
+#include "SkillSystemTestAttributeSet.h"
+#include "AttributeComponent.h"
+
+FName  ASkillSystemTestPawn::AttributeComponentName(TEXT("AttributeComponent0"));
+
+ASkillSystemTestPawn::ASkillSystemTestPawn(const class FPostConstructInitializeProperties& PCIP)
+: Super(PCIP)
+{
+	AttributeComponent = PCIP.CreateDefaultSubobject<UAttributeComponent>(this, ASkillSystemTestPawn::AttributeComponentName);
+	AttributeComponent->SetIsReplicated(true);
+
+	//DefaultAbilitySet = NULL;
+}
+
+void ASkillSystemTestPawn::PostInitializeComponents()
+{	
+	static UProperty *DamageProperty = FindFieldChecked<UProperty>(USkillSystemTestAttributeSet::StaticClass(), GET_MEMBER_NAME_CHECKED(USkillSystemTestAttributeSet, Damage));
+
+	Super::PostInitializeComponents();
+	GameplayCueHandler.Owner = this;
+	AttributeComponent->InitStats(USkillSystemTestAttributeSet::StaticClass(), NULL);
+
+	/*
+	if (DefaultAbilitySet != NULL)
+	{
+		AttributeComponent->InitializeAbilities(DefaultAbilitySet);
+	}
+	*/
+}
+
+void ASkillSystemTestPawn::GameplayCueActivated(const FGameplayTagContainer & GameplayCueTags, float NormalizedMagnitude, const FGameplayEffectInstigatorContext InstigatorContext)
+{
+	SKILL_LOG(Log, TEXT("GameplayCueExecuted: %s. %.2f"), *GetName(), NormalizedMagnitude);
+	GameplayCueHandler.GameplayCueActivated(GameplayCueTags, NormalizedMagnitude, InstigatorContext);
+}
+
+void ASkillSystemTestPawn::GameplayCueExecuted(const FGameplayTagContainer & GameplayCueTags, float NormalizedMagnitude, const FGameplayEffectInstigatorContext InstigatorContext)
+{
+	SKILL_LOG(Log, TEXT("GameplayCueExecuted: %s. %.2f"), *GetName(), NormalizedMagnitude);
+	GameplayCueHandler.GameplayCueExecuted(GameplayCueTags, NormalizedMagnitude, InstigatorContext);
+}
+
+void ASkillSystemTestPawn::GameplayCueAdded(const FGameplayTagContainer & GameplayCueTags, float NormalizedMagnitude, const FGameplayEffectInstigatorContext InstigatorContext)
+{
+	SKILL_LOG(Log, TEXT("GameplayCueAdded: %s. %.2f"), *GetName(), NormalizedMagnitude);
+	GameplayCueHandler.GameplayCueAdded(GameplayCueTags, NormalizedMagnitude, InstigatorContext);
+}
+
+void ASkillSystemTestPawn::GameplayCueRemoved(const FGameplayTagContainer & GameplayCueTags, float NormalizedMagnitude, const FGameplayEffectInstigatorContext InstigatorContext)
+{
+	SKILL_LOG(Log, TEXT("GameplayCueRemoved: %s. %.2f"), *GetName(), NormalizedMagnitude);
+	GameplayCueHandler.GameplayCueRemoved(GameplayCueTags, NormalizedMagnitude, InstigatorContext);
+}
\ No newline at end of file
diff --git a/Engine/Source/Runtime/GameplayAbilities/Public/SkillSystemLog.h b/Engine/Source/Runtime/GameplayAbilities/Public/SkillSystemLog.h
new file mode 100644
--- /dev/null
+++ b/Engine/Source/Runtime/GameplayAbilities/Public/SkillSystemLog.h
@@ -0,0 +1,66 @@
+// Copyright 1998-2014 Epic Games, Inc. All Rights Reserved.
+
+#pragma once
+
+// Intended categories:
+//	Log - This happened. What gameplay programmers may care about to debug
+//	Verbose - This is why this happened. What you may turn on to debug the skill system code.
+//  VeryVerbose - This is what didn't happen, and why. Extreme printf debugging
+//
+
+
+GAMEPLAYABILITIES_API DECLARE_LOG_CATEGORY_EXTERN(LogSkillSystem, Warning, All);
+
+#define SKILL_LOG(Verbosity, Format, ...) \
+{ \
+	FString Str = SkillSystemLog::Log(ELogVerbosity::Verbosity, FString::Printf(Format, ##__VA_ARGS__)); \
+	UE_LOG(LogSkillSystem, Verbosity, TEXT("%s"), *Str); \
+}
+
+#define SKILL_LOG_TOKENPASTE_INNER(x,y) x##y
+#define SKILL_LOG_TOKENPASTE(x,y) SKILL_LOG_TOKENPASTE_INNER(x,y)
+#define SKILL_LOG_SCOPE( Format, ... ) SkillSystemLogScope SKILL_LOG_TOKENPASTE(LogScope,__LINE__)( FString::Printf(Format, ##__VA_ARGS__));
+
+struct SkillSystemLogScope
+{
+	SkillSystemLogScope()
+	{
+		Init();
+	}
+
+	SkillSystemLogScope(FString InScopeName) :
+		ScopeName(InScopeName)
+	{
+		Init();
+	}
+
+	~SkillSystemLogScope();
+	
+	FString ScopeName;
+	bool PrintedInThisScope;
+
+private:
+
+	void Init();
+};
+
+class SkillSystemLog
+{
+public:
+
+	SkillSystemLog()
+	{
+		Indent = 0;
+		NeedNewLine = false;
+	}
+
+	static FString Log(ELogVerbosity::Type, FString Str);
+
+	static void PushScope(SkillSystemLogScope* Scope);
+
+	static void PopScope();	
+
+	int32	Indent;
+	TArray<SkillSystemLogScope*>	ScopeStack;
+	bool NeedNewLine;
+};
\ No newline at end of file
