Mass rename SkillSystem -> AbilitySystem

diff --git a/Engine/Source/Editor/GameplayAbilitiesEditor/Private/AbilitySystemEditorPrivatePCH.h b/Engine/Source/Editor/GameplayAbilitiesEditor/Private/AbilitySystemEditorPrivatePCH.h
new file mode 100644
--- /dev/null
+++ b/Engine/Source/Editor/GameplayAbilitiesEditor/Private/AbilitySystemEditorPrivatePCH.h
@@ -0,0 +1,15 @@
+// Copyright 1998-2013 Epic Games, Inc. All Rights Reserved.
+
+#include "GameplayAbilitiesEditorModule.h"
+
+#if WITH_EDITOR
+#include "UnrealEd.h"
+#endif
+
+
+// You should place include statements to your module's private header files here.  You only need to
+// add includes for headers that are used in most of your module's source files though.
+
+#include "Core.h"
+#include "Engine.h"
+#include "BlueprintGraphDefinitions.h"
\ No newline at end of file
diff --git a/Engine/Source/Editor/GameplayAbilitiesEditor/Private/AttributeDetails.cpp b/Engine/Source/Editor/GameplayAbilitiesEditor/Private/AttributeDetails.cpp
--- a/Engine/Source/Editor/GameplayAbilitiesEditor/Private/AttributeDetails.cpp
+++ b/Engine/Source/Editor/GameplayAbilitiesEditor/Private/AttributeDetails.cpp
@@ -1,13 +1,13 @@
 // Copyright 1998-2013 Epic Games, Inc. All Rights Reserved.
 
-#include "SkillSystemEditorModulePrivatePCH.h"
+#include "AbilitySystemEditorPrivatePCH.h"
 #include "AttributeDetails.h"
 #include "Editor/PropertyEditor/Public/PropertyEditing.h"
 #include "AttributeDetails.h"
 #include "AttributeSet.h"
 #include "BlueprintUtilities.h"
 #include "KismetEditorUtilities.h"
-#include "SkillSystemGlobals.h"
+#include "AbilitySystemGlobals.h"
 #include "GameplayAbilitiesModule.h"
 
 #define LOCTEXT_NAMESPACE "AttributeDetailsCustomization"
@@ -362,7 +362,7 @@ UCurveTable * FScalableFloatDetails::GetCurveTable()
 
 	if (CurveTable == NULL)
 	{
-		CurveTable = IGameplayAbilitiesModule::Get().GetSkillSystemGlobals()->GetGlobalCurveTable();
+		CurveTable = IGameplayAbilitiesModule::Get().GetAbilitySystemGlobals()->GetGlobalCurveTable();
 	}
 
 	return CurveTable;
diff --git a/Engine/Source/Editor/GameplayAbilitiesEditor/Private/GameplayAbilitiesEditorModule.cpp b/Engine/Source/Editor/GameplayAbilitiesEditor/Private/GameplayAbilitiesEditorModule.cpp
--- a/Engine/Source/Editor/GameplayAbilitiesEditor/Private/GameplayAbilitiesEditorModule.cpp
+++ b/Engine/Source/Editor/GameplayAbilitiesEditor/Private/GameplayAbilitiesEditorModule.cpp
@@ -1,6 +1,6 @@
 // Copyright 1998-2013 Epic Games, Inc. All Rights Reserved.
 
-#include "SkillSystemEditorModulePrivatePCH.h"
+#include "AbilitySystemEditorPrivatePCH.h"
 #include "Editor/PropertyEditor/Public/PropertyEditorModule.h"
 #include "AttributeDetails.h"
 #include "AttributeSet.h"
diff --git a/Engine/Source/Editor/GameplayAbilitiesEditor/Private/K2Node_PlayMontageAndWait.cpp b/Engine/Source/Editor/GameplayAbilitiesEditor/Private/K2Node_PlayMontageAndWait.cpp
--- a/Engine/Source/Editor/GameplayAbilitiesEditor/Private/K2Node_PlayMontageAndWait.cpp
+++ b/Engine/Source/Editor/GameplayAbilitiesEditor/Private/K2Node_PlayMontageAndWait.cpp
@@ -1,9 +1,9 @@
 // Copyright 1998-2014 Epic Games, Inc. All Rights Reserved.
 
 
-#include "SkillSystemEditorModulePrivatePCH.h"
+#include "AbilitySystemEditorPrivatePCH.h"
 #include "K2Node_PlayMontageAndWait.h"
-#include "SkillSystemBlueprintLibrary.h"
+#include "AbilitySystemBlueprintLibrary.h"
 #include "Abilities/Tasks/AbilityTask_PlayMontageAndWait.h"
 #include "CompilerResultsLog.h"
 
@@ -13,8 +13,8 @@ UK2Node_PlayMontageAndWait::UK2Node_PlayMontageAndWait(const class FPostConstruc
 	: Super(PCIP)
 {
 	
-	ProxyFactoryFunctionName = GET_FUNCTION_NAME_CHECKED(USkillSystemBlueprintLibrary, CreatePlayMontageAndWaitProxy);
-	ProxyFactoryClass = USkillSystemBlueprintLibrary::StaticClass();
+	ProxyFactoryFunctionName = GET_FUNCTION_NAME_CHECKED(UAbilitySystemBlueprintLibrary, CreatePlayMontageAndWaitProxy);
+	ProxyFactoryClass = UAbilitySystemBlueprintLibrary::StaticClass();
 	ProxyClass = UAbilityTask_PlayMontageAndWait::StaticClass();
 }
 
diff --git a/Engine/Source/Editor/GameplayAbilitiesEditor/Private/K2Node_WaitMovementModeChange.cpp b/Engine/Source/Editor/GameplayAbilitiesEditor/Private/K2Node_WaitMovementModeChange.cpp
--- a/Engine/Source/Editor/GameplayAbilitiesEditor/Private/K2Node_WaitMovementModeChange.cpp
+++ b/Engine/Source/Editor/GameplayAbilitiesEditor/Private/K2Node_WaitMovementModeChange.cpp
@@ -1,9 +1,9 @@
 // Copyright 1998-2014 Epic Games, Inc. All Rights Reserved.
 
 
-#include "SkillSystemEditorModulePrivatePCH.h"
+#include "AbilitySystemEditorPrivatePCH.h"
 #include "K2Node_WaitMovementModeChange.h"
-#include "SkillSystemBlueprintLibrary.h"
+#include "AbilitySystemBlueprintLibrary.h"
 #include "Abilities/Tasks/AbilityTask_WaitMovementModeChange.h"
 #include "CompilerResultsLog.h"
 
@@ -13,8 +13,8 @@ UK2Node_WaitMovementModeChange::UK2Node_WaitMovementModeChange(const class FPost
 	: Super(PCIP)
 {
 	
-	ProxyFactoryFunctionName = GET_FUNCTION_NAME_CHECKED(USkillSystemBlueprintLibrary, CreateWaitMovementModeChange);
-	ProxyFactoryClass = USkillSystemBlueprintLibrary::StaticClass();
+	ProxyFactoryFunctionName = GET_FUNCTION_NAME_CHECKED(UAbilitySystemBlueprintLibrary, CreateWaitMovementModeChange);
+	ProxyFactoryClass = UAbilitySystemBlueprintLibrary::StaticClass();
 	ProxyClass = UAbilityTask_WaitMovementModeChange::StaticClass();
 }
 
diff --git a/Engine/Source/Editor/GameplayAbilitiesEditor/Private/K2Node_WaitOverlap.cpp b/Engine/Source/Editor/GameplayAbilitiesEditor/Private/K2Node_WaitOverlap.cpp
--- a/Engine/Source/Editor/GameplayAbilitiesEditor/Private/K2Node_WaitOverlap.cpp
+++ b/Engine/Source/Editor/GameplayAbilitiesEditor/Private/K2Node_WaitOverlap.cpp
@@ -1,9 +1,9 @@
 // Copyright 1998-2014 Epic Games, Inc. All Rights Reserved.
 
 
-#include "SkillSystemEditorModulePrivatePCH.h"
+#include "AbilitySystemEditorPrivatePCH.h"
 #include "K2Node_WaitOverlap.h"
-#include "SkillSystemBlueprintLibrary.h"
+#include "AbilitySystemBlueprintLibrary.h"
 #include "Abilities/Tasks/AbilityTask_WaitOverlap.h"
 #include "CompilerResultsLog.h"
 
@@ -13,8 +13,8 @@ UK2Node_WaitOverlap::UK2Node_WaitOverlap(const class FPostConstructInitializePro
 	: Super(PCIP)
 {
 	
-	ProxyFactoryFunctionName = GET_FUNCTION_NAME_CHECKED(USkillSystemBlueprintLibrary, CreateWaitOverlap);
-	ProxyFactoryClass = USkillSystemBlueprintLibrary::StaticClass();
+	ProxyFactoryFunctionName = GET_FUNCTION_NAME_CHECKED(UAbilitySystemBlueprintLibrary, CreateWaitOverlap);
+	ProxyFactoryClass = UAbilitySystemBlueprintLibrary::StaticClass();
 	ProxyClass = UAbilityTask_WaitOverlap::StaticClass();
 }
 
diff --git a/Engine/Source/Editor/GameplayAbilitiesEditor/Public/GameplayAbilitiesEditorModule.h b/Engine/Source/Editor/GameplayAbilitiesEditor/Public/GameplayAbilitiesEditorModule.h
--- a/Engine/Source/Editor/GameplayAbilitiesEditor/Public/GameplayAbilitiesEditorModule.h
+++ b/Engine/Source/Editor/GameplayAbilitiesEditor/Public/GameplayAbilitiesEditorModule.h
@@ -31,7 +31,7 @@ public:
 	 */
 	static inline bool IsAvailable()
 	{
-		return FModuleManager::Get().IsModuleLoaded( "SkillSystemEditor" );
+		return FModuleManager::Get().IsModuleLoaded( "GameplayAbilitiesEditor" );
 	}
 };
 
diff --git a/Engine/Source/Editor/UnrealEd/Private/Editor.cpp b/Engine/Source/Editor/UnrealEd/Private/Editor.cpp
--- a/Engine/Source/Editor/UnrealEd/Private/Editor.cpp
+++ b/Engine/Source/Editor/UnrealEd/Private/Editor.cpp
@@ -564,11 +564,11 @@ void UEditorEngine::Init(IEngineLoop* InEngineLoop)
 			FModuleManager::Get().LoadModule(TEXT("EnvironmentQueryEditor"));
 		}
 
-		bool bSkillSystemEditorEnabled = false;
-		GConfig->GetBool(TEXT("SkillSystemEd"), TEXT("SKillSystemEditorEnabled"), bSkillSystemEditorEnabled, GEngineIni);
-		if (bSkillSystemEditorEnabled)
+		bool bGameplayAbilitiesEnabled = false;
+		GConfig->GetBool(TEXT("GameplayAbilities"), TEXT("GameplayAbilitiedEditorEnabled"), bGameplayAbilitiesEnabled, GEngineIni);
+		if (bGameplayAbilitiesEnabled)
 		{
-			FModuleManager::Get().LoadModule(TEXT("SkillSystemEditor"));
+			FModuleManager::Get().LoadModule(TEXT("GameplayAbilitiesEditor"));
 		}
 	}
 
diff --git a/Engine/Source/Runtime/GameplayAbilities/Classes/Abilities/GameplayAbility.h b/Engine/Source/Runtime/GameplayAbilities/Classes/Abilities/GameplayAbility.h
--- a/Engine/Source/Runtime/GameplayAbilities/Classes/Abilities/GameplayAbility.h
+++ b/Engine/Source/Runtime/GameplayAbilities/Classes/Abilities/GameplayAbility.h
@@ -5,7 +5,7 @@
 #include "GameplayTagContainer.h"
 #include "GameplayTagAssetInterface.h"
 #include "Runtime/Engine/Classes/Animation/AnimInstance.h"
-#include "SkillSystemTypes.h"
+#include "AbilitySystemTypes.h"
 #include "GameplayAbility.generated.h"
 
 class UAnimInstance;
diff --git a/Engine/Source/Runtime/GameplayAbilities/Classes/Abilities/Tasks/AbilityTask_WaitOverlap.h b/Engine/Source/Runtime/GameplayAbilities/Classes/Abilities/Tasks/AbilityTask_WaitOverlap.h
--- a/Engine/Source/Runtime/GameplayAbilities/Classes/Abilities/Tasks/AbilityTask_WaitOverlap.h
+++ b/Engine/Source/Runtime/GameplayAbilities/Classes/Abilities/Tasks/AbilityTask_WaitOverlap.h
@@ -1,7 +1,7 @@
 // Copyright 1998-2014 Epic Games, Inc. All Rights Reserved.
 #pragma once
 #include "AbilityTask.h"
-#include "SkillSystemTypes.h"
+#include "AbilitySystemTypes.h"
 #include "AbilityTask_WaitOverlap.generated.h"
 
 DECLARE_DYNAMIC_MULTICAST_DELEGATE_OneParam(FWaitOverlapDelegate, FGameplayAbilityTargetDataHandle, TargetData);
diff --git a/Engine/Source/Runtime/GameplayAbilities/Classes/AttributeComponent.h b/Engine/Source/Runtime/GameplayAbilities/Classes/AttributeComponent.h
--- a/Engine/Source/Runtime/GameplayAbilities/Classes/AttributeComponent.h
+++ b/Engine/Source/Runtime/GameplayAbilities/Classes/AttributeComponent.h
@@ -2,7 +2,7 @@
 #pragma once
 
 #include "GameplayEffect.h"
-#include "SkillSystemTypes.h"
+#include "AbilitySystemTypes.h"
 #include "AttributeComponent.generated.h"
 
 GAMEPLAYABILITIES_API DECLARE_LOG_CATEGORY_EXTERN(LogAttributeComponent, Log, All);
@@ -24,13 +24,13 @@ struct GAMEPLAYABILITIES_API FAttributeDefaults
 /** 
  * 
  */
-UCLASS(ClassGroup=SkillSystem, hidecategories=(Object,LOD,Lighting,Transform,Sockets,TextureStreaming), editinlinenew, meta=(BlueprintSpawnableComponent))
+UCLASS(ClassGroup=AbilitySystem, hidecategories=(Object,LOD,Lighting,Transform,Sockets,TextureStreaming), editinlinenew, meta=(BlueprintSpawnableComponent))
 class GAMEPLAYABILITIES_API UAttributeComponent : public UActorComponent
 {
 	GENERATED_UCLASS_BODY()
 
 	friend FGameplayEffectSpec;
-	friend class ASkillSystemDebugHUD;
+	friend class AAbilitySystemDebugHUD;
 
 	virtual ~UAttributeComponent();
 
diff --git a/Engine/Source/Runtime/GameplayAbilities/Private/Abilities/GameplayAbility.cpp b/Engine/Source/Runtime/GameplayAbilities/Private/Abilities/GameplayAbility.cpp
--- a/Engine/Source/Runtime/GameplayAbilities/Private/Abilities/GameplayAbility.cpp
+++ b/Engine/Source/Runtime/GameplayAbilities/Private/Abilities/GameplayAbility.cpp
@@ -1,10 +1,10 @@
 // Copyright 1998-2013 Epic Games, Inc. All Rights Reserved.
 
-#include "SkillSystemModulePrivatePCH.h"
+#include "AbilitySystemPrivatePCH.h"
 #include "Abilities/GameplayAbility.h"
 #include "GameplayEffect.h"
 #include "AttributeComponent.h"
-#include "SkillSystemBlueprintLibrary.h"
+#include "AbilitySystemBlueprintLibrary.h"
 
 // --------------------------------------------------------------------------------------------------------------------------------------------------------
 //
@@ -207,7 +207,7 @@ void UGameplayAbility::InputPressed(int32 InputID, const FGameplayAbilityActorIn
 
 void UGameplayAbility::InputReleased(int32 InputID, const FGameplayAbilityActorInfo* ActorInfo)
 {
-	SKILL_LOG(Log, TEXT("InputReleased: %d"), InputID);
+	ABILITY_LOG(Log, TEXT("InputReleased: %d"), InputID);
 }
 
 bool UGameplayAbility::CheckCooldown(const FGameplayAbilityActorInfo* ActorInfo) const
@@ -289,7 +289,7 @@ void UGameplayAbility::ClientActivateAbilitySucceed_Implementation(int32 Predict
 
 void UGameplayAbility::ClientActivateAbilitySucceed_Internal(int32 PredictionKey)
 {
-	SKILL_LOG(Fatal, TEXT("ClientActivateAbilitySucceed_Internal called on base class ability %s"), *GetName());
+	ABILITY_LOG(Fatal, TEXT("ClientActivateAbilitySucceed_Internal called on base class ability %s"), *GetName());
 }
 
 //----------------------------------------------------------------------
@@ -353,7 +353,7 @@ void FGameplayAbilityTargetData::ApplyGameplayEffect(UGameplayEffect* GameplayEf
 	{
 		check(TargetActor);
 		
-		UAttributeComponent * TargetComponent = USkillSystemBlueprintLibrary::GetAttributeComponent(TargetActor);
+		UAttributeComponent * TargetComponent = UAbilitySystemBlueprintLibrary::GetAttributeComponent(TargetActor);
 		if (TargetComponent)
 		{
 			InstigatorInfo.AttributeComponent->ApplyGameplayEffectSpecToTarget(SpecToApply, TargetComponent);
diff --git a/Engine/Source/Runtime/GameplayAbilities/Private/Abilities/GameplayAbility_CharacterJump.cpp b/Engine/Source/Runtime/GameplayAbilities/Private/Abilities/GameplayAbility_CharacterJump.cpp
--- a/Engine/Source/Runtime/GameplayAbilities/Private/Abilities/GameplayAbility_CharacterJump.cpp
+++ b/Engine/Source/Runtime/GameplayAbilities/Private/Abilities/GameplayAbility_CharacterJump.cpp
@@ -1,6 +1,6 @@
 // Copyright 1998-2014 Epic Games, Inc. All Rights Reserved.
 
-#include "SkillSystemModulePrivatePCH.h"
+#include "AbilitySystemPrivatePCH.h"
 #include "Abilities/GameplayAbility_CharacterJump.h"
 
 // --------------------------------------------------------------------------------------------------------------------------------------------------------
diff --git a/Engine/Source/Runtime/GameplayAbilities/Private/Abilities/GameplayAbility_Instanced.cpp b/Engine/Source/Runtime/GameplayAbilities/Private/Abilities/GameplayAbility_Instanced.cpp
--- a/Engine/Source/Runtime/GameplayAbilities/Private/Abilities/GameplayAbility_Instanced.cpp
+++ b/Engine/Source/Runtime/GameplayAbilities/Private/Abilities/GameplayAbility_Instanced.cpp
@@ -1,6 +1,6 @@
 // Copyright 1998-2013 Epic Games, Inc. All Rights Reserved.
 
-#include "SkillSystemModulePrivatePCH.h"
+#include "AbilitySystemPrivatePCH.h"
 #include "Abilities/GameplayAbility_Instanced.h"
 
 // --------------------------------------------------------------------------------------------------------------------------------------------------------
@@ -34,7 +34,7 @@ bool UGameplayAbility_Instanced::CanActivateAbility(const FGameplayAbilityActorI
 		&& CanActivateFunction->GetOuter()->IsA(UBlueprintGeneratedClass::StaticClass())
 		&& K2_CanActivateAbility() == false)
 	{
-		SKILL_LOG(Log, TEXT("CanActivateAbility %s failed, blueprint refused"), *GetName());
+		ABILITY_LOG(Log, TEXT("CanActivateAbility %s failed, blueprint refused"), *GetName());
 		return false;
 	}
 
diff --git a/Engine/Source/Runtime/GameplayAbilities/Private/Abilities/GameplayAbility_Montage.cpp b/Engine/Source/Runtime/GameplayAbilities/Private/Abilities/GameplayAbility_Montage.cpp
--- a/Engine/Source/Runtime/GameplayAbilities/Private/Abilities/GameplayAbility_Montage.cpp
+++ b/Engine/Source/Runtime/GameplayAbilities/Private/Abilities/GameplayAbility_Montage.cpp
@@ -1,6 +1,6 @@
 // Copyright 1998-2013 Epic Games, Inc. All Rights Reserved.
 
-#include "SkillSystemModulePrivatePCH.h"
+#include "AbilitySystemPrivatePCH.h"
 #include "Abilities/GameplayAbility_Montage.h"
 #include "GameplayEffect.h"
 #include "AttributeComponent.h"
diff --git a/Engine/Source/Runtime/GameplayAbilities/Private/AbilitySystemBlueprintLibrary.cpp b/Engine/Source/Runtime/GameplayAbilities/Private/AbilitySystemBlueprintLibrary.cpp
new file mode 100644
--- /dev/null
+++ b/Engine/Source/Runtime/GameplayAbilities/Private/AbilitySystemBlueprintLibrary.cpp
@@ -0,0 +1,121 @@
+// Copyright 1998-2014 Epic Games, Inc. All Rights Reserved.
+
+#include "AbilitySystemPrivatePCH.h"
+#include "AbilitySystemBlueprintLibrary.h"
+#include "Abilities/GameplayAbility.h"
+#include "Abilities/GameplayAbility_Instanced.h"
+#include "Abilities/Tasks/AbilityTask_PlayMontageAndWait.h"
+#include "Abilities/Tasks/AbilityTask_WaitMovementModeChange.h"
+#include "Abilities/Tasks/AbilityTask_WaitOverlap.h"
+#include "LatentActions.h"
+
+UAbilitySystemBlueprintLibrary::UAbilitySystemBlueprintLibrary(const class FPostConstructInitializeProperties& PCIP)
+: Super(PCIP)
+{
+}
+
+UAbilityTask_PlayMontageAndWait* UAbilitySystemBlueprintLibrary::CreatePlayMontageAndWaitProxy(class UObject* WorldContextObject, UAnimMontage *MontageToPlay)
+{
+	check(WorldContextObject);
+	UGameplayAbility_Instanced * Ability = CastChecked<UGameplayAbility_Instanced>(WorldContextObject);
+	if (Ability)
+	{
+		AActor * ActorOwner = Cast<AActor>(Ability->GetOuter());
+
+		UAbilityTask_PlayMontageAndWait * MyObj = NULL;
+		MyObj = NewObject<UAbilityTask_PlayMontageAndWait>();
+
+		FOnMontageEnded EndDelegate;
+		EndDelegate.BindUObject(MyObj, &UAbilityTask_PlayMontageAndWait::OnMontageEnded);
+
+		FGameplayAbilityActorInfo	ActorInfo;
+		ActorInfo.InitFromActor(ActorOwner);
+
+		ActorInfo.AnimInstance->Montage_Play(MontageToPlay, 1.f);
+		ActorInfo.AnimInstance->Montage_SetEndDelegate(EndDelegate);
+
+		return MyObj;
+	}
+
+	return NULL;	
+}
+
+
+UAbilityTask_WaitMovementModeChange* UAbilitySystemBlueprintLibrary::CreateWaitMovementModeChange(class UObject* WorldContextObject, EMovementMode NewMode)
+{
+	check(WorldContextObject);
+	UGameplayAbility_Instanced * Ability = CastChecked<UGameplayAbility_Instanced>(WorldContextObject);
+	if (Ability)
+	{
+		AActor * ActorOwner = Cast<AActor>(Ability->GetOuter());
+
+		UAbilityTask_WaitMovementModeChange * MyObj = NULL;
+		MyObj = NewObject<UAbilityTask_WaitMovementModeChange>();
+		MyObj->RequiredMode = NewMode;
+
+		FGameplayAbilityActorInfo	ActorInfo;
+		ActorInfo.InitFromActor(ActorOwner);
+
+		ACharacter * Character = Cast<ACharacter>(ActorInfo.Actor.Get());
+		if (Character)
+		{
+			Character->MovementModeChangedDelegate.AddDynamic(MyObj, &UAbilityTask_WaitMovementModeChange::OnMovementModeChange);
+		}
+		return MyObj;
+	}
+
+	return NULL;
+}
+
+/**
+ *	Need:
+ *	-Easy way to specify which primitive components should be used for this overlap test
+ *	-Easy way to specify which types of actors/collision overlaps that we care about/want to block on
+ */
+
+UAbilityTask_WaitOverlap* UAbilitySystemBlueprintLibrary::CreateWaitOverlap(class UObject* WorldContextObject, EMovementMode NewMode)
+{
+	check(WorldContextObject);
+	UGameplayAbility_Instanced * Ability = CastChecked<UGameplayAbility_Instanced>(WorldContextObject);
+	if (Ability)
+	{
+		AActor * ActorOwner = Cast<AActor>(Ability->GetOuter());
+
+		UAbilityTask_WaitOverlap * MyObj = NULL;
+		MyObj = NewObject<UAbilityTask_WaitOverlap>();
+
+
+		// TEMP
+		UPrimitiveComponent * PrimComponent = Cast<UPrimitiveComponent>(ActorOwner->GetRootComponent());
+		if (!PrimComponent)
+		{			
+			PrimComponent = ActorOwner->FindComponentByClass<UPrimitiveComponent>();
+		}
+		check(PrimComponent);
+
+		PrimComponent->OnComponentBeginOverlap.AddDynamic(MyObj, &UAbilityTask_WaitOverlap::OnOverlapCallback);
+		PrimComponent->OnComponentHit.AddDynamic(MyObj, &UAbilityTask_WaitOverlap::OnHitCallback);
+
+		return MyObj;
+	}
+
+	return NULL;
+}
+
+/**
+ *	Fixme: This is supposed to be a small optimization but here we are going through module code, then a global uobject, then finally a virtual function - yuck!
+ *		-think of a better way to have a global, overridable per project function. Thats it!
+ */
+
+UAttributeComponent* UAbilitySystemBlueprintLibrary::GetAttributeComponent(AActor *Actor)
+{
+	return UAbilitySystemGlobals::Get().GetAttributeComponentFromActor(Actor);
+}
+
+void UAbilitySystemBlueprintLibrary::ApplyGameplayEffectToTargetData(FGameplayAbilityTargetDataHandle Target, UGameplayEffect *GameplayEffect, const FGameplayAbilityActorInfo InstigatorInfo)
+{
+	if (Target.Data.IsValid())
+	{
+		Target.Data->ApplyGameplayEffect(GameplayEffect, InstigatorInfo);
+	}
+}
\ No newline at end of file
diff --git a/Engine/Source/Runtime/GameplayAbilities/Private/AbilitySystemDebugHUD.cpp b/Engine/Source/Runtime/GameplayAbilities/Private/AbilitySystemDebugHUD.cpp
new file mode 100644
--- /dev/null
+++ b/Engine/Source/Runtime/GameplayAbilities/Private/AbilitySystemDebugHUD.cpp
@@ -0,0 +1,166 @@
+// Copyright 1998-2014 Epic Games, Inc. All Rights Reserved.
+
+#include "AbilitySystemPrivatePCH.h"
+#include "AttributeComponent.h"
+#include "AbilitySystemDebugHUD.h"
+#include "DebugRenderSceneProxy.h"
+
+AAbilitySystemDebugHUD::AAbilitySystemDebugHUD(const class FPostConstructInitializeProperties& PCIP) : Super(PCIP)
+{
+
+}
+
+void AAbilitySystemDebugHUD::DrawHUD()
+{
+	Super::DrawHUD();
+}
+
+void AAbilitySystemDebugHUD::DrawWithBackground(UFont* InFont, const FString& Text, const FColor& TextColor, EAlignHorizontal::Type HAlign, float& OffsetX, EAlignVertical::Type VAlign, float& OffsetY, float Alpha)
+{
+	float SizeX, SizeY;
+	Canvas->StrLen(InFont, Text, SizeX, SizeY);
+
+	const float PosX = (HAlign == EAlignHorizontal::Center) ? OffsetX + (Canvas->ClipX - SizeX) * 0.5f :
+		(HAlign == EAlignHorizontal::Left) ? Canvas->OrgX + OffsetX :
+		Canvas->ClipX - SizeX - OffsetX;
+
+	const float PosY = (VAlign == EAlignVertical::Center) ? OffsetY + (Canvas->ClipY - SizeY) * 0.5f :
+		(VAlign == EAlignVertical::Top) ? Canvas->OrgY + OffsetY :
+		Canvas->ClipY - SizeY - OffsetY;
+
+	const float BoxPadding = 5.0f;
+
+	const float X = PosX - BoxPadding;
+	const float Y = PosY - BoxPadding;
+	const float Z = 0.1f;
+	FCanvasTileItem TileItem(FVector2D(X, Y), FVector2D(SizeX + BoxPadding * 2.0f, SizeY + BoxPadding * 2.0f), FLinearColor(0.75f, 0.75f, 0.75f, Alpha));
+	Canvas->DrawItem(TileItem);
+
+	FLinearColor TextCol(TextColor);
+	TextCol.A = Alpha;
+	FCanvasTextItem TextItem(FVector2D(PosX, PosY), FText::FromString(Text), GEngine->GetSmallFont(), TextCol);
+	Canvas->DrawItem(TextItem);
+
+	OffsetY += 25;
+}
+
+void AAbilitySystemDebugHUD::DrawDebugHUD(UCanvas* InCanvas, APlayerController* PC)
+{
+	Canvas = InCanvas;
+	if (!Canvas)
+	{
+		return;
+	}
+
+	UPlayer * LocalPlayer = GetWorld()->GetFirstLocalPlayerFromController();
+	PC = LocalPlayer->PlayerController;	
+
+	if (!PC)
+	{
+		return;
+	}
+
+	if (PC->GetPawn())
+	{
+		UAttributeComponent * AttributeComponent = PC->GetPawn()->FindComponentByClass<UAttributeComponent>();
+		if (AttributeComponent)
+		{
+			DrawDebugAttributeComponent(AttributeComponent);
+		}
+	}
+}
+
+void AAbilitySystemDebugHUD::DrawDebugAttributeComponent(UAttributeComponent *Component)
+{
+	UWorld *World = GetWorld();
+	float GameWorldTime = World->GetTimeSeconds();
+
+	UFont* Font = GEngine->GetSmallFont();
+	FColor Color(38, 128, 0);
+	float X = 20.f;
+	float Y = 20.f;
+
+	FString String = FString::Printf(TEXT("%.2f"), Component->GetWorld()->GetTimeSeconds());
+	DrawWithBackground(Font, String, Color, EAlignHorizontal::Left, X, EAlignVertical::Top, Y);
+
+	String = FString::Printf(TEXT("%s (%d)"), *Component->GetPathName(), Component->IsDefaultSubobject());
+	DrawWithBackground(Font, String, Color, EAlignHorizontal::Left, X, EAlignVertical::Top, Y);
+
+	
+	String = FString::Printf(TEXT("%s == %s"), *Component->GetArchetype()->GetPathName(), *Component->GetClass()->GetDefaultObject()->GetPathName());
+	DrawWithBackground(Font, String, Color, EAlignHorizontal::Left, X, EAlignVertical::Top, Y);
+
+
+	for (UAttributeSet * Set : Component->SpawnedAttributes)
+	{
+		if (!Set)
+			continue;
+		check(Set);
+
+		// Draw Attribute Set
+		DrawWithBackground(Font, FString::Printf(TEXT("%s (%d)"), *Set->GetName(), Set->IsDefaultSubobject()), Color, EAlignHorizontal::Left, X, EAlignVertical::Top, Y);
+
+		String = FString::Printf(TEXT("%s == %s"), *Set->GetArchetype()->GetPathName(), *Set->GetClass()->GetDefaultObject()->GetPathName());
+		DrawWithBackground(Font, String, Color, EAlignHorizontal::Left, X, EAlignVertical::Top, Y);
+
+		for (TFieldIterator<UProperty> PropertyIt(Set->GetClass(), EFieldIteratorFlags::IncludeSuper); PropertyIt; ++PropertyIt)
+		{
+			UProperty *Prop = *PropertyIt;			
+
+			FString ValueString;
+			void *PropertyValue = Prop->ContainerPtrToValuePtr<void>(Set);
+			Prop->ExportTextItem(ValueString, PropertyValue, NULL, NULL, 0);
+
+			String = FString::Printf(TEXT("%s: %s"), *Prop->GetName(), *ValueString);
+			DrawWithBackground(Font, String, Color, EAlignHorizontal::Left, X, EAlignVertical::Top, Y);
+		}
+
+		Y+= 25;
+		// Draw Active GameplayEffect
+		for (FActiveGameplayEffect &Effect : Component->ActiveGameplayEffects.GameplayEffects)
+		{
+			String = FString::Printf(TEXT("%s. [%d] %.2f"), *Effect.Spec.ToSimpleString(), Effect.PredictionKey, Effect.GetTimeRemaining(GameWorldTime));
+			DrawWithBackground(Font, String, Color, EAlignHorizontal::Left, X, EAlignVertical::Top, Y);	
+		}
+	}
+	
+}
+
+#if !(UE_BUILD_SHIPPING || UE_BUILD_TEST)
+static void	ToggleDebugHUD(const TArray<FString>& Args, UWorld* InWorld)
+{
+
+	if (!InWorld)
+		return;
+
+	AAbilitySystemDebugHUD *HUD = NULL;
+
+	for (FActorIterator It(InWorld); It; ++It)
+	{
+		AActor *Actor = *It;
+		HUD = Cast<AAbilitySystemDebugHUD>(Actor);
+		if (HUD)
+			break;
+	}
+
+	if (!HUD)
+	{
+		HUD = InWorld->SpawnActor<AAbilitySystemDebugHUD>();
+		
+		FDebugDrawDelegate DrawDebugDelegate = FDebugDrawDelegate::CreateUObject(HUD, &AAbilitySystemDebugHUD::DrawDebugHUD);
+		UDebugDrawService::Register(TEXT("GameplayDebug"), DrawDebugDelegate);
+	}
+	else
+	{
+		FDebugDrawDelegate DrawDebugDelegate = FDebugDrawDelegate::CreateUObject(HUD, &AAbilitySystemDebugHUD::DrawDebugHUD);
+		UDebugDrawService::Unregister(DrawDebugDelegate);
+		HUD->Destroy();
+	}
+}
+
+FAutoConsoleCommandWithWorldAndArgs AbilitySystemToggleDebugHUDCommand(
+	TEXT("AbilitySystem.ToggleDebugHUD"),
+	TEXT("ToggleDebugHUD Drawing"),
+	FConsoleCommandWithWorldAndArgsDelegate::CreateStatic(ToggleDebugHUD)
+	);
+#endif
diff --git a/Engine/Source/Runtime/GameplayAbilities/Private/AbilitySystemGlobals.cpp b/Engine/Source/Runtime/GameplayAbilities/Private/AbilitySystemGlobals.cpp
new file mode 100644
--- /dev/null
+++ b/Engine/Source/Runtime/GameplayAbilities/Private/AbilitySystemGlobals.cpp
@@ -0,0 +1,93 @@
+// Copyright 1998-2014 Epic Games, Inc. All Rights Reserved.
+
+#include "AbilitySystemPrivatePCH.h"
+#include "AttributeComponent.h"
+#include "AbilitySystemTypes.h"
+
+UAbilitySystemGlobals::UAbilitySystemGlobals(const class FPostConstructInitializeProperties& PCIP)
+: Super(PCIP)
+{
+
+}
+
+UCurveTable * UAbilitySystemGlobals::GetGlobalCurveTable()
+{
+	if (!GlobalCurveTable && !GlobalCurveTableName.IsEmpty())
+	{
+		GlobalCurveTable = LoadObject<UCurveTable>(NULL, *GlobalCurveTableName, NULL, LOAD_None, NULL);
+#if WITH_EDITOR
+		// Hook into notifications for object re-imports so that the gameplay tag tree can be reconstructed if the table changes
+		if (GIsEditor && GlobalCurveTable)
+		{
+			GEditor->OnObjectReimported().AddUObject(this, &UAbilitySystemGlobals::OnCurveTableReimported);
+		}
+#endif
+	}
+	
+	return GlobalCurveTable;
+}
+
+UDataTable * UAbilitySystemGlobals::GetGlobalAttributeDataTable()
+{
+	if (!GlobalAttributeDataTable && !GlobalAttributeDataTableName.IsEmpty())
+	{
+		GlobalAttributeDataTable = LoadObject<UDataTable>(NULL, *GlobalAttributeDataTableName, NULL, LOAD_None, NULL);
+#if WITH_EDITOR
+		// Hook into notifications for object re-imports so that the gameplay tag tree can be reconstructed if the table changes
+		if (GIsEditor && GlobalAttributeDataTable)
+		{
+			GEditor->OnObjectReimported().AddUObject(this, &UAbilitySystemGlobals::OnDataTableReimported);
+		}
+#endif
+	}
+
+	return GlobalAttributeDataTable;
+}
+
+#if WITH_EDITOR
+
+void UAbilitySystemGlobals::OnCurveTableReimported(UObject* InObject)
+{
+	if (GIsEditor && !IsRunningCommandlet() && InObject && InObject == GlobalCurveTable)
+	{
+	}
+}
+
+void UAbilitySystemGlobals::OnDataTableReimported(UObject* InObject)
+{
+	if (GIsEditor && !IsRunningCommandlet() && InObject && InObject == GlobalAttributeDataTable)
+	{
+	}
+}
+
+#endif
+
+
+UAttributeComponent * UAbilitySystemGlobals::GetAttributeComponentFromActor(AActor *Actor) const
+{
+	if (Actor)
+	{
+		return Actor->FindComponentByClass<UAttributeComponent>();
+	}
+
+	// Caller should check this
+	ensure(false);
+	return NULL;
+}
+
+FGameplayAbilityActorInfo * UAbilitySystemGlobals::AllocAbilityActorInfo(AActor *Actor) const
+{
+	FGameplayAbilityActorInfo * Info = new FGameplayAbilityActorInfo();
+	Info->InitFromActor(Actor);
+	return Info;
+}
+
+/** This is just some syntax sugar to avoid calling gode to have to IGameplayAbilitiesModule::Get() */
+UAbilitySystemGlobals& UAbilitySystemGlobals::Get()
+{
+	static UAbilitySystemGlobals * GlobalPtr = IGameplayAbilitiesModule::Get().GetAbilitySystemGlobals();
+	check(GlobalPtr == IGameplayAbilitiesModule::Get().GetAbilitySystemGlobals());
+	check(GlobalPtr);
+
+	return *GlobalPtr;
+}
diff --git a/Engine/Source/Runtime/GameplayAbilities/Private/AbilitySystemLog.cpp b/Engine/Source/Runtime/GameplayAbilities/Private/AbilitySystemLog.cpp
new file mode 100644
--- /dev/null
+++ b/Engine/Source/Runtime/GameplayAbilities/Private/AbilitySystemLog.cpp
@@ -0,0 +1,83 @@
+// Copyright 1998-2014 Epic Games, Inc. All Rights Reserved.
+#include "AbilitySystemPrivatePCH.h"
+
+DEFINE_LOG_CATEGORY(LogAbilitySystem);
+
+AbilitySystemLogScope::~AbilitySystemLogScope()
+{
+	AbilitySystemLog::PopScope();
+}
+
+void AbilitySystemLogScope::Init()
+{
+	PrintedInThisScope = false;
+	AbilitySystemLog::PushScope(this);
+}
+
+// ----------------------------------------------
+
+AbilitySystemLog * GetInstance()
+{
+	static AbilitySystemLog Instance;
+	return &Instance;
+}
+
+FString AbilitySystemLog::Log(ELogVerbosity::Type Verbosity, FString Log)
+{
+	AbilitySystemLog * Instance = GetInstance();
+
+#if !NO_LOGGING
+	if (!LogAbilitySystem.IsSuppressed(Verbosity))
+	{
+		for (int32 idx=0; idx < Instance->ScopeStack.Num(); ++idx)
+		{
+			AbilitySystemLogScope *Scope = Instance->ScopeStack[idx];
+			if (!Scope->PrintedInThisScope && !Scope->ScopeName.IsEmpty())
+			{
+				if (Instance->NeedNewLine)
+				{
+					UE_LOG(LogAbilitySystem, Log, TEXT(""));
+					Instance->NeedNewLine = false;
+				}
+
+				Scope->PrintedInThisScope = true;
+				int32 ident = (2 * idx);
+				FString IndentStrX = FString::Printf(TEXT("%*s"), ident, TEXT(""));
+				UE_LOG(LogAbilitySystem, Log, TEXT("%s<%s>"), *IndentStrX, *Scope->ScopeName);
+			}
+		}
+	}
+#endif
+
+	FString IndentStr = FString::Printf(TEXT("%*s"), Instance->Indent, TEXT(""));
+	
+	return IndentStr + Log;
+}
+
+void AbilitySystemLog::PushScope(AbilitySystemLogScope * Scope)
+{
+	AbilitySystemLog * Instance = GetInstance();
+
+	Instance->Indent += 2;
+	Instance->ScopeStack.Push(Scope);
+}
+
+void AbilitySystemLog::PopScope()
+{
+	AbilitySystemLog * Instance = GetInstance();
+
+	Instance->NeedNewLine = true;
+	Instance->Indent -= 2;
+	check(Instance->Indent >= 0);
+
+	if (Instance->ScopeStack.Top()->PrintedInThisScope)
+	{
+		if (!Instance->ScopeStack.Top()->ScopeName.IsEmpty())
+		{
+			FString IndentStrX = FString::Printf(TEXT("%*s"), Instance->Indent, TEXT(""));
+			UE_LOG(LogAbilitySystem, Log, TEXT("%s</%s>"), *IndentStrX, *Instance->ScopeStack.Top()->ScopeName);
+		}
+	}
+	
+	Instance->ScopeStack.Pop();
+}
\ No newline at end of file
diff --git a/Engine/Source/Runtime/GameplayAbilities/Private/AbilitySystemPrivatePCH.h b/Engine/Source/Runtime/GameplayAbilities/Private/AbilitySystemPrivatePCH.h
new file mode 100644
--- /dev/null
+++ b/Engine/Source/Runtime/GameplayAbilities/Private/AbilitySystemPrivatePCH.h
@@ -0,0 +1,21 @@
+// Copyright 1998-2014 Epic Games, Inc. All Rights Reserved.
+
+#pragma once
+
+#if WITH_EDITOR
+#include "UnrealEd.h"
+#endif
+
+
+// You should place include statements to your module's private header files here.  You only need to
+// add includes for headers that are used in most of your module's source files though.
+
+#include "Core.h"
+#include "Engine.h"
+#include "Engine/DataTable.h"
+
+#include "AbilitySystemTypes.h"
+
+#include "AbilitySystemStats.h"
+#include "GameplayAbilitiesModule.h"
+#include "AbilitySystemLog.h"
diff --git a/Engine/Source/Runtime/GameplayAbilities/Private/AbilitySystemStats.cpp b/Engine/Source/Runtime/GameplayAbilities/Private/AbilitySystemStats.cpp
new file mode 100644
--- /dev/null
+++ b/Engine/Source/Runtime/GameplayAbilities/Private/AbilitySystemStats.cpp
@@ -0,0 +1,10 @@
+// Copyright 1998-2014 Epic Games, Inc. All Rights Reserved.
+
+#include "AbilitySystemPrivatePCH.h"
+
+DEFINE_STAT(STAT_GameplayEffectsHasAllTags);
+DEFINE_STAT(STAT_GameplayEffectsHasAnyTag);
+DEFINE_STAT(STAT_GameplayEffectsTick);
+DEFINE_STAT(STAT_GameplayEffectsCanApplyAttributeModifiers);
+DEFINE_STAT(STAT_GameplayEffectsGetActiveEffectsData);
+DEFINE_STAT(STAT_GameplayAbilityGetCooldownTimeRemaining);
\ No newline at end of file
diff --git a/Engine/Source/Runtime/GameplayAbilities/Private/AbilitySystemTestAttributeSet.cpp b/Engine/Source/Runtime/GameplayAbilities/Private/AbilitySystemTestAttributeSet.cpp
new file mode 100644
--- /dev/null
+++ b/Engine/Source/Runtime/GameplayAbilities/Private/AbilitySystemTestAttributeSet.cpp
@@ -0,0 +1,131 @@
+// Copyright 1998-2014 Epic Games, Inc. All Rights Reserved.
+
+#include "AbilitySystemPrivatePCH.h"
+#include "AbilitySystemTestAttributeSet.h"
+#include "GameplayTagsModule.h"
+#include "GameplayEffectExtension.h"
+#include "AttributeComponent.h"
+
+
+UAbilitySystemTestAttributeSet::UAbilitySystemTestAttributeSet(const class FPostConstructInitializeProperties& PCIP)
+: Super(PCIP)
+{
+	Health = MaxHealth = 100.f;
+	Mana = MaxMana = 100.f;
+	
+	Damage = 0.f;
+	CritChance = 0.f;
+	SpellDamage = 0.f;
+	PhysicalDamage = 0.f;
+	Strength = 0.f;
+	StackingAttribute1 = 0.f;
+	StackingAttribute2 = 0.f;
+	NoStackAttribute = 0.f;
+}
+
+void UAbilitySystemTestAttributeSet::PreAttributeModify(struct FGameplayEffectModCallbackData &Data)
+{
+	static UProperty *HealthProperty = FindFieldChecked<UProperty>(UAbilitySystemTestAttributeSet::StaticClass(), GET_MEMBER_NAME_CHECKED(UAbilitySystemTestAttributeSet, Health));
+	static UProperty *DamageProperty = FindFieldChecked<UProperty>(UAbilitySystemTestAttributeSet::StaticClass(), GET_MEMBER_NAME_CHECKED(UAbilitySystemTestAttributeSet, Damage));
+
+	// In this function, our GameplayEffect mod has been evaluated. We have a magnitude and a Tags collection that we can still modify before it is applied.
+	// We also still have the Aggregation data that calculated Data.EvaluatedData. If we really needed to, we could look at this, remove or change things at the aggregator level, and reevaluate ourselves.
+	// But that would be considered very advanced/rare.
+
+	UProperty *ModifiedProperty = Data.ModifierSpec.Info.Attribute.GetUProperty();
+
+	// Is Damage about to be applied?
+	if (DamageProperty == ModifiedProperty)
+	{
+		// Can the target dodge this completely?
+		if (DodgeChance > 0.f)
+		{
+			if (FMath::FRand() <= DodgeChance)
+			{
+				// Dodge!
+				Data.EvaluatedData.Magnitude = 0.f;
+				Data.EvaluatedData.Tags.AddTag(IGameplayTagsModule::RequestGameplayTag(FName(TEXT("Dodged"))));
+
+				// How dodge is handled will be game dependant. There are a few options I think of:
+				// -We still apply 0 damage, but tag it as Dodged. The GameplayCue system could pick up on this and play a visual effect. The combat log could pick up in and print it out too.
+				// -We throw out this GameplayEffect right here, and apply another GameplayEffect for 'Dodge' it wouldn't modify an attribute but could trigger gameplay cues, it could serve as a 'cooldown' for dodge
+				//		if the game wanted rules like 'you can't dodge more than once every .5 seconds', etc.
+			}
+		}		
+		
+		if (Data.EvaluatedData.Magnitude > 0.f)
+		{
+			// Check the source - does he have Crit?
+			UAbilitySystemTestAttributeSet * SourceAttributes = Data.EffectSpec.InstigatorContext.GetOriginalInstigatorAttributeComponent()->GetSet<UAbilitySystemTestAttributeSet>();
+			if (SourceAttributes && SourceAttributes->CritChance > 0.f)
+			{
+				if (FMath::FRand() <= SourceAttributes->CritChance)
+				{
+					// Crit!
+					Data.EvaluatedData.Magnitude *= SourceAttributes->CritMultiplier;
+					Data.EvaluatedData.Tags.AddTag(IGameplayTagsModule::RequestGameplayTag(FName(TEXT("Damage.Crit"))));
+				}
+			}
+
+			// Now apply armor reduction
+			if (Data.EvaluatedData.Tags.HasTag(IGameplayTagsModule::RequestGameplayTag(FName(TEXT("Damage.Physical"))), EGameplayTagMatchType::IncludeParentTags, EGameplayTagMatchType::Explicit))
+			{
+				// This is a trivial/naive implementation of armor. It assumes the rmorDamageReduction is an actual % to reduce physics damage by.
+				// Real games would probably use armor rating as an attribute and calculate a % here based on the damage source's level, etc.
+				Data.EvaluatedData.Magnitude *= (1.f - ArmorDamageReduction);
+				Data.EvaluatedData.Tags.AddTag(IGameplayTagsModule::RequestGameplayTag(FName(TEXT("Damage.Mitigated.Armor"))));
+			}
+		}
+
+		// At this point, the Magnitude of the applied damage may have been modified by us. We still do the translation to Health in UAbilitySystemTestAttributeSet::PostAttributeModify.
+	}
+}
+
+void UAbilitySystemTestAttributeSet::PostAttributeModify(const struct FGameplayEffectModCallbackData &Data)
+{
+	static UProperty *HealthProperty = FindFieldChecked<UProperty>(UAbilitySystemTestAttributeSet::StaticClass(), GET_MEMBER_NAME_CHECKED(UAbilitySystemTestAttributeSet, Health));
+	static UProperty *DamageProperty = FindFieldChecked<UProperty>(UAbilitySystemTestAttributeSet::StaticClass(), GET_MEMBER_NAME_CHECKED(UAbilitySystemTestAttributeSet, Damage));
+
+	UProperty *ModifiedProperty = Data.ModifierSpec.Info.Attribute.GetUProperty();
+
+	// What property was modified?
+	if (DamageProperty == ModifiedProperty)
+	{
+		// Anytime Damage is applied with 'Damage.Fire' tag, there is a chance to apply a burning DOT
+		if (Data.EvaluatedData.Tags.HasTag( IGameplayTagsModule::RequestGameplayTag(FName(TEXT("FireDamage"))), EGameplayTagMatchType::IncludeParentTags, EGameplayTagMatchType::Explicit))
+		{
+			// Logic to rand() a burning DOT, if successful, apply DOT GameplayEffect to the target
+		}
+
+		// Treat damage as minus health
+		Health -= Damage;
+		Damage = 0.f;
+
+		// Check for Death?
+		//  -This could be defined here or at the actor level.
+		//  -Doing it here makes a lot of sense to me, but we have legacy code in ::TakeDamage function, so some games may just want to punt to that pipeline from here.
+	}
+}
+
+
+void UAbilitySystemTestAttributeSet::GetLifetimeReplicatedProps(TArray< FLifetimeProperty > & OutLifetimeProps) const
+{
+	/*
+	DOREPLIFETIME( UAbilitySystemTestAttributeSet, MaxHealth);
+	DOREPLIFETIME( UAbilitySystemTestAttributeSet, Health);
+	DOREPLIFETIME( UAbilitySystemTestAttributeSet, Mana);
+	DOREPLIFETIME( UAbilitySystemTestAttributeSet, MaxMana);
+
+	DOREPLIFETIME( UAbilitySystemTestAttributeSet, SpellDamage);
+	DOREPLIFETIME( UAbilitySystemTestAttributeSet, PhysicalDamage);
+
+	DOREPLIFETIME( UAbilitySystemTestAttributeSet, CritChance);
+	DOREPLIFETIME( UAbilitySystemTestAttributeSet, CritMultiplier);
+	DOREPLIFETIME( UAbilitySystemTestAttributeSet, ArmorDamageReduction);
+
+	DOREPLIFETIME( UAbilitySystemTestAttributeSet, DodgeChance);
+	DOREPLIFETIME( UAbilitySystemTestAttributeSet, LifeSteal);
+
+	DOREPLIFETIME( UAbilitySystemTestAttributeSet, Strength);
+	*/
+}
\ No newline at end of file
diff --git a/Engine/Source/Runtime/GameplayAbilities/Private/AbilitySystemTestPawn.cpp b/Engine/Source/Runtime/GameplayAbilities/Private/AbilitySystemTestPawn.cpp
new file mode 100644
--- /dev/null
+++ b/Engine/Source/Runtime/GameplayAbilities/Private/AbilitySystemTestPawn.cpp
@@ -0,0 +1,57 @@
+// Copyright 1998-2014 Epic Games, Inc. All Rights Reserved.
+
+#include "AbilitySystemPrivatePCH.h"
+#include "AbilitySystemTestPawn.h"
+#include "AbilitySystemTestAttributeSet.h"
+#include "AttributeComponent.h"
+
+FName  AAbilitySystemTestPawn::AttributeComponentName(TEXT("AttributeComponent0"));
+
+AAbilitySystemTestPawn::AAbilitySystemTestPawn(const class FPostConstructInitializeProperties& PCIP)
+: Super(PCIP)
+{
+	AttributeComponent = PCIP.CreateDefaultSubobject<UAttributeComponent>(this, AAbilitySystemTestPawn::AttributeComponentName);
+	AttributeComponent->SetIsReplicated(true);
+
+	//DefaultAbilitySet = NULL;
+}
+
+void AAbilitySystemTestPawn::PostInitializeComponents()
+{	
+	static UProperty *DamageProperty = FindFieldChecked<UProperty>(UAbilitySystemTestAttributeSet::StaticClass(), GET_MEMBER_NAME_CHECKED(UAbilitySystemTestAttributeSet, Damage));
+
+	Super::PostInitializeComponents();
+	GameplayCueHandler.Owner = this;
+	AttributeComponent->InitStats(UAbilitySystemTestAttributeSet::StaticClass(), NULL);
+
+	/*
+	if (DefaultAbilitySet != NULL)
+	{
+		AttributeComponent->InitializeAbilities(DefaultAbilitySet);
+	}
+	*/
+}
+
+void AAbilitySystemTestPawn::GameplayCueActivated(const FGameplayTagContainer & GameplayCueTags, float NormalizedMagnitude, const FGameplayEffectInstigatorContext InstigatorContext)
+{
+	ABILITY_LOG(Log, TEXT("GameplayCueExecuted: %s. %.2f"), *GetName(), NormalizedMagnitude);
+	GameplayCueHandler.GameplayCueActivated(GameplayCueTags, NormalizedMagnitude, InstigatorContext);
+}
+
+void AAbilitySystemTestPawn::GameplayCueExecuted(const FGameplayTagContainer & GameplayCueTags, float NormalizedMagnitude, const FGameplayEffectInstigatorContext InstigatorContext)
+{
+	ABILITY_LOG(Log, TEXT("GameplayCueExecuted: %s. %.2f"), *GetName(), NormalizedMagnitude);
+	GameplayCueHandler.GameplayCueExecuted(GameplayCueTags, NormalizedMagnitude, InstigatorContext);
+}
+
+void AAbilitySystemTestPawn::GameplayCueAdded(const FGameplayTagContainer & GameplayCueTags, float NormalizedMagnitude, const FGameplayEffectInstigatorContext InstigatorContext)
+{
+	ABILITY_LOG(Log, TEXT("GameplayCueAdded: %s. %.2f"), *GetName(), NormalizedMagnitude);
+	GameplayCueHandler.GameplayCueAdded(GameplayCueTags, NormalizedMagnitude, InstigatorContext);
+}
+
+void AAbilitySystemTestPawn::GameplayCueRemoved(const FGameplayTagContainer & GameplayCueTags, float NormalizedMagnitude, const FGameplayEffectInstigatorContext InstigatorContext)
+{
+	ABILITY_LOG(Log, TEXT("GameplayCueRemoved: %s. %.2f"), *GetName(), NormalizedMagnitude);
+	GameplayCueHandler.GameplayCueRemoved(GameplayCueTags, NormalizedMagnitude, InstigatorContext);
+}
\ No newline at end of file
diff --git a/Engine/Source/Runtime/GameplayAbilities/Private/AttributeComponent.cpp b/Engine/Source/Runtime/GameplayAbilities/Private/AttributeComponent.cpp
--- a/Engine/Source/Runtime/GameplayAbilities/Private/AttributeComponent.cpp
+++ b/Engine/Source/Runtime/GameplayAbilities/Private/AttributeComponent.cpp
@@ -1,7 +1,7 @@
 // Copyright 1998-2013 Epic Games, Inc. All Rights Reserved.
 // ActorComponent.cpp: Actor component implementation.
 
-#include "SkillSystemModulePrivatePCH.h"
+#include "AbilitySystemPrivatePCH.h"
 #include "AttributeComponent.h"
 #include "GameplayCueInterface.h"
 #include "Abilities/GameplayAbility.h"
@@ -20,7 +20,7 @@ DEFINE_LOG_CATEGORY(LogAttributeComponent);
 
 int32 DebugGameplayCues = 0;
 static FAutoConsoleVariableRef CVarDebugGameplayCues(
-	TEXT("SkillSystem.DebugGameplayCues"),
+	TEXT("AbilitySystem.DebugGameplayCues"),
 	DebugGameplayCues,
 	TEXT("Enables Debugging for GameplayCue events"),
 	ECVF_Default
@@ -231,7 +231,7 @@ void UAttributeComponent::TEMP_ApplyActiveGameplayEffects()
 
 		ExecuteGameplayEffect(ActiveEffect.Spec, FModifierQualifier().IgnoreHandle(ActiveEffect.Handle));
 
-		SKILL_LOG(Log, TEXT("ActiveEffect[%d] %s - Duration: %.2f]"), idx, *ActiveEffect.Spec.ToSimpleString(), ActiveEffect.Spec.GetDuration());
+		ABILITY_LOG(Log, TEXT("ActiveEffect[%d] %s - Duration: %.2f]"), idx, *ActiveEffect.Spec.ToSimpleString(), ActiveEffect.Spec.GetDuration());
 	}
 }
 
@@ -338,7 +338,7 @@ FActiveGameplayEffectHandle UAttributeComponent::ApplyGameplayEffectSpecToSelf(c
 
 	if (Spec.GetPeriod() != UGameplayEffect::NO_PERIOD && Spec.TargetEffectSpecs.Num() > 0)
 	{
-		SKILL_LOG(Warning, TEXT("%s is periodic but also applies GameplayEffects to its target. GameplayEffects will only be applied once, not every period."), *Spec.Def->GetPathName());
+		ABILITY_LOG(Warning, TEXT("%s is periodic but also applies GameplayEffects to its target. GameplayEffects will only be applied once, not every period."), *Spec.Def->GetPathName());
 	}
 	// todo: this is ignoring the returned handles, should we put them into a TArray and return all of the handles?
 	for (const TSharedRef<FGameplayEffectSpec> TargetSpec : Spec.TargetEffectSpecs)
@@ -387,7 +387,7 @@ void UAttributeComponent::InvokeGameplayCueExecute(const FGameplayEffectSpec &Sp
 {
 	if (DebugGameplayCues)
 	{
-		SKILL_LOG(Warning, TEXT("InvokeGameplayCueExecute: %s"), *Spec.ToSimpleString());
+		ABILITY_LOG(Warning, TEXT("InvokeGameplayCueExecute: %s"), *Spec.ToSimpleString());
 	}
 
 
@@ -408,7 +408,7 @@ void UAttributeComponent::InvokeGameplayCueExecute(const FGameplayEffectSpec &Sp
 		if (DebugGameplayCues)
 		{
 			DrawDebugSphere(GetWorld(), Spec.InstigatorContext.HitResult->Location, 30.f, 32, FColor(255, 0, 0), true, 30.f);
-			SKILL_LOG(Warning, TEXT("   %s"), *CueInfo.GameplayCueTags.ToString());
+			ABILITY_LOG(Warning, TEXT("   %s"), *CueInfo.GameplayCueTags.ToString());
 		}
 	}
 }
@@ -601,14 +601,14 @@ int32 UAttributeComponent::GetNextPredictionKey()
 
 void UAttributeComponent::PrintAllGameplayEffects() const
 {
-	SKILL_LOG_SCOPE(TEXT("PrintAllGameplayEffects %s"), *GetName());
-	SKILL_LOG(Log, TEXT("Owner: %s"), *GetOwner()->GetName());
+	ABILITY_LOG_SCOPE(TEXT("PrintAllGameplayEffects %s"), *GetName());
+	ABILITY_LOG(Log, TEXT("Owner: %s"), *GetOwner()->GetName());
 	ActiveGameplayEffects.PrintAllGameplayEffects();
 }
 
 void FActiveGameplayEffectsContainer::PrintAllGameplayEffects() const
 {
-	SKILL_LOG_SCOPE(TEXT("ActiveGameplayEffects. Num: %d"), GameplayEffects.Num());
+	ABILITY_LOG_SCOPE(TEXT("ActiveGameplayEffects. Num: %d"), GameplayEffects.Num());
 	for (const FActiveGameplayEffect& Effect : GameplayEffects)
 	{
 		Effect.PrintAll();
@@ -617,23 +617,23 @@ void FActiveGameplayEffectsContainer::PrintAllGameplayEffects() const
 
 void FActiveGameplayEffect::PrintAll() const
 {
-	SKILL_LOG(Log, TEXT("Handle: %s"), *Handle.ToString());
-	SKILL_LOG(Log, TEXT("StartWorldTime: %.2f"), StartWorldTime);
+	ABILITY_LOG(Log, TEXT("Handle: %s"), *Handle.ToString());
+	ABILITY_LOG(Log, TEXT("StartWorldTime: %.2f"), StartWorldTime);
 	Spec.PrintAll();
 }
 
 void FGameplayEffectSpec::PrintAll() const
 {
-	SKILL_LOG_SCOPE(TEXT("GameplayEffectSpec"));
-	SKILL_LOG(Log, TEXT("Def: %s"), *Def->GetName());
+	ABILITY_LOG_SCOPE(TEXT("GameplayEffectSpec"));
+	ABILITY_LOG(Log, TEXT("Def: %s"), *Def->GetName());
 	
-	SKILL_LOG(Log, TEXT("Duration: "));
+	ABILITY_LOG(Log, TEXT("Duration: "));
 	Duration.PrintAll();
 
-	SKILL_LOG(Log, TEXT("Period:"));
+	ABILITY_LOG(Log, TEXT("Period:"));
 	Period.PrintAll();
 
-	SKILL_LOG(Log, TEXT("Modifiers:"));
+	ABILITY_LOG(Log, TEXT("Modifiers:"));
 	for (const FModifierSpec &Mod : Modifiers)
 	{
 		Mod.PrintAll();
@@ -642,14 +642,14 @@ void FGameplayEffectSpec::PrintAll() const
 
 void FModifierSpec::PrintAll() const
 {
-	SKILL_LOG_SCOPE(TEXT("ModifierSpec"));
-	SKILL_LOG(Log, TEXT("Attribute: %s"), *Info.Attribute.GetName());
-	SKILL_LOG(Log, TEXT("ModifierType: %s"), *EGameplayModToString(Info.ModifierType));
-	SKILL_LOG(Log, TEXT("ModifierOp: %s"), *EGameplayModOpToString(Info.ModifierOp));
-	SKILL_LOG(Log, TEXT("EffectType: %s"), *EGameplayModEffectToString(Info.EffectType));
-	SKILL_LOG(Log, TEXT("RequiredTags: %s"), *Info.RequiredTags.ToString());
-	SKILL_LOG(Log, TEXT("OwnedTags: %s"), *Info.OwnedTags.ToString());
-	SKILL_LOG(Log, TEXT("(Base) Magnitude: %s"), *Info.Magnitude.ToSimpleString());
+	ABILITY_LOG_SCOPE(TEXT("ModifierSpec"));
+	ABILITY_LOG(Log, TEXT("Attribute: %s"), *Info.Attribute.GetName());
+	ABILITY_LOG(Log, TEXT("ModifierType: %s"), *EGameplayModToString(Info.ModifierType));
+	ABILITY_LOG(Log, TEXT("ModifierOp: %s"), *EGameplayModOpToString(Info.ModifierOp));
+	ABILITY_LOG(Log, TEXT("EffectType: %s"), *EGameplayModEffectToString(Info.EffectType));
+	ABILITY_LOG(Log, TEXT("RequiredTags: %s"), *Info.RequiredTags.ToString());
+	ABILITY_LOG(Log, TEXT("OwnedTags: %s"), *Info.OwnedTags.ToString());
+	ABILITY_LOG(Log, TEXT("(Base) Magnitude: %s"), *Info.Magnitude.ToSimpleString());
 
 	Aggregator.PrintAll();
 }
@@ -658,17 +658,17 @@ void FAggregatorRef::PrintAll() const
 {
 	if (!WeakPtr.IsValid())
 	{
-		SKILL_LOG(Log, TEXT("Invalid AggregatorRef"));
+		ABILITY_LOG(Log, TEXT("Invalid AggregatorRef"));
 		return;
 	}
 
 	if (SharedPtr.IsValid())
 	{
-		SKILL_LOG(Log, TEXT("HardRef AggregatorRef"));
+		ABILITY_LOG(Log, TEXT("HardRef AggregatorRef"));
 	}
 	else
 	{
-		SKILL_LOG(Log, TEXT("SoftRef AggregatorRef"));
+		ABILITY_LOG(Log, TEXT("SoftRef AggregatorRef"));
 
 	}
 	
@@ -677,30 +677,30 @@ void FAggregatorRef::PrintAll() const
 
 void FAggregator::PrintAll() const
 {
-	SKILL_LOG_SCOPE(TEXT("FAggregator 0x%X"), this);
+	ABILITY_LOG_SCOPE(TEXT("FAggregator 0x%X"), this);
 
 #if SKILL_SYSTEM_AGGREGATOR_DEBUG
-	SKILL_LOG(Log, TEXT("DebugStr: %s"), *DebugString);
-	SKILL_LOG(Log, TEXT("Copies (of me): %d"), CopiesMade);
+	ABILITY_LOG(Log, TEXT("DebugStr: %s"), *DebugString);
+	ABILITY_LOG(Log, TEXT("Copies (of me): %d"), CopiesMade);
 #endif
 
 	if (Level.IsValid())
 	{
-		SKILL_LOG_SCOPE(TEXT("LevelInfo"));
+		ABILITY_LOG_SCOPE(TEXT("LevelInfo"));
 		Level->PrintAll();
 	}
 	else
 	{
-		SKILL_LOG(Log, TEXT("No Level Data"));
+		ABILITY_LOG(Log, TEXT("No Level Data"));
 	}
 
 	{
-		SKILL_LOG_SCOPE(TEXT("BaseData"));
+		ABILITY_LOG_SCOPE(TEXT("BaseData"));
 		BaseData.PrintAll();
 	}
 
 	{
-		SKILL_LOG_SCOPE(TEXT("CachedData"));
+		ABILITY_LOG_SCOPE(TEXT("CachedData"));
 		CachedData.PrintAll();
 	}
 	
@@ -708,7 +708,7 @@ void FAggregator::PrintAll() const
 	{
 		if (Mods[i].Num() > 0)
 		{
-			SKILL_LOG_SCOPE(TEXT("%s Mods"), *EGameplayModOpToString(i));
+			ABILITY_LOG_SCOPE(TEXT("%s Mods"), *EGameplayModOpToString(i));
 			for (const FAggregatorRef &Ref : Mods[i])
 			{
 				Ref.PrintAll();
@@ -719,20 +719,20 @@ void FAggregator::PrintAll() const
 
 void FGameplayModifierData::PrintAll() const
 {
-	SKILL_LOG(Log, TEXT("Magnitude: %s"), *Magnitude.ToSimpleString());
-	SKILL_LOG(Log, TEXT("Tags: %s"), *Tags.ToString());
+	ABILITY_LOG(Log, TEXT("Magnitude: %s"), *Magnitude.ToSimpleString());
+	ABILITY_LOG(Log, TEXT("Tags: %s"), *Tags.ToString());
 }
 
 void FGameplayModifierEvaluatedData::PrintAll() const
 {
-	SKILL_LOG(Log, TEXT("IsValid: %d"), IsValid); 
-	SKILL_LOG(Log, TEXT("Magnitude: %.2f"), Magnitude);
-	SKILL_LOG(Log, TEXT("Tags: %s"), *Tags.ToString());
+	ABILITY_LOG(Log, TEXT("IsValid: %d"), IsValid); 
+	ABILITY_LOG(Log, TEXT("Magnitude: %.2f"), Magnitude);
+	ABILITY_LOG(Log, TEXT("Tags: %s"), *Tags.ToString());
 }
 
 void FGameplayEffectLevelSpec::PrintAll() const
 {
-	SKILL_LOG(Log, TEXT("ConstantLevel: %.2f"), ConstantLevel);
+	ABILITY_LOG(Log, TEXT("ConstantLevel: %.2f"), ConstantLevel);
 }
 
 #undef LOCTEXT_NAMESPACE
diff --git a/Engine/Source/Runtime/GameplayAbilities/Private/AttributeComponent_Abilities.cpp b/Engine/Source/Runtime/GameplayAbilities/Private/AttributeComponent_Abilities.cpp
--- a/Engine/Source/Runtime/GameplayAbilities/Private/AttributeComponent_Abilities.cpp
+++ b/Engine/Source/Runtime/GameplayAbilities/Private/AttributeComponent_Abilities.cpp
@@ -1,7 +1,7 @@
 // Copyright 1998-2014 Epic Games, Inc. All Rights Reserved.
 // ActorComponent.cpp: Actor component implementation.
 
-#include "SkillSystemModulePrivatePCH.h"
+#include "AbilitySystemPrivatePCH.h"
 #include "AttributeComponent.h"
 #include "Abilities/GameplayAbility.h"
 
@@ -45,7 +45,7 @@ void UAttributeComponent::InitAbilityActorInfo()
 	if (!AbilityActorInfo.IsValid())
 	{
 		// Alloc (and init) a new actor info
-		AbilityActorInfo = TSharedPtr<FGameplayAbilityActorInfo>(USkillSystemGlobals::Get().AllocAbilityActorInfo(GetOwner()));
+		AbilityActorInfo = TSharedPtr<FGameplayAbilityActorInfo>(UAbilitySystemGlobals::Get().AllocAbilityActorInfo(GetOwner()));
 	}
 	else
 	{
@@ -134,7 +134,7 @@ void UAttributeComponent::CancelAbilitiesWithTags(const FGameplayTagContainer Ta
 #if !(UE_BUILD_SHIPPING || UE_BUILD_TEST)
 int32 DenyClientActivation = 0;
 static FAutoConsoleVariableRef CVarDenyClientActivation(
-	TEXT("SkillSystem.DenyClientActivations"),
+	TEXT("AbilitySystem.DenyClientActivations"),
 	DenyClientActivation,
 	TEXT("Make server deny the next X ability activations from clients. For testing misprediction."),
 	ECVF_Default
@@ -200,7 +200,7 @@ void UAttributeComponent::ClientActivateAbilityFailed_Implementation(UGameplayAb
 			int32 Key = PredictionDelegates[idx].Key;
 			if (Key == PredictionKey)
 			{
-				SKILL_LOG(Warning, TEXT("Failed ActivateAbility, clearing prediction data %d"), PredictionKey);
+				ABILITY_LOG(Warning, TEXT("Failed ActivateAbility, clearing prediction data %d"), PredictionKey);
 
 				PredictionDelegates[idx].Value.Broadcast();
 				PredictionDelegates.RemoveAt(idx);
diff --git a/Engine/Source/Runtime/GameplayAbilities/Private/AttributeSet.cpp b/Engine/Source/Runtime/GameplayAbilities/Private/AttributeSet.cpp
--- a/Engine/Source/Runtime/GameplayAbilities/Private/AttributeSet.cpp
+++ b/Engine/Source/Runtime/GameplayAbilities/Private/AttributeSet.cpp
@@ -1,6 +1,6 @@
 // Copyright 1998-2013 Epic Games, Inc. All Rights Reserved.
 
-#include "SkillSystemModulePrivatePCH.h"
+#include "AbilitySystemPrivatePCH.h"
 #include "AttributeSet.h"
 
 void FGameplayAttribute::SetNumericValueChecked(const float NewValue, class UAttributeSet* Dest) const
@@ -199,7 +199,7 @@ void FScalableFloat::FinalizeCurveData(const FGlobalCurveDataOverride *GlobalOve
 	}
 
 	// Look at global defaults
-	const UCurveTable * GlobalTable = IGameplayAbilitiesModule::Get().GetSkillSystemGlobals()->GetGlobalCurveTable();
+	const UCurveTable * GlobalTable = IGameplayAbilitiesModule::Get().GetAbilitySystemGlobals()->GetGlobalCurveTable();
 	if (GlobalTable)
 	{
 		FinalCurve = GlobalTable->FindCurve(Curve.RowName, ContextString, false);
@@ -207,7 +207,7 @@ void FScalableFloat::FinalizeCurveData(const FGlobalCurveDataOverride *GlobalOve
 
 	if (!FinalCurve)
 	{
-		SKILL_LOG(Warning, TEXT("Unable to find RowName: %s for FScalableFloat."), *Curve.RowName.ToString());
+		ABILITY_LOG(Warning, TEXT("Unable to find RowName: %s for FScalableFloat."), *Curve.RowName.ToString());
 	}
 }
 
diff --git a/Engine/Source/Runtime/GameplayAbilities/Private/GameplayAbilitiesModule.cpp b/Engine/Source/Runtime/GameplayAbilities/Private/GameplayAbilitiesModule.cpp
--- a/Engine/Source/Runtime/GameplayAbilities/Private/GameplayAbilitiesModule.cpp
+++ b/Engine/Source/Runtime/GameplayAbilities/Private/GameplayAbilitiesModule.cpp
@@ -1,7 +1,7 @@
 // Copyright 1998-2013 Epic Games, Inc. All Rights Reserved.
 
-#include "SkillSystemModulePrivatePCH.h"
-#include "SkillSystemGlobals.h"
+#include "AbilitySystemPrivatePCH.h"
+#include "AbilitySystemGlobals.h"
 
 
 class FGameplayAbilitiesModule : public IGameplayAbilitiesModule
@@ -11,13 +11,13 @@ class FGameplayAbilitiesModule : public IGameplayAbilitiesModule
 	virtual void ShutdownModule() OVERRIDE;
 	// End IModuleInterface
 
-	virtual USkillSystemGlobals* GetSkillSystemGlobals()
+	virtual UAbilitySystemGlobals* GetAbilitySystemGlobals()
 	{
-		check(SkillSystemGlobals);
-		return SkillSystemGlobals;
+		check(AbilitySystemGlobals);
+		return AbilitySystemGlobals;
 	}
 
-	USkillSystemGlobals *SkillSystemGlobals;
+	UAbilitySystemGlobals *AbilitySystemGlobals;
 
 private:
 	
@@ -28,7 +28,7 @@ IMPLEMENT_MODULE(FGameplayAbilitiesModule, GameplayAbilities)
 void FGameplayAbilitiesModule::StartupModule()
 {	
 	// This code will execute after your module is loaded into memory (but after global variables are initialized, of course.)
-	SkillSystemGlobals = ConstructObject<USkillSystemGlobals>(USkillSystemGlobals::StaticClass(), GetTransientPackage(), NAME_None, RF_RootSet);
+	AbilitySystemGlobals = ConstructObject<UAbilitySystemGlobals>(UAbilitySystemGlobals::StaticClass(), GetTransientPackage(), NAME_None, RF_RootSet);
 }
 
 void FGameplayAbilitiesModule::ShutdownModule()
@@ -36,5 +36,5 @@ void FGameplayAbilitiesModule::ShutdownModule()
 	// This function may be called during shutdown to clean up your module.  For modules that support dynamic reloading,
 	// we call this function before unloading the module.
 
-	SkillSystemGlobals = NULL;
+	AbilitySystemGlobals = NULL;
 }
diff --git a/Engine/Source/Runtime/GameplayAbilities/Private/GameplayCueView.cpp b/Engine/Source/Runtime/GameplayAbilities/Private/GameplayCueView.cpp
--- a/Engine/Source/Runtime/GameplayAbilities/Private/GameplayCueView.cpp
+++ b/Engine/Source/Runtime/GameplayAbilities/Private/GameplayCueView.cpp
@@ -2,7 +2,7 @@
 
 #pragma once
 
-#include "SkillSystemModulePrivatePCH.h"
+#include "AbilitySystemPrivatePCH.h"
 #include "GameplayEffect.h"
 #include "GameplayCueView.h"
 #include "ParticleDefinitions.h"
@@ -199,7 +199,7 @@ TSharedPtr<FGameplayCueViewEffects> FGameplayCueViewInfo::SpawnViewEffects(AActo
 					SpawnedEffects->ParticleSystemComponent->EmitterInstances[EmitterIndx]->CurrentLODLevel->RequiredModule &&
 					SpawnedEffects->ParticleSystemComponent->EmitterInstances[EmitterIndx]->CurrentLODLevel->RequiredModule->EmitterLoops == 0)
 				{
-					SKILL_LOG(Warning, TEXT("%s - particle system has a looping emitter. This should not be used in a executed GameplayCue!"), *SpawnedEffects->ParticleSystemComponent->GetName());
+					ABILITY_LOG(Warning, TEXT("%s - particle system has a looping emitter. This should not be used in a executed GameplayCue!"), *SpawnedEffects->ParticleSystemComponent->GetName());
 					break;
 				}
 			}
diff --git a/Engine/Source/Runtime/GameplayAbilities/Private/GameplayEffect.cpp b/Engine/Source/Runtime/GameplayAbilities/Private/GameplayEffect.cpp
--- a/Engine/Source/Runtime/GameplayAbilities/Private/GameplayEffect.cpp
+++ b/Engine/Source/Runtime/GameplayAbilities/Private/GameplayEffect.cpp
@@ -1,6 +1,6 @@
 // Copyright 1998-2013 Epic Games, Inc. All Rights Reserved.
 
-#include "SkillSystemModulePrivatePCH.h"
+#include "AbilitySystemPrivatePCH.h"
 #include "GameplayEffect.h"
 #include "GameplayEffectExtension.h"
 #include "GameplayEffectStackingExtension.h"
@@ -68,30 +68,30 @@ void UGameplayEffect::ValidateStacking()
 
 		if (StackedAttribName == NAME_None)
 		{
-			SKILL_LOG(Warning, TEXT("%s has a stacking rule but does not have an attribute to apply it to. Removing stacking rule."), *GetPathName());
+			ABILITY_LOG(Warning, TEXT("%s has a stacking rule but does not have an attribute to apply it to. Removing stacking rule."), *GetPathName());
 			StackingPolicy = EGameplayEffectStackingPolicy::Unlimited;
 		}
 
-		UDataTable* DataTable = IGameplayAbilitiesModule::Get().GetSkillSystemGlobals()->GetGlobalAttributeDataTable();
+		UDataTable* DataTable = IGameplayAbilitiesModule::Get().GetAbilitySystemGlobals()->GetGlobalAttributeDataTable();
 
 		if (DataTable && DataTable->RowMap.Contains(StackedAttribName))
 		{
 			FAttributeMetaData* Row = (FAttributeMetaData*)DataTable->RowMap[StackedAttribName];
 			if (Row && Row->bCanStack == false)
 			{
-				SKILL_LOG(Warning, TEXT("%s has a stacking rule but modifies attribute %s. %s is not allowed to stack. Removing stacking rule."), *GetPathName(), *Modifiers[0].Attribute.GetName(), *Modifiers[0].Attribute.GetName());
+				ABILITY_LOG(Warning, TEXT("%s has a stacking rule but modifies attribute %s. %s is not allowed to stack. Removing stacking rule."), *GetPathName(), *Modifiers[0].Attribute.GetName(), *Modifiers[0].Attribute.GetName());
 				StackingPolicy = EGameplayEffectStackingPolicy::Unlimited;
 			}
 		}
 		else
 		{
-			SKILL_LOG(Warning, TEXT("%s has a stacking rule but modifies attribute %s. %s was not found in the global attribute data table. Removing stacking rule."), *GetPathName(), *Modifiers[0].Attribute.GetName(), *Modifiers[0].Attribute.GetName());
+			ABILITY_LOG(Warning, TEXT("%s has a stacking rule but modifies attribute %s. %s was not found in the global attribute data table. Removing stacking rule."), *GetPathName(), *Modifiers[0].Attribute.GetName(), *Modifiers[0].Attribute.GetName());
 			StackingPolicy = EGameplayEffectStackingPolicy::Unlimited;
 		}
 
 		if (Modifiers.Num() == 0)
 		{
-			SKILL_LOG(Warning, TEXT("%s has a stacking rule but does not have any modifiers to apply it to. Removing stacking rule."), *GetPathName());
+			ABILITY_LOG(Warning, TEXT("%s has a stacking rule but does not have any modifiers to apply it to. Removing stacking rule."), *GetPathName());
 			StackingPolicy = EGameplayEffectStackingPolicy::Unlimited;
 		}
 		else
@@ -105,12 +105,12 @@ void UGameplayEffect::ValidateStacking()
 
 					if (AttributeName != StackedAttribName)
 					{
-						SKILL_LOG(Warning, TEXT("%s has a stacking rule and modifies attribute %s but one of the modifiers modifies attribute %s. Changes to %s may not behave in a predictable fashion."), *GetPathName(), *StackedAttribName.ToString(), *AttributeName.ToString(), *AttributeName.ToString());
+						ABILITY_LOG(Warning, TEXT("%s has a stacking rule and modifies attribute %s but one of the modifiers modifies attribute %s. Changes to %s may not behave in a predictable fashion."), *GetPathName(), *StackedAttribName.ToString(), *AttributeName.ToString(), *AttributeName.ToString());
 					}
 				}
 				else if (Modifiers[Idx].ModifierType == EGameplayMod::ActiveGE || Modifiers[Idx].ModifierType == EGameplayMod::IncomingGE || Modifiers[Idx].ModifierType == EGameplayMod::OutgoingGE)
 				{
-					SKILL_LOG(Warning, TEXT("%s has a stacking rule and modifies attribute %s but one of its modifiers modifies gameplay effects. Changes to other gameplay effects may not behave in a predictable fashion."), *GetPathName(), *StackedAttribName.ToString());
+					ABILITY_LOG(Warning, TEXT("%s has a stacking rule and modifies attribute %s but one of its modifiers modifies gameplay effects. Changes to other gameplay effects may not behave in a predictable fashion."), *GetPathName(), *StackedAttribName.ToString());
 				}
 			}
 		}
@@ -180,7 +180,7 @@ void FGameplayEffectSpec::MakeUnique()
 
 int32 FGameplayEffectSpec::ApplyModifiersFrom(FGameplayEffectSpec &InSpec, const FModifierQualifier &QualifierContext)
 {
-	SKILL_LOG_SCOPE(TEXT("FGameplayEffectSpec::ApplyModifiersFrom %s. InSpec: %s"), *this->ToSimpleString(), *InSpec.ToSimpleString());
+	ABILITY_LOG_SCOPE(TEXT("FGameplayEffectSpec::ApplyModifiersFrom %s. InSpec: %s"), *this->ToSimpleString(), *InSpec.ToSimpleString());
 
 	int32 NumApplied = 0;
 	bool ShouldSnapshot = InSpec.ShouldApplyAsSnapshot(QualifierContext);
@@ -421,7 +421,7 @@ void FModifierSpec::ApplyModTo(FModifierSpec &Other, bool TakeSnapshot) const
 
 void FModifierSpec::ExecuteModOn(FModifierSpec &Other) const
 {
-	SKILL_LOG_SCOPE(TEXT("Executing %s on %s"), *ToSimpleString(), *Other.ToSimpleString() );
+	ABILITY_LOG_SCOPE(TEXT("Executing %s on %s"), *ToSimpleString(), *Other.ToSimpleString() );
 	Other.Aggregator.Get()->ExecuteModAggr(this->Info.ModifierOp, this->Aggregator);
 }
 
@@ -506,7 +506,7 @@ bool FGameplayEffectInstigatorContext::NetSerialize(FArchive& Ar, class UPackage
 
 void FAggregatorRef::MakeUnique()
 {
-	SKILL_LOG(Log, TEXT("MakeUnique %s"), *ToString());
+	ABILITY_LOG(Log, TEXT("MakeUnique %s"), *ToString());
 
 	// Make a hard ref copy of our FAggregator
 	MakeHardRef();
@@ -519,7 +519,7 @@ void FAggregatorRef::MakeUnique()
 
 void FAggregatorRef::MakeUniqueDeep()
 {
-	SKILL_LOG_SCOPE(TEXT("MakeUniqueDeep %s"), *ToString());
+	ABILITY_LOG_SCOPE(TEXT("MakeUniqueDeep %s"), *ToString());
 
 	// Make a hard ref copy of our FAggregator
 	MakeHardRef();
@@ -651,7 +651,7 @@ FAggregator::~FAggregator()
 	{
 		if (WeakPtr.IsValid())
 		{
-			SKILL_LOG(Log, TEXT("%s Marking Dependant %s Dirty on Destroy"), *ToSimpleString(), *WeakPtr.Pin()->ToSimpleString());
+			ABILITY_LOG(Log, TEXT("%s Marking Dependant %s Dirty on Destroy"), *ToSimpleString(), *WeakPtr.Pin()->ToSimpleString());
 
 			WeakPtr.Pin()->MarkDirty();
 #if SKILL_SYSTEM_AGGREGATOR_DEBUG
@@ -713,7 +713,7 @@ void FAggregator::ExecuteMod(EGameplayModOp::Type ModType, const FGameplayModifi
 	// Am I allowed to even be executed on? If my base data scales then all I can do is apply this to myself
 	if (!BaseData.Magnitude.IsStatic())
 	{
-		SKILL_LOG(Log, TEXT("Treating Execute as an Apply since our Level is not static!"));
+		ABILITY_LOG(Log, TEXT("Treating Execute as an Apply since our Level is not static!"));
 		FAggregatorRef NewAggregator(new FAggregator(EvaluatedData, SKILL_AGG_DEBUG(TEXT("ExecutedMod"))));
 		ApplyMod(ModType, NewAggregator, false);
 		return;
@@ -753,7 +753,7 @@ void FAggregator::ExecuteMod(EGameplayModOp::Type ModType, const FGameplayModifi
 		}
 	}
 
-	SKILL_LOG(Log, TEXT("ExecuteMod: %s new BaseData.Magnitude: %s"), *ToSimpleString(), *BaseData.Magnitude.ToSimpleString());
+	ABILITY_LOG(Log, TEXT("ExecuteMod: %s new BaseData.Magnitude: %s"), *ToSimpleString(), *BaseData.Magnitude.ToSimpleString());
 	MarkDirty();
 }
 
@@ -762,7 +762,7 @@ void FAggregator::AddDependantAggregator(TWeakPtr<FAggregator> InDependant)
 	check(InDependant.IsValid());
 	check(!Dependants.Contains(InDependant));
 
-	SKILL_LOG(Log, TEXT("AddDependantAggregator: %s is a dependant of %s"), *InDependant.Pin()->ToSimpleString(), *ToSimpleString());
+	ABILITY_LOG(Log, TEXT("AddDependantAggregator: %s is a dependant of %s"), *InDependant.Pin()->ToSimpleString(), *ToSimpleString());
 	Dependants.Add(InDependant);
 }
 
@@ -787,7 +787,7 @@ void FAggregator::RegisterLevelDependancies()
 
 const FGameplayModifierEvaluatedData& FAggregator::Evaluate() const
 {
-	SKILL_LOG_SCOPE(TEXT("Aggregator Evaluate %s"), *ToSimpleString());
+	ABILITY_LOG_SCOPE(TEXT("Aggregator Evaluate %s"), *ToSimpleString());
 
 	if (!CachedData.IsValid)
 	{
@@ -796,7 +796,7 @@ const FGameplayModifierEvaluatedData& FAggregator::Evaluate() const
 		// ------------------------------------------------------------------------
 		for (const FAggregatorRef &Agg : Mods[EGameplayModOp::Override])
 		{
-			SKILL_LOG_SCOPE(TEXT("EGameplayModOp::Override"));
+			ABILITY_LOG_SCOPE(TEXT("EGameplayModOp::Override"));
 			if (Agg.IsValid())
 			{
 				CachedData = Agg.Get()->Evaluate();
@@ -826,7 +826,7 @@ const FGameplayModifierEvaluatedData& FAggregator::Evaluate() const
 		// We need to calculate num of mods anyways to do ModLIst.Reserve.
 		if (TotalModCount <= 0)
 		{
-			SKILL_LOG(Log, TEXT("Final Magnitude: %.2f"), CachedData.Magnitude);
+			ABILITY_LOG(Log, TEXT("Final Magnitude: %.2f"), CachedData.Magnitude);
 			return CachedData;
 		}
 
@@ -927,11 +927,11 @@ const FGameplayModifierEvaluatedData& FAggregator::Evaluate() const
 				if (ModIgnoreTags.Num() > 0 && NewGiveTags.Num() > 0 && ModIgnoreTags.MatchesAny(NewGiveTags, false))
 				{
 					// Pass problem!
-					SKILL_LOG(Warning, TEXT("Tagging conflicts during Aggregate! Use Stacking rules to avoid this"));
-					SKILL_LOG(Warning, TEXT("   While Evaluating: %s"), *ToSimpleString());
-					SKILL_LOG(Warning, TEXT("   Applying Mod: %s"), *Agg->ToSimpleString());
-					SKILL_LOG(Warning, TEXT("   ModIgnoreTags: %s"), *ModIgnoreTags.ToString() );
-					SKILL_LOG(Warning, TEXT("   NewGiveTags: %s"), *NewGiveTags.ToString());
+					ABILITY_LOG(Warning, TEXT("Tagging conflicts during Aggregate! Use Stacking rules to avoid this"));
+					ABILITY_LOG(Warning, TEXT("   While Evaluating: %s"), *ToSimpleString());
+					ABILITY_LOG(Warning, TEXT("   Applying Mod: %s"), *Agg->ToSimpleString());
+					ABILITY_LOG(Warning, TEXT("   ModIgnoreTags: %s"), *ModIgnoreTags.ToString() );
+					ABILITY_LOG(Warning, TEXT("   NewGiveTags: %s"), *NewGiveTags.ToString());
 					continue;
 				}
 
@@ -947,11 +947,11 @@ const FGameplayModifierEvaluatedData& FAggregator::Evaluate() const
 				if (NewIgnoreTags.Num() > 0 && ModEvaluatedData.Tags.Num() > 0 && NewIgnoreTags.MatchesAny(ModEvaluatedData.Tags, false))
 				{
 					// Pass problem!
-					SKILL_LOG(Warning, TEXT("Tagging conflicts during Aggregate! Use Stacking rules to avoid this"));
-					SKILL_LOG(Warning, TEXT("   While Evaluating: %s"), *ToSimpleString());
-					SKILL_LOG(Warning, TEXT("   Applying Mod: %s"), *Agg->ToSimpleString());
-					SKILL_LOG(Warning, TEXT("   Mod Tags: %s"), *ModEvaluatedData.Tags.ToString());
-					SKILL_LOG(Warning, TEXT("   NewIgnoreTags: %s"), *NewIgnoreTags.ToString());
+					ABILITY_LOG(Warning, TEXT("Tagging conflicts during Aggregate! Use Stacking rules to avoid this"));
+					ABILITY_LOG(Warning, TEXT("   While Evaluating: %s"), *ToSimpleString());
+					ABILITY_LOG(Warning, TEXT("   Applying Mod: %s"), *Agg->ToSimpleString());
+					ABILITY_LOG(Warning, TEXT("   Mod Tags: %s"), *ModEvaluatedData.Tags.ToString());
+					ABILITY_LOG(Warning, TEXT("   NewIgnoreTags: %s"), *NewIgnoreTags.ToString());
 					continue;
 				}
 
@@ -974,12 +974,12 @@ const FGameplayModifierEvaluatedData& FAggregator::Evaluate() const
 
 		CachedData.Magnitude = ((CachedData.Magnitude + OpAggregation[EGameplayModOp::Additive]) * OpAggregation[EGameplayModOp::Multiplicitive]) / Division;
 
-		SKILL_LOG(Log, TEXT("Final Magnitude: %.2f"), CachedData.Magnitude);
+		ABILITY_LOG(Log, TEXT("Final Magnitude: %.2f"), CachedData.Magnitude);
 		
 	}
 	else
 	{
-		SKILL_LOG(Log, TEXT("CachedData was valid. Magnitude: %.2f"), CachedData.Magnitude);
+		ABILITY_LOG(Log, TEXT("CachedData was valid. Magnitude: %.2f"), CachedData.Magnitude);
 	}
 	return CachedData;
 }
@@ -1021,7 +1021,7 @@ FAggregator & FAggregator::MarkDirty()
 		TWeakPtr<FAggregator> WeakPtr = Dependants[i];
 		if (WeakPtr.IsValid())
 		{
-			SKILL_LOG(Log, TEXT("%s Marking Dependant %s Dirty (from ::MarkDirty())"), *ToSimpleString(), *WeakPtr.Pin()->ToSimpleString());
+			ABILITY_LOG(Log, TEXT("%s Marking Dependant %s Dirty (from ::MarkDirty())"), *ToSimpleString(), *WeakPtr.Pin()->ToSimpleString());
 #if SKILL_SYSTEM_AGGREGATOR_DEBUG
 			FAggregator::AllocationStats.DependantsUpdated++;
 #endif
@@ -1146,7 +1146,7 @@ bool FActiveGameplayEffectsContainer::ApplyActiveEffectsTo(OUT FGameplayEffectSp
 {
 	FActiveGameplayEffectHandle().IsValid();
 
-	SKILL_LOG_SCOPE(TEXT("ApplyActiveEffectsTo: %s %s"), *Spec.ToSimpleString(), *QualifierContext.ToString());
+	ABILITY_LOG_SCOPE(TEXT("ApplyActiveEffectsTo: %s %s"), *Spec.ToSimpleString(), *QualifierContext.ToString());
 	for (FActiveGameplayEffect & ActiveEffect : GameplayEffects)
 	{
 		// We dont want to use FModifierQualifier::TestTarget here, since we aren't the 'target'. We are applying stuff to Spec which will be applied to a target.
@@ -1179,7 +1179,7 @@ void FActiveGameplayEffectsContainer::ApplySpecToActiveEffectsAndAttributes(FGam
 		{
 			// Todo: Tag/application checks here
 
-			SKILL_LOG_SCOPE(TEXT("Applying Attribute Mod %s to property"), *Mod.ToSimpleString());
+			ABILITY_LOG_SCOPE(TEXT("Applying Attribute Mod %s to property"), *Mod.ToSimpleString());
 
 			FAggregator * Aggregator = FindOrCreateAttributeAggregator(Mod.Info.Attribute).Get();
 
@@ -1190,7 +1190,7 @@ void FActiveGameplayEffectsContainer::ApplySpecToActiveEffectsAndAttributes(FGam
 
 		if (Mod.Info.ModifierType == EGameplayMod::ActiveGE)
 		{
-			SKILL_LOG_SCOPE(TEXT("Applying Mod %s to ActiveEffects"),  *Mod.ToSimpleString());
+			ABILITY_LOG_SCOPE(TEXT("Applying Mod %s to ActiveEffects"),  *Mod.ToSimpleString());
 
 			bool bUpdateStackedEffects = false;
 
@@ -1274,14 +1274,14 @@ void FActiveGameplayEffectsContainer::ExecuteActiveEffectsFrom(const FGameplayEf
 			if (AttributeSet == NULL)
 			{
 				// Our owner doesn't have this attribute, so we can't do anything
-				SKILL_LOG(Log, TEXT("%s does not have attribute %s. Skipping modifier"), *Owner->GetPathName(), *Mod.Info.Attribute.GetName());
+				ABILITY_LOG(Log, TEXT("%s does not have attribute %s. Skipping modifier"), *Owner->GetPathName(), *Mod.Info.Attribute.GetName());
 				continue;
 			}
 
 			// Todo: Tags/application checks here - make sure we can still apply
 			InvokeGameplayCueExecute = true;
 
-			SKILL_LOG_SCOPE(TEXT("Executing Attribute Mod %s"), *Mod.ToSimpleString());
+			ABILITY_LOG_SCOPE(TEXT("Executing Attribute Mod %s"), *Mod.ToSimpleString());
 
 			// First, evaluate all of our data 
 
@@ -1299,7 +1299,7 @@ void FActiveGameplayEffectsContainer::ExecuteActiveEffectsFrom(const FGameplayEf
 			FAggregatorRef *RefPtr = OngoingPropertyEffects.Find(Mod.Info.Attribute);
 			if (RefPtr)
 			{
-				SKILL_LOG(Log, TEXT("Property %s has active mods. Adding to Aggregator."), *Mod.Info.Attribute.GetName());
+				ABILITY_LOG(Log, TEXT("Property %s has active mods. Adding to Aggregator."), *Mod.Info.Attribute.GetName());
 				RefPtr->Get()->ExecuteMod(Mod.Info.ModifierOp, EvaluatedData);
 			}
 			else
@@ -1312,7 +1312,7 @@ void FActiveGameplayEffectsContainer::ExecuteActiveEffectsFrom(const FGameplayEf
 				
 				const float NewPropertyValue = Aggregator.Evaluate().Magnitude;
 
-				SKILL_LOG(Log, TEXT("Property %s new value is: %.2f [was: %.2f]"), *Mod.Info.Attribute.GetName(), NewPropertyValue, CurrentValueOfProperty);
+				ABILITY_LOG(Log, TEXT("Property %s new value is: %.2f [was: %.2f]"), *Mod.Info.Attribute.GetName(), NewPropertyValue, CurrentValueOfProperty);
 				Owner->SetNumericAttribute(Mod.Info.Attribute, NewPropertyValue);
 			}
 
@@ -1358,7 +1358,7 @@ void FActiveGameplayEffectsContainer::ExecuteActiveEffectsFrom(const FGameplayEf
 	{
 		// TODO: check replication policy. Right now we will replicate every execute via a multicast RPC
 
-		SKILL_LOG(Log, TEXT("Invoking Execute GameplayCue for %s"), *Spec.ToSimpleString() );
+		ABILITY_LOG(Log, TEXT("Invoking Execute GameplayCue for %s"), *Spec.ToSimpleString() );
 		Owner->NetMulticast_InvokeGameplayCueExecuted(Spec);
 	}
 
@@ -1395,7 +1395,7 @@ FAggregatorRef & FActiveGameplayEffectsContainer::FindOrCreateAttributeAggregato
 
 	// Create a new aggregator for this attribute.
 	float CurrentValueOfProperty = Owner->GetNumericAttribute(Attribute);
-	SKILL_LOG(Log, TEXT("Creating new entry in OngoingPropertyEffect map for AddDependancyToAttribute. CurrentValue: %.2f"), CurrentValueOfProperty);
+	ABILITY_LOG(Log, TEXT("Creating new entry in OngoingPropertyEffect map for AddDependancyToAttribute. CurrentValue: %.2f"), CurrentValueOfProperty);
 
 	FAggregator *NewPropertyAggregator = new FAggregator(FGameplayModifierEvaluatedData(CurrentValueOfProperty), SKILL_AGG_DEBUG(TEXT("Attribute %s Aggregator"), *Attribute.GetName()));
 
@@ -1415,7 +1415,7 @@ float FActiveGameplayEffectsContainer::GetGameplayEffectDuration(FActiveGameplay
 		}
 	}
 	
-	SKILL_LOG(Warning, TEXT("GetGameplayEffectDuration called with invalid Handle: %s"), *Handle.ToString() );
+	ABILITY_LOG(Warning, TEXT("GetGameplayEffectDuration called with invalid Handle: %s"), *Handle.ToString() );
 	return UGameplayEffect::INFINITE_DURATION;
 }
 
@@ -1453,7 +1453,7 @@ float FActiveGameplayEffectsContainer::GetGameplayEffectMagnitude(FActiveGamepla
 		}
 	}
 
-	SKILL_LOG(Warning, TEXT("GetGameplayEffectMagnitude called with invalid Handle: %s"), *Handle.ToString());
+	ABILITY_LOG(Warning, TEXT("GetGameplayEffectMagnitude called with invalid Handle: %s"), *Handle.ToString());
 	return -1.f;
 }
 
@@ -1494,7 +1494,7 @@ float FActiveGameplayEffectsContainer::GetGameplayEffectMagnitudeByTag(FActiveGa
 		}
 	}
 
-	SKILL_LOG(Warning, TEXT("GetGameplayEffectMagnitude called with invalid Handle: %s"), *Handle.ToString());
+	ABILITY_LOG(Warning, TEXT("GetGameplayEffectMagnitude called with invalid Handle: %s"), *Handle.ToString());
 	return -1.f;
 }
 
@@ -1502,12 +1502,12 @@ void FActiveGameplayEffectsContainer::OnPropertyAggregatorDirty(FAggregator* Agg
 {
 	check(OngoingPropertyEffects.FindChecked(Attribute).Get() == Aggregator);
 
-	SKILL_LOG_SCOPE(TEXT("FActiveGameplayEffectsContainer::OnPropertyAggregatorDirty"));
+	ABILITY_LOG_SCOPE(TEXT("FActiveGameplayEffectsContainer::OnPropertyAggregatorDirty"));
 
 	// Immediately calculate the newest value of the property			
 	float NewPropertyValue = Aggregator->Evaluate().Magnitude;
 
-	SKILL_LOG(Log, TEXT("Property %s new value is: %.2f"), *Attribute.GetName(), NewPropertyValue);
+	ABILITY_LOG(Log, TEXT("Property %s new value is: %.2f"), *Attribute.GetName(), NewPropertyValue);
 	
 	Owner->SetNumericAttribute(Attribute, NewPropertyValue);
 }
@@ -1597,7 +1597,7 @@ bool FActiveGameplayEffectsContainer::RemoveActiveGameplayEffect(FActiveGameplay
 		}
 	}
 
-	SKILL_LOG(Warning, TEXT("RemoveActiveGameplayEffect called with invalid Handle: %s"), *Handle.ToString());
+	ABILITY_LOG(Warning, TEXT("RemoveActiveGameplayEffect called with invalid Handle: %s"), *Handle.ToString());
 	return false;
 }
 
@@ -1630,7 +1630,7 @@ bool FActiveGameplayEffectsContainer::InternalRemoveActiveGameplayEffect(int32 I
 		return true;
 	}
 
-	SKILL_LOG(Warning, TEXT("InternalRemoveActiveGameplayEffect called with invalid index: %d"), Idx);
+	ABILITY_LOG(Warning, TEXT("InternalRemoveActiveGameplayEffect called with invalid index: %d"), Idx);
 	return false;
 }
 
@@ -1769,7 +1769,7 @@ void FActiveGameplayEffectsContainer::RecalculateStacking()
 						}
 						default:
 							// we shouldn't ever be here
-							SKILL_LOG(Warning, TEXT("%s uses unhandled stacking rule %s."), *Effect.Spec.Def->GetPathName(), *EGameplayEffectStackingPolicyToString(Effect.Spec.GetStackingType()));
+							ABILITY_LOG(Warning, TEXT("%s uses unhandled stacking rule %s."), *Effect.Spec.Def->GetPathName(), *EGameplayEffectStackingPolicyToString(Effect.Spec.GetStackingType()));
 							break;
 						};
 
diff --git a/Engine/Source/Runtime/GameplayAbilities/Private/GameplayEffectExtension_LifestealTest.cpp b/Engine/Source/Runtime/GameplayAbilities/Private/GameplayEffectExtension_LifestealTest.cpp
--- a/Engine/Source/Runtime/GameplayAbilities/Private/GameplayEffectExtension_LifestealTest.cpp
+++ b/Engine/Source/Runtime/GameplayAbilities/Private/GameplayEffectExtension_LifestealTest.cpp
@@ -1,8 +1,8 @@
 // Copyright 1998-2013 Epic Games, Inc. All Rights Reserved.
 
-#include "SkillSystemModulePrivatePCH.h"
+#include "AbilitySystemPrivatePCH.h"
 #include "GameplayEffectExtension_LifestealTest.h"
-#include "SkillSystemTestAttributeSet.h"
+#include "AbilitySystemTestAttributeSet.h"
 #include "GameplayTagsModule.h"
 #include "AttributeComponent.h"
 
@@ -33,7 +33,7 @@ void UGameplayEffectExtension_LifestealTest::PostGameplayEffectExecute(const FGa
 		UGameplayEffect * LocalHealthRestore = HealthRestoreGameplayEffect;
 		if (!LocalHealthRestore)
 		{
-			UProperty *HealthProperty = FindFieldChecked<UProperty>(USkillSystemTestAttributeSet::StaticClass(), GET_MEMBER_NAME_CHECKED(USkillSystemTestAttributeSet, Health));
+			UProperty *HealthProperty = FindFieldChecked<UProperty>(UAbilitySystemTestAttributeSet::StaticClass(), GET_MEMBER_NAME_CHECKED(UAbilitySystemTestAttributeSet, Health));
 
 			// Since this is a test class and we don't want to tie it any actual content assets, just construct a GameplayEffect here.
 			LocalHealthRestore = Cast<UGameplayEffect>(StaticConstructObject(UGameplayEffect::StaticClass(), GetTransientPackage(), FName(TEXT("LifestealHealthRestore"))));
diff --git a/Engine/Source/Runtime/GameplayAbilities/Private/GameplayEffectExtension_ShieldTest.cpp b/Engine/Source/Runtime/GameplayAbilities/Private/GameplayEffectExtension_ShieldTest.cpp
--- a/Engine/Source/Runtime/GameplayAbilities/Private/GameplayEffectExtension_ShieldTest.cpp
+++ b/Engine/Source/Runtime/GameplayAbilities/Private/GameplayEffectExtension_ShieldTest.cpp
@@ -1,10 +1,10 @@
 // Copyright 1998-2013 Epic Games, Inc. All Rights Reserved.
 
-#include "SkillSystemModulePrivatePCH.h"
+#include "AbilitySystemPrivatePCH.h"
 #include "GameplayEffectExtension_ShieldTest.h"
 #include "GameplayTagsModule.h"
 #include "AttributeComponent.h"
-#include "SkillSystemTestAttributeSet.h"
+#include "AbilitySystemTestAttributeSet.h"
 
 UGameplayEffectExtension_ShieldTest::UGameplayEffectExtension_ShieldTest(const class FPostConstructInitializeProperties& PCIP)
 : Super(PCIP)
@@ -50,7 +50,7 @@ void UGameplayEffectExtension_ShieldTest::PreGameplayEffectExecute(const FGamepl
 
 		if (!LocalShieldRemoval)
 		{
-			UProperty *HealthProperty = FindFieldChecked<UProperty>(USkillSystemTestAttributeSet::StaticClass(), GET_MEMBER_NAME_CHECKED(USkillSystemTestAttributeSet, Health));
+			UProperty *HealthProperty = FindFieldChecked<UProperty>(UAbilitySystemTestAttributeSet::StaticClass(), GET_MEMBER_NAME_CHECKED(UAbilitySystemTestAttributeSet, Health));
 
 			// Since this is a test class and we don't want to tie it any actual content assets, just construct a GameplayEffect here.
 			LocalShieldRemoval = Cast<UGameplayEffect>(StaticConstructObject(UGameplayEffect::StaticClass(), GetTransientPackage(), FName(TEXT("ShieldAbsorbRemoval"))));
