Mass rename SkillSystem -> AbilitySystem

diff --git a/Engine/Source/Runtime/GameplayAbilities/Private/GameplayEffectTests.cpp b/Engine/Source/Runtime/GameplayAbilities/Private/GameplayEffectTests.cpp
--- a/Engine/Source/Runtime/GameplayAbilities/Private/GameplayEffectTests.cpp
+++ b/Engine/Source/Runtime/GameplayAbilities/Private/GameplayEffectTests.cpp
@@ -1,8 +1,8 @@
 // Copyright 1998-2013 Epic Games, Inc. All Rights Reserved.
 
-#include "SkillSystemModulePrivatePCH.h"
-#include "SkillSystemTestPawn.h"
-#include "SkillSystemTestAttributeSet.h"
+#include "AbilitySystemPrivatePCH.h"
+#include "AbilitySystemTestPawn.h"
+#include "AbilitySystemTestAttributeSet.h"
 #include "GameplayEffect.h"
 #include "AttributeSet.h"
 #include "GameplayTagsModule.h"
@@ -12,7 +12,7 @@
 #include "GameplayEffectStackingExtension_CappedNumberTest.h"
 #include "GameplayEffectStackingExtension_DiminishingReturnsTest.h"
 
-IMPLEMENT_SIMPLE_AUTOMATION_TEST(FGameplayEffectsTest, "SkillSystem.GameplayEffects", EAutomationTestFlags::ATF_Editor)
+IMPLEMENT_SIMPLE_AUTOMATION_TEST(FGameplayEffectsTest, "AbilitySystem.GameplayEffects", EAutomationTestFlags::ATF_Editor)
 
 #define SKILL_TEST_TEXT( Format, ... ) FString::Printf(TEXT("%s - %d: %s"), TEXT(__FILE__) , __LINE__ , *FString::Printf(TEXT(Format), ##__VA_ARGS__) )
 
@@ -105,18 +105,18 @@ bool GameplayEffectsTest_InstantDamage(UWorld *World, FAutomationTestBase * Test
 	const float StartHealth = 100.f;
 	const float DamageValue = -5.f;
 
-	ASkillSystemTestPawn *SourceActor = World->SpawnActor<ASkillSystemTestPawn>();
-	ASkillSystemTestPawn *DestActor = World->SpawnActor<ASkillSystemTestPawn>();
+	AAbilitySystemTestPawn *SourceActor = World->SpawnActor<AAbilitySystemTestPawn>();
+	AAbilitySystemTestPawn *DestActor = World->SpawnActor<AAbilitySystemTestPawn>();
 
-	UProperty *HealthProperty = FindFieldChecked<UProperty>(USkillSystemTestAttributeSet::StaticClass(), GET_MEMBER_NAME_CHECKED(USkillSystemTestAttributeSet, Health));
+	UProperty *HealthProperty = FindFieldChecked<UProperty>(UAbilitySystemTestAttributeSet::StaticClass(), GET_MEMBER_NAME_CHECKED(UAbilitySystemTestAttributeSet, Health));
 
 	UAttributeComponent * SourceComponent = SourceActor->AttributeComponent;
 	UAttributeComponent * DestComponent = DestActor->AttributeComponent;
-	SourceComponent->GetSet<USkillSystemTestAttributeSet>()->Health = StartHealth;
-	DestComponent->GetSet<USkillSystemTestAttributeSet>()->Health = StartHealth;
+	SourceComponent->GetSet<UAbilitySystemTestAttributeSet>()->Health = StartHealth;
+	DestComponent->GetSet<UAbilitySystemTestAttributeSet>()->Health = StartHealth;
 
 	{
-		SKILL_LOG_SCOPE(TEXT("Apply InstantDamage"))
+		ABILITY_LOG_SCOPE(TEXT("Apply InstantDamage"))
 
 		UGameplayEffect * BaseDmgEffect = Cast<UGameplayEffect>(StaticConstructObject(UGameplayEffect::StaticClass(), GetTransientPackage(), FName(TEXT("BaseDmgEffect"))));
 		BaseDmgEffect->Modifiers.SetNum(1);
@@ -129,8 +129,8 @@ bool GameplayEffectsTest_InstantDamage(UWorld *World, FAutomationTestBase * Test
 
 		SourceComponent->ApplyGameplayEffectToTarget(BaseDmgEffect, DestComponent, 1.f);
 		
-		Test->TestTrue(SKILL_TEST_TEXT("Basic Instant Damage Applied"), (DestComponent->GetSet<USkillSystemTestAttributeSet>()->Health == (StartHealth + DamageValue)));
-		SKILL_LOG(Log, TEXT("Final Health: %.2f"), DestComponent->GetSet<USkillSystemTestAttributeSet>()->Health);
+		Test->TestTrue(SKILL_TEST_TEXT("Basic Instant Damage Applied"), (DestComponent->GetSet<UAbilitySystemTestAttributeSet>()->Health == (StartHealth + DamageValue)));
+		ABILITY_LOG(Log, TEXT("Final Health: %.2f"), DestComponent->GetSet<UAbilitySystemTestAttributeSet>()->Health);
 	}
 
 	World->EditorDestroyActor(SourceActor, false);
@@ -141,24 +141,24 @@ bool GameplayEffectsTest_InstantDamage(UWorld *World, FAutomationTestBase * Test
 
 bool GameplayEffectsTest_InstantDamageRemap(UWorld *World, FAutomationTestBase * Test)
 {
-	// This is the same as GameplayEffectsTest_InstantDamage but modifies the Damage attribute and confirms it is remapped to -Health by USkillSystemTestAttributeSet::PostAttributeModify
+	// This is the same as GameplayEffectsTest_InstantDamage but modifies the Damage attribute and confirms it is remapped to -Health by UAbilitySystemTestAttributeSet::PostAttributeModify
 
 	const float StartHealth = 100.f;
-	const float DamageValue = 5.f;		// Note: Damage is positive, mapped to -Health in USkillSystemTestAttributeSet::PostAttributeModify
+	const float DamageValue = 5.f;		// Note: Damage is positive, mapped to -Health in UAbilitySystemTestAttributeSet::PostAttributeModify
 
-	ASkillSystemTestPawn *SourceActor = World->SpawnActor<ASkillSystemTestPawn>();
-	ASkillSystemTestPawn *DestActor = World->SpawnActor<ASkillSystemTestPawn>();
+	AAbilitySystemTestPawn *SourceActor = World->SpawnActor<AAbilitySystemTestPawn>();
+	AAbilitySystemTestPawn *DestActor = World->SpawnActor<AAbilitySystemTestPawn>();
 
-	UProperty *HealthProperty = FindFieldChecked<UProperty>(USkillSystemTestAttributeSet::StaticClass(), GET_MEMBER_NAME_CHECKED(USkillSystemTestAttributeSet, Health));
-	UProperty *DamageProperty = FindFieldChecked<UProperty>(USkillSystemTestAttributeSet::StaticClass(), GET_MEMBER_NAME_CHECKED(USkillSystemTestAttributeSet, Damage));
+	UProperty *HealthProperty = FindFieldChecked<UProperty>(UAbilitySystemTestAttributeSet::StaticClass(), GET_MEMBER_NAME_CHECKED(UAbilitySystemTestAttributeSet, Health));
+	UProperty *DamageProperty = FindFieldChecked<UProperty>(UAbilitySystemTestAttributeSet::StaticClass(), GET_MEMBER_NAME_CHECKED(UAbilitySystemTestAttributeSet, Damage));
 
 	UAttributeComponent * SourceComponent = SourceActor->AttributeComponent;
 	UAttributeComponent * DestComponent = DestActor->AttributeComponent;
-	SourceComponent->GetSet<USkillSystemTestAttributeSet>()->Health = StartHealth;
-	DestComponent->GetSet<USkillSystemTestAttributeSet>()->Health = StartHealth;
+	SourceComponent->GetSet<UAbilitySystemTestAttributeSet>()->Health = StartHealth;
+	DestComponent->GetSet<UAbilitySystemTestAttributeSet>()->Health = StartHealth;
 
 	{
-		SKILL_LOG_SCOPE(TEXT("Apply InstantDamage"));
+		ABILITY_LOG_SCOPE(TEXT("Apply InstantDamage"));
 
 		UGameplayEffect * BaseDmgEffect = Cast<UGameplayEffect>(StaticConstructObject(UGameplayEffect::StaticClass(), GetTransientPackage(), FName(TEXT("BaseDmgEffect"))));
 		BaseDmgEffect->Modifiers.SetNum(1);
@@ -173,14 +173,14 @@ bool GameplayEffectsTest_InstantDamageRemap(UWorld *World, FAutomationTestBase *
 
 		// Now we should have lost some health
 		{
-			float ActualValue = DestComponent->GetSet<USkillSystemTestAttributeSet>()->Health;
+			float ActualValue = DestComponent->GetSet<UAbilitySystemTestAttributeSet>()->Health;
 			float ExpectedValue = StartHealth + -DamageValue;
 			Test->TestTrue(SKILL_TEST_TEXT("Damage Applied. Actual: %.2f == Exected: %.2f", ActualValue, ExpectedValue), ActualValue == ExpectedValue);
 		}
 
 		// Confirm the damage attribute itself was reset to 0 when it was applied to health
 		{
-			float ActualValue = DestComponent->GetSet<USkillSystemTestAttributeSet>()->Damage;
+			float ActualValue = DestComponent->GetSet<UAbilitySystemTestAttributeSet>()->Damage;
 			float ExpectedValue = 0.f;
 			Test->TestTrue(SKILL_TEST_TEXT("Damage Applied. Actual: %.2f == Exected: %.2f", ActualValue, ExpectedValue), ActualValue == ExpectedValue);
 		}
@@ -198,19 +198,19 @@ bool GameplayEffectsTest_InstantDamage_Buffed(UWorld *World, FAutomationTestBase
 	const float DamageValue = -5.f;
 	const float BonusDamageMultiplier = 2.f;
 
-	ASkillSystemTestPawn *SourceActor = World->SpawnActor<ASkillSystemTestPawn>();
-	ASkillSystemTestPawn *DestActor = World->SpawnActor<ASkillSystemTestPawn>();
+	AAbilitySystemTestPawn *SourceActor = World->SpawnActor<AAbilitySystemTestPawn>();
+	AAbilitySystemTestPawn *DestActor = World->SpawnActor<AAbilitySystemTestPawn>();
 
-	UProperty *HealthProperty = FindFieldChecked<UProperty>(USkillSystemTestAttributeSet::StaticClass(), GET_MEMBER_NAME_CHECKED(USkillSystemTestAttributeSet, Health));
+	UProperty *HealthProperty = FindFieldChecked<UProperty>(UAbilitySystemTestAttributeSet::StaticClass(), GET_MEMBER_NAME_CHECKED(UAbilitySystemTestAttributeSet, Health));
 
 	UAttributeComponent * SourceComponent = SourceActor->AttributeComponent;
 	UAttributeComponent * DestComponent = DestActor->AttributeComponent;
-	SourceComponent->GetSet<USkillSystemTestAttributeSet>()->Health = StartHealth;
-	DestComponent->GetSet<USkillSystemTestAttributeSet>()->Health = StartHealth;
+	SourceComponent->GetSet<UAbilitySystemTestAttributeSet>()->Health = StartHealth;
+	DestComponent->GetSet<UAbilitySystemTestAttributeSet>()->Health = StartHealth;
 
 	// Setup a GE to modify OutgoingGEs
 	{
-		SKILL_LOG_SCOPE(TEXT("Apply DamageBuff"))
+		ABILITY_LOG_SCOPE(TEXT("Apply DamageBuff"))
 		
 		UGameplayEffect * BaseDmgEffect = Cast<UGameplayEffect>(StaticConstructObject(UGameplayEffect::StaticClass(), GetTransientPackage(), FName(TEXT("DamageBuff"))));
 		BaseDmgEffect->Modifiers.SetNum(1);
@@ -227,7 +227,7 @@ bool GameplayEffectsTest_InstantDamage_Buffed(UWorld *World, FAutomationTestBase
 
 	// Apply Damage
 	{
-		SKILL_LOG_SCOPE(TEXT("Apply InstantDamage"))
+		ABILITY_LOG_SCOPE(TEXT("Apply InstantDamage"))
 
 		UGameplayEffect * BaseDmgEffect = Cast<UGameplayEffect>(StaticConstructObject(UGameplayEffect::StaticClass(), GetTransientPackage(), FName(TEXT("Damage"))));
 		BaseDmgEffect->Modifiers.SetNum(1);
@@ -243,8 +243,8 @@ bool GameplayEffectsTest_InstantDamage_Buffed(UWorld *World, FAutomationTestBase
 
 		float ExpectedValue = (StartHealth + (DamageValue * BonusDamageMultiplier));
 
-		Test->TestTrue(SKILL_TEST_TEXT("Buff Instant Damage Applied"), (DestComponent->GetSet<USkillSystemTestAttributeSet>()->Health == ExpectedValue));
-		SKILL_LOG(Log, TEXT("Final Health: %.2f"), DestComponent->GetSet<USkillSystemTestAttributeSet>()->Health);
+		Test->TestTrue(SKILL_TEST_TEXT("Buff Instant Damage Applied"), (DestComponent->GetSet<UAbilitySystemTestAttributeSet>()->Health == ExpectedValue));
+		ABILITY_LOG(Log, TEXT("Final Health: %.2f"), DestComponent->GetSet<UAbilitySystemTestAttributeSet>()->Health);
 	}
 
 	World->EditorDestroyActor(SourceActor, false);
@@ -262,21 +262,21 @@ bool GameplayEffectsTest_TemporaryDamage(UWorld *World, FAutomationTestBase * Te
 	const float StartHealth = 100.f;
 	const float DamageValue = -5.f;
 
-	ASkillSystemTestPawn *SourceActor = World->SpawnActor<ASkillSystemTestPawn>();
-	ASkillSystemTestPawn *DestActor = World->SpawnActor<ASkillSystemTestPawn>();
+	AAbilitySystemTestPawn *SourceActor = World->SpawnActor<AAbilitySystemTestPawn>();
+	AAbilitySystemTestPawn *DestActor = World->SpawnActor<AAbilitySystemTestPawn>();
 
-	UProperty *HealthProperty = FindFieldChecked<UProperty>(USkillSystemTestAttributeSet::StaticClass(), GET_MEMBER_NAME_CHECKED(USkillSystemTestAttributeSet, Health));
+	UProperty *HealthProperty = FindFieldChecked<UProperty>(UAbilitySystemTestAttributeSet::StaticClass(), GET_MEMBER_NAME_CHECKED(UAbilitySystemTestAttributeSet, Health));
 
 	UAttributeComponent * SourceComponent = SourceActor->AttributeComponent;
 	UAttributeComponent * DestComponent = DestActor->AttributeComponent;
-	SourceComponent->GetSet<USkillSystemTestAttributeSet>()->Health = StartHealth;
-	DestComponent->GetSet<USkillSystemTestAttributeSet>()->Health = StartHealth;
+	SourceComponent->GetSet<UAbilitySystemTestAttributeSet>()->Health = StartHealth;
+	DestComponent->GetSet<UAbilitySystemTestAttributeSet>()->Health = StartHealth;
 
 	// Apply "Damage" but set to INFINITE_DURATION
 	// (An odd example for damage, but would make sense for something like run speed, etc)
 	FActiveGameplayEffectHandle AppliedHandle;
 	{
-		SKILL_LOG_SCOPE(TEXT("Apply Permanent Damage"))
+		ABILITY_LOG_SCOPE(TEXT("Apply Permanent Damage"))
 
 		UGameplayEffect * BaseDmgEffect = Cast<UGameplayEffect>(StaticConstructObject(UGameplayEffect::StaticClass(), GetTransientPackage(), FName(TEXT("Damage"))));
 		BaseDmgEffect->Modifiers.SetNum(1);
@@ -292,8 +292,8 @@ bool GameplayEffectsTest_TemporaryDamage(UWorld *World, FAutomationTestBase * Te
 
 		float ExpectedValue = (StartHealth + DamageValue);
 
-		Test->TestTrue(SKILL_TEST_TEXT("INFINITE_DURATION Damage Applied"), (DestComponent->GetSet<USkillSystemTestAttributeSet>()->Health == ExpectedValue));
-		SKILL_LOG(Log, TEXT("After Damage Applied: Health: %.2f"), DestComponent->GetSet<USkillSystemTestAttributeSet>()->Health);
+		Test->TestTrue(SKILL_TEST_TEXT("INFINITE_DURATION Damage Applied"), (DestComponent->GetSet<UAbilitySystemTestAttributeSet>()->Health == ExpectedValue));
+		ABILITY_LOG(Log, TEXT("After Damage Applied: Health: %.2f"), DestComponent->GetSet<UAbilitySystemTestAttributeSet>()->Health);
 	}
 
 	// Now remove the GameplayEffect we just added and confirm Health goes back to starting value
@@ -301,8 +301,8 @@ bool GameplayEffectsTest_TemporaryDamage(UWorld *World, FAutomationTestBase * Te
 		bool RemovedEffect = DestComponent->RemoveActiveGameplayEffect(AppliedHandle);
 		float ExpectedValue = StartHealth;
 
-		Test->TestTrue(SKILL_TEST_TEXT("INFINITE_DURATION Damage Applied"), (DestComponent->GetSet<USkillSystemTestAttributeSet>()->Health == ExpectedValue));
-		SKILL_LOG(Log, TEXT("After Removal. Health: %.2f. RemovedEffecte: %d"), DestComponent->GetSet<USkillSystemTestAttributeSet>()->Health, RemovedEffect);
+		Test->TestTrue(SKILL_TEST_TEXT("INFINITE_DURATION Damage Applied"), (DestComponent->GetSet<UAbilitySystemTestAttributeSet>()->Health == ExpectedValue));
+		ABILITY_LOG(Log, TEXT("After Removal. Health: %.2f. RemovedEffecte: %d"), DestComponent->GetSet<UAbilitySystemTestAttributeSet>()->Health, RemovedEffect);
 	}
 
 	World->EditorDestroyActor(SourceActor, false);
@@ -321,21 +321,21 @@ bool GameplayEffectsTest_TemporaryDamageBuffed(UWorld *World, FAutomationTestBas
 	const float DamageValue = -5.f;
 	const float DamageBuffMultiplier = 2.f;
 
-	ASkillSystemTestPawn *SourceActor = World->SpawnActor<ASkillSystemTestPawn>();
-	ASkillSystemTestPawn *DestActor = World->SpawnActor<ASkillSystemTestPawn>();
+	AAbilitySystemTestPawn *SourceActor = World->SpawnActor<AAbilitySystemTestPawn>();
+	AAbilitySystemTestPawn *DestActor = World->SpawnActor<AAbilitySystemTestPawn>();
 
-	UProperty *HealthProperty = FindFieldChecked<UProperty>(USkillSystemTestAttributeSet::StaticClass(), GET_MEMBER_NAME_CHECKED(USkillSystemTestAttributeSet, Health));
+	UProperty *HealthProperty = FindFieldChecked<UProperty>(UAbilitySystemTestAttributeSet::StaticClass(), GET_MEMBER_NAME_CHECKED(UAbilitySystemTestAttributeSet, Health));
 
 	UAttributeComponent * SourceComponent = SourceActor->AttributeComponent;
 	UAttributeComponent * DestComponent = DestActor->AttributeComponent;
-	SourceComponent->GetSet<USkillSystemTestAttributeSet>()->Health = StartHealth;
-	DestComponent->GetSet<USkillSystemTestAttributeSet>()->Health = StartHealth;
+	SourceComponent->GetSet<UAbilitySystemTestAttributeSet>()->Health = StartHealth;
+	DestComponent->GetSet<UAbilitySystemTestAttributeSet>()->Health = StartHealth;
 
 	// Apply "Damage" but set to INFINITE_DURATION
 	// (An odd example for damage, but would make sense for something like run speed, etc)
 	FActiveGameplayEffectHandle AppliedHandle;
 	{
-		SKILL_LOG_SCOPE(TEXT("Apply Permanent Damage"))
+		ABILITY_LOG_SCOPE(TEXT("Apply Permanent Damage"))
 
 		UGameplayEffect * BaseDmgEffect = Cast<UGameplayEffect>(StaticConstructObject(UGameplayEffect::StaticClass(), GetTransientPackage(), FName(TEXT("Damage"))));
 		BaseDmgEffect->Modifiers.SetNum(1);
@@ -351,15 +351,15 @@ bool GameplayEffectsTest_TemporaryDamageBuffed(UWorld *World, FAutomationTestBas
 
 		float ExpectedValue = (StartHealth + DamageValue);
 
-		Test->TestTrue(SKILL_TEST_TEXT("INFINITE_DURATION Damage Applied"), (DestComponent->GetSet<USkillSystemTestAttributeSet>()->Health == ExpectedValue));
-		SKILL_LOG(Log, TEXT("After Damage Applied: Health: %.2f"), DestComponent->GetSet<USkillSystemTestAttributeSet>()->Health);
+		Test->TestTrue(SKILL_TEST_TEXT("INFINITE_DURATION Damage Applied"), (DestComponent->GetSet<UAbilitySystemTestAttributeSet>()->Health == ExpectedValue));
+		ABILITY_LOG(Log, TEXT("After Damage Applied: Health: %.2f"), DestComponent->GetSet<UAbilitySystemTestAttributeSet>()->Health);
 
 		Test->TestTrue(SKILL_TEST_TEXT("Number of GameplayEffects=1"), DestComponent->GetNumActiveGameplayEffect() == 1);
 	}
 
 	// Now Buff the GameplayEffect we just added and confirm the health removal is increased 2x
 	{
-		SKILL_LOG_SCOPE(TEXT("Buff Permanent Damage"))
+		ABILITY_LOG_SCOPE(TEXT("Buff Permanent Damage"))
 
 		UGameplayEffect * BuffDmgEffect = Cast<UGameplayEffect>(StaticConstructObject(UGameplayEffect::StaticClass(), GetTransientPackage(), FName(TEXT("BuffDamage"))));
 		BuffDmgEffect->Modifiers.SetNum(1);
@@ -375,8 +375,8 @@ bool GameplayEffectsTest_TemporaryDamageBuffed(UWorld *World, FAutomationTestBas
 
 		float ExpectedValue = (StartHealth + (DamageValue * DamageBuffMultiplier));
 
-		Test->TestTrue(SKILL_TEST_TEXT("INFINITE_DURATION Buffed Damage Applied"), (DestComponent->GetSet<USkillSystemTestAttributeSet>()->Health == ExpectedValue));
-		SKILL_LOG(Log, TEXT("After Damage Applied: Health: %.2f"), DestComponent->GetSet<USkillSystemTestAttributeSet>()->Health);
+		Test->TestTrue(SKILL_TEST_TEXT("INFINITE_DURATION Buffed Damage Applied"), (DestComponent->GetSet<UAbilitySystemTestAttributeSet>()->Health == ExpectedValue));
+		ABILITY_LOG(Log, TEXT("After Damage Applied: Health: %.2f"), DestComponent->GetSet<UAbilitySystemTestAttributeSet>()->Health);
 
 		// Confirm still only 1 active GE (since this was instant application
 		Test->TestTrue(SKILL_TEST_TEXT("Number of GameplayEffects=1"), DestComponent->GetNumActiveGameplayEffect() == 1);
@@ -388,8 +388,8 @@ bool GameplayEffectsTest_TemporaryDamageBuffed(UWorld *World, FAutomationTestBas
 		bool RemovedEffect = DestComponent->RemoveActiveGameplayEffect(AppliedHandle);
 		float ExpectedValue = StartHealth;
 
-		Test->TestTrue(SKILL_TEST_TEXT("INFINITE_DURATION Damage Applied"), (DestComponent->GetSet<USkillSystemTestAttributeSet>()->Health == ExpectedValue));
-		SKILL_LOG(Log, TEXT("After Removal. Health: %.2f. RemovedEffecte: %d"), DestComponent->GetSet<USkillSystemTestAttributeSet>()->Health, RemovedEffect);
+		Test->TestTrue(SKILL_TEST_TEXT("INFINITE_DURATION Damage Applied"), (DestComponent->GetSet<UAbilitySystemTestAttributeSet>()->Health == ExpectedValue));
+		ABILITY_LOG(Log, TEXT("After Removal. Health: %.2f. RemovedEffecte: %d"), DestComponent->GetSet<UAbilitySystemTestAttributeSet>()->Health, RemovedEffect);
 
 		// Confirm no more GEs
 		Test->TestTrue(SKILL_TEST_TEXT("Number of GameplayEffects=0"), DestComponent->GetNumActiveGameplayEffect() == 0);
@@ -411,15 +411,15 @@ bool GameplayEffectsTest_TemporaryDamageTemporaryBuff(UWorld *World, FAutomation
 	const float DamageValue = -5.f;
 	const float DamageBuffMultiplier = 2.f;
 
-	ASkillSystemTestPawn *SourceActor = World->SpawnActor<ASkillSystemTestPawn>();
-	ASkillSystemTestPawn *DestActor = World->SpawnActor<ASkillSystemTestPawn>();
+	AAbilitySystemTestPawn *SourceActor = World->SpawnActor<AAbilitySystemTestPawn>();
+	AAbilitySystemTestPawn *DestActor = World->SpawnActor<AAbilitySystemTestPawn>();
 
-	UProperty *HealthProperty = FindFieldChecked<UProperty>(USkillSystemTestAttributeSet::StaticClass(), GET_MEMBER_NAME_CHECKED(USkillSystemTestAttributeSet, Health));
+	UProperty *HealthProperty = FindFieldChecked<UProperty>(UAbilitySystemTestAttributeSet::StaticClass(), GET_MEMBER_NAME_CHECKED(UAbilitySystemTestAttributeSet, Health));
 
 	UAttributeComponent * SourceComponent = SourceActor->AttributeComponent;
 	UAttributeComponent * DestComponent = DestActor->AttributeComponent;
-	SourceComponent->GetSet<USkillSystemTestAttributeSet>()->Health = StartHealth;
-	DestComponent->GetSet<USkillSystemTestAttributeSet>()->Health = StartHealth;
+	SourceComponent->GetSet<UAbilitySystemTestAttributeSet>()->Health = StartHealth;
+	DestComponent->GetSet<UAbilitySystemTestAttributeSet>()->Health = StartHealth;
 
 	// Apply "Damage" but set to INFINITE_DURATION
 	// (An odd example for damage, but would make sense for something like run speed, etc)
@@ -427,7 +427,7 @@ bool GameplayEffectsTest_TemporaryDamageTemporaryBuff(UWorld *World, FAutomation
 	FActiveGameplayEffectHandle AppliedHandleBuff;
 
 	{
-		SKILL_LOG_SCOPE(TEXT("Apply Permanent Damage"))
+		ABILITY_LOG_SCOPE(TEXT("Apply Permanent Damage"))
 
 		UGameplayEffect * BaseDmgEffect = Cast<UGameplayEffect>(StaticConstructObject(UGameplayEffect::StaticClass(), GetTransientPackage(), FName(TEXT("Damage"))));
 		BaseDmgEffect->Modifiers.SetNum(1);
@@ -443,15 +443,15 @@ bool GameplayEffectsTest_TemporaryDamageTemporaryBuff(UWorld *World, FAutomation
 
 		float ExpectedValue = (StartHealth + DamageValue);
 
-		Test->TestTrue(SKILL_TEST_TEXT("GameplayEffectsTest_TemporaryDamageTemporaryBuff INFINITE_DURATION Damage Applied"), (DestComponent->GetSet<USkillSystemTestAttributeSet>()->Health == ExpectedValue));
-		SKILL_LOG(Log, TEXT("After Damage Applied: Health: %.2f"), DestComponent->GetSet<USkillSystemTestAttributeSet>()->Health);
+		Test->TestTrue(SKILL_TEST_TEXT("GameplayEffectsTest_TemporaryDamageTemporaryBuff INFINITE_DURATION Damage Applied"), (DestComponent->GetSet<UAbilitySystemTestAttributeSet>()->Health == ExpectedValue));
+		ABILITY_LOG(Log, TEXT("After Damage Applied: Health: %.2f"), DestComponent->GetSet<UAbilitySystemTestAttributeSet>()->Health);
 
 		Test->TestTrue(SKILL_TEST_TEXT("GameplayEffectsTest_TemporaryDamageTemporaryBuff Number of GameplayEffects=1"), DestComponent->GetNumActiveGameplayEffect() == 1);
 	}
 
 	// Now Buff the GameplayEffect we just added and confirm the health removal is increased 2x
 	{
-		SKILL_LOG_SCOPE(TEXT("Buff Permanent Damage"))
+		ABILITY_LOG_SCOPE(TEXT("Buff Permanent Damage"))
 
 		UGameplayEffect * BuffDmgEffect = Cast<UGameplayEffect>(StaticConstructObject(UGameplayEffect::StaticClass(), GetTransientPackage(), FName(TEXT("BuffDamage"))));
 		BuffDmgEffect->Modifiers.SetNum(1);
@@ -468,8 +468,8 @@ bool GameplayEffectsTest_TemporaryDamageTemporaryBuff(UWorld *World, FAutomation
 
 		float ExpectedValue = (StartHealth + (DamageValue * DamageBuffMultiplier));
 
-		Test->TestTrue(SKILL_TEST_TEXT("GameplayEffectsTest_TemporaryDamageTemporaryBuff INFINITE_DURATION Buffed Damage Applied"), (DestComponent->GetSet<USkillSystemTestAttributeSet>()->Health == ExpectedValue));
-		SKILL_LOG(Log, TEXT("After Damage Applied: Health: %.2f"), DestComponent->GetSet<USkillSystemTestAttributeSet>()->Health);
+		Test->TestTrue(SKILL_TEST_TEXT("GameplayEffectsTest_TemporaryDamageTemporaryBuff INFINITE_DURATION Buffed Damage Applied"), (DestComponent->GetSet<UAbilitySystemTestAttributeSet>()->Health == ExpectedValue));
+		ABILITY_LOG(Log, TEXT("After Damage Applied: Health: %.2f"), DestComponent->GetSet<UAbilitySystemTestAttributeSet>()->Health);
 
 		// Confirm there are 2 GEs
 		Test->TestTrue(SKILL_TEST_TEXT("GameplayEffectsTest_TemporaryDamageTemporaryBuff Number of GameplayEffects=1"), DestComponent->GetNumActiveGameplayEffect() == 2);
@@ -485,8 +485,8 @@ bool GameplayEffectsTest_TemporaryDamageTemporaryBuff(UWorld *World, FAutomation
 		bool RemovedEffect = DestComponent->RemoveActiveGameplayEffect(AppliedHandleBuff);
 		float ExpectedValue = StartHealth + DamageValue;
 
-		Test->TestTrue(SKILL_TEST_TEXT("GameplayEffectsTest_TemporaryDamageTemporaryBuff INFINITE_DURATION Damage Buff Removed"), (DestComponent->GetSet<USkillSystemTestAttributeSet>()->Health == ExpectedValue));
-		SKILL_LOG(Log, TEXT("After Buff Removal. Health: %.2f. RemovedEffecte: %d"), DestComponent->GetSet<USkillSystemTestAttributeSet>()->Health, RemovedEffect);
+		Test->TestTrue(SKILL_TEST_TEXT("GameplayEffectsTest_TemporaryDamageTemporaryBuff INFINITE_DURATION Damage Buff Removed"), (DestComponent->GetSet<UAbilitySystemTestAttributeSet>()->Health == ExpectedValue));
+		ABILITY_LOG(Log, TEXT("After Buff Removal. Health: %.2f. RemovedEffecte: %d"), DestComponent->GetSet<UAbilitySystemTestAttributeSet>()->Health, RemovedEffect);
 
 		// Confirm 1 more GEs
 		Test->TestTrue(SKILL_TEST_TEXT("Number of GameplayEffects=1"), DestComponent->GetNumActiveGameplayEffect() == 1);
@@ -497,8 +497,8 @@ bool GameplayEffectsTest_TemporaryDamageTemporaryBuff(UWorld *World, FAutomation
 		bool RemovedEffect = DestComponent->RemoveActiveGameplayEffect(AppliedHandle);
 		float ExpectedValue = StartHealth;
 
-		Test->TestTrue(SKILL_TEST_TEXT("GameplayEffectsTest_TemporaryDamageTemporaryBuff INFINITE_DURATION Damage Removed"), (DestComponent->GetSet<USkillSystemTestAttributeSet>()->Health == ExpectedValue));
-		SKILL_LOG(Log, TEXT("After Removal. Health: %.2f. RemovedEffecte: %d"), DestComponent->GetSet<USkillSystemTestAttributeSet>()->Health, RemovedEffect);
+		Test->TestTrue(SKILL_TEST_TEXT("GameplayEffectsTest_TemporaryDamageTemporaryBuff INFINITE_DURATION Damage Removed"), (DestComponent->GetSet<UAbilitySystemTestAttributeSet>()->Health == ExpectedValue));
+		ABILITY_LOG(Log, TEXT("After Removal. Health: %.2f. RemovedEffecte: %d"), DestComponent->GetSet<UAbilitySystemTestAttributeSet>()->Health, RemovedEffect);
 
 		// Confirm no more GEs
 		Test->TestTrue(SKILL_TEST_TEXT("GameplayEffectsTest_TemporaryDamageTemporaryBuff Number of GameplayEffects=0"), DestComponent->GetNumActiveGameplayEffect() == 0);
@@ -520,15 +520,15 @@ bool GameplayEffectsTest_LinkedBuffDestroy(UWorld *World, FAutomationTestBase *
 	const float DamageValue = -5.f;
 	const float DamageBuffMultiplier = 2.f;
 
-	ASkillSystemTestPawn *SourceActor = World->SpawnActor<ASkillSystemTestPawn>();
-	ASkillSystemTestPawn *DestActor = World->SpawnActor<ASkillSystemTestPawn>();
+	AAbilitySystemTestPawn *SourceActor = World->SpawnActor<AAbilitySystemTestPawn>();
+	AAbilitySystemTestPawn *DestActor = World->SpawnActor<AAbilitySystemTestPawn>();
 
-	UProperty *HealthProperty = FindFieldChecked<UProperty>(USkillSystemTestAttributeSet::StaticClass(), GET_MEMBER_NAME_CHECKED(USkillSystemTestAttributeSet, Health));
+	UProperty *HealthProperty = FindFieldChecked<UProperty>(UAbilitySystemTestAttributeSet::StaticClass(), GET_MEMBER_NAME_CHECKED(UAbilitySystemTestAttributeSet, Health));
 
 	UAttributeComponent * SourceComponent = SourceActor->AttributeComponent;
 	UAttributeComponent * DestComponent = DestActor->AttributeComponent;
-	SourceComponent->GetSet<USkillSystemTestAttributeSet>()->Health = StartHealth;
-	DestComponent->GetSet<USkillSystemTestAttributeSet>()->Health = StartHealth;
+	SourceComponent->GetSet<UAbilitySystemTestAttributeSet>()->Health = StartHealth;
+	DestComponent->GetSet<UAbilitySystemTestAttributeSet>()->Health = StartHealth;
 
 	// Apply "Damage" but set to INFINITE_DURATION
 	// (An odd example for damage, but would make sense for something like run speed, etc)
@@ -536,7 +536,7 @@ bool GameplayEffectsTest_LinkedBuffDestroy(UWorld *World, FAutomationTestBase *
 	FActiveGameplayEffectHandle AppliedHandleBuff;
 
 	{
-		SKILL_LOG_SCOPE(TEXT("Apply Damage Buff"))
+		ABILITY_LOG_SCOPE(TEXT("Apply Damage Buff"))
 
 		UGameplayEffect * BuffEffect = Cast<UGameplayEffect>(StaticConstructObject(UGameplayEffect::StaticClass(), GetTransientPackage(), FName(TEXT("BuffOutgoingDamage"))));
 		BuffEffect->Modifiers.SetNum(1);
@@ -554,7 +554,7 @@ bool GameplayEffectsTest_LinkedBuffDestroy(UWorld *World, FAutomationTestBase *
 	}
 
 	{
-		SKILL_LOG_SCOPE(TEXT("Apply Permanent (infinite duration) Damage"))
+		ABILITY_LOG_SCOPE(TEXT("Apply Permanent (infinite duration) Damage"))
 
 		UGameplayEffect * BaseDmgEffect = Cast<UGameplayEffect>(StaticConstructObject(UGameplayEffect::StaticClass(), GetTransientPackage(), FName(TEXT("Damage"))));
 		BaseDmgEffect->Modifiers.SetNum(1);
@@ -569,7 +569,7 @@ bool GameplayEffectsTest_LinkedBuffDestroy(UWorld *World, FAutomationTestBase *
 		AppliedHandle = SourceComponent->ApplyGameplayEffectToTarget(BaseDmgEffect, DestComponent, 1.f);
 
 		float ExpectedValue = (StartHealth + (DamageValue * DamageBuffMultiplier));
-		float ActualValue = DestComponent->GetSet<USkillSystemTestAttributeSet>()->Health;
+		float ActualValue = DestComponent->GetSet<UAbilitySystemTestAttributeSet>()->Health;
 
 		Test->TestTrue(SKILL_TEST_TEXT("Damaged Applied. ActualValue: %.2f. ExpectedValue: %.2f.", ActualValue, ExpectedValue), (ActualValue == ExpectedValue));
 		Test->TestTrue(SKILL_TEST_TEXT("NumberOfActive GameplayEffects GameplayEffects %d == 1", DestComponent->GetNumActiveGameplayEffect() ), DestComponent->GetNumActiveGameplayEffect() == 1);
@@ -579,7 +579,7 @@ bool GameplayEffectsTest_LinkedBuffDestroy(UWorld *World, FAutomationTestBase *
 	{
 		bool RemovedEffect = SourceComponent->RemoveActiveGameplayEffect(AppliedHandleBuff);
 		float ExpectedValue = StartHealth + DamageValue;
-		float ActualValue = DestComponent->GetSet<USkillSystemTestAttributeSet>()->Health;
+		float ActualValue = DestComponent->GetSet<UAbilitySystemTestAttributeSet>()->Health;
 
 		// Confirm we regained healtht
 		Test->TestTrue(SKILL_TEST_TEXT("After Buff Removal. ActualValue: %.2f. ExpectedValue: %.2f. RemovedEffecte: %d", ActualValue, ExpectedValue, RemovedEffect), (ActualValue  == ExpectedValue));
@@ -607,16 +607,16 @@ bool GameplayEffectsTest_SnapshotBuffDestroy(UWorld *World, FAutomationTestBase
 	const float DamageValue = -5.f;
 	const float DamageBuffMultiplier = 2.f;
 
-	ASkillSystemTestPawn *SourceActor = World->SpawnActor<ASkillSystemTestPawn>();
-	ASkillSystemTestPawn *DestActor = World->SpawnActor<ASkillSystemTestPawn>();
+	AAbilitySystemTestPawn *SourceActor = World->SpawnActor<AAbilitySystemTestPawn>();
+	AAbilitySystemTestPawn *DestActor = World->SpawnActor<AAbilitySystemTestPawn>();
 
-	UProperty *HealthProperty = FindFieldChecked<UProperty>(USkillSystemTestAttributeSet::StaticClass(), GET_MEMBER_NAME_CHECKED(USkillSystemTestAttributeSet, Health));
+	UProperty *HealthProperty = FindFieldChecked<UProperty>(UAbilitySystemTestAttributeSet::StaticClass(), GET_MEMBER_NAME_CHECKED(UAbilitySystemTestAttributeSet, Health));
 
 	UAttributeComponent * SourceComponent = SourceActor->AttributeComponent;
 	UAttributeComponent * DestComponent = DestActor->AttributeComponent;
 	
-	SourceComponent->GetSet<USkillSystemTestAttributeSet>()->Health = StartHealth;
-	DestComponent->GetSet<USkillSystemTestAttributeSet>()->Health = StartHealth;
+	SourceComponent->GetSet<UAbilitySystemTestAttributeSet>()->Health = StartHealth;
+	DestComponent->GetSet<UAbilitySystemTestAttributeSet>()->Health = StartHealth;
 
 	// Apply "Damage" but set to INFINITE_DURATION
 	// (An odd example for damage, but would make sense for something like run speed, etc)
@@ -624,7 +624,7 @@ bool GameplayEffectsTest_SnapshotBuffDestroy(UWorld *World, FAutomationTestBase
 	FActiveGameplayEffectHandle AppliedHandleBuff;
 
 	{
-		SKILL_LOG_SCOPE(TEXT("Apply Damage Buff"))
+		ABILITY_LOG_SCOPE(TEXT("Apply Damage Buff"))
 
 		UGameplayEffect * BuffEffect = Cast<UGameplayEffect>(StaticConstructObject(UGameplayEffect::StaticClass(), GetTransientPackage(), FName(TEXT("BuffOutgoingDamage"))));
 		BuffEffect->Modifiers.SetNum(1);
@@ -642,7 +642,7 @@ bool GameplayEffectsTest_SnapshotBuffDestroy(UWorld *World, FAutomationTestBase
 	}
 
 	{
-		SKILL_LOG_SCOPE(TEXT("Apply Permanent (infinite duration) Damage"))
+		ABILITY_LOG_SCOPE(TEXT("Apply Permanent (infinite duration) Damage"))
 
 		UGameplayEffect * BaseDmgEffect = Cast<UGameplayEffect>(StaticConstructObject(UGameplayEffect::StaticClass(), GetTransientPackage(), FName(TEXT("Damage"))));
 		BaseDmgEffect->Modifiers.SetNum(1);
@@ -657,7 +657,7 @@ bool GameplayEffectsTest_SnapshotBuffDestroy(UWorld *World, FAutomationTestBase
 		AppliedHandle = SourceComponent->ApplyGameplayEffectToTarget(BaseDmgEffect, DestComponent, 1.f);
 
 		float ExpectedValue = (StartHealth + (DamageValue * DamageBuffMultiplier));
-		float ActualValue = DestComponent->GetSet<USkillSystemTestAttributeSet>()->Health;
+		float ActualValue = DestComponent->GetSet<UAbilitySystemTestAttributeSet>()->Health;
 
 		Test->TestTrue(SKILL_TEST_TEXT("Damaged Applied. ActualValue: %.2f. ExpectedValue: %.2f.", ActualValue, ExpectedValue), (ActualValue == ExpectedValue));
 		Test->TestTrue(SKILL_TEST_TEXT("NumberOfActive GameplayEffects GameplayEffects %d == 1", DestComponent->GetNumActiveGameplayEffect()), DestComponent->GetNumActiveGameplayEffect() == 1);
@@ -669,7 +669,7 @@ bool GameplayEffectsTest_SnapshotBuffDestroy(UWorld *World, FAutomationTestBase
 		
 		// Check health again - it shouldn't have changed!
 		float ExpectedValue = (StartHealth + (DamageValue * DamageBuffMultiplier));
-		float ActualValue = DestComponent->GetSet<USkillSystemTestAttributeSet>()->Health;
+		float ActualValue = DestComponent->GetSet<UAbilitySystemTestAttributeSet>()->Health;
 
 		// Confirm we regained healtht
 		Test->TestTrue(SKILL_TEST_TEXT("After Buff Removal. ActualValue: %.2f. ExpectedValue: %.2f. RemovedEffecte: %d", ActualValue, ExpectedValue, RemovedEffect), (ActualValue == ExpectedValue));
@@ -699,16 +699,16 @@ bool GameplayEffectsTest_DurationBuff(UWorld *World, FAutomationTestBase * Test)
 	const float DurationBuff = 1.f;
 	const float DurationDebuff = -1.f;
 
-	ASkillSystemTestPawn *SourceActor = World->SpawnActor<ASkillSystemTestPawn>();
-	ASkillSystemTestPawn *DestActor = World->SpawnActor<ASkillSystemTestPawn>();
+	AAbilitySystemTestPawn *SourceActor = World->SpawnActor<AAbilitySystemTestPawn>();
+	AAbilitySystemTestPawn *DestActor = World->SpawnActor<AAbilitySystemTestPawn>();
 
-	UProperty *HealthProperty = FindFieldChecked<UProperty>(USkillSystemTestAttributeSet::StaticClass(), GET_MEMBER_NAME_CHECKED(USkillSystemTestAttributeSet, Health));
+	UProperty *HealthProperty = FindFieldChecked<UProperty>(UAbilitySystemTestAttributeSet::StaticClass(), GET_MEMBER_NAME_CHECKED(UAbilitySystemTestAttributeSet, Health));
 
 	UAttributeComponent * SourceComponent = SourceActor->AttributeComponent;
 	UAttributeComponent * DestComponent = DestActor->AttributeComponent;
 
-	SourceComponent->GetSet<USkillSystemTestAttributeSet>()->Health = StartHealth;
-	DestComponent->GetSet<USkillSystemTestAttributeSet>()->Health = StartHealth;
+	SourceComponent->GetSet<UAbilitySystemTestAttributeSet>()->Health = StartHealth;
+	DestComponent->GetSet<UAbilitySystemTestAttributeSet>()->Health = StartHealth;
 
 	// Apply Damage with 2 duration
 
@@ -718,7 +718,7 @@ bool GameplayEffectsTest_DurationBuff(UWorld *World, FAutomationTestBase * Test)
 	// Duration Debuff
 	//
 	{
-		SKILL_LOG_SCOPE(TEXT("Apply Damage mod that lasts 2 seconds"))
+		ABILITY_LOG_SCOPE(TEXT("Apply Damage mod that lasts 2 seconds"))
 
 		UGameplayEffect * BaseDmgEffect = Cast<UGameplayEffect>(StaticConstructObject(UGameplayEffect::StaticClass(), GetTransientPackage(), FName(TEXT("Damage"))));
 		BaseDmgEffect->Modifiers.SetNum(1);
@@ -733,7 +733,7 @@ bool GameplayEffectsTest_DurationBuff(UWorld *World, FAutomationTestBase * Test)
 		AppliedDamageHandle = SourceComponent->ApplyGameplayEffectToTarget(BaseDmgEffect, DestComponent, 1.f);
 
 		float ExpectedHealthValue = (StartHealth + (DamageValue));
-		float ActualHealthValue = DestComponent->GetSet<USkillSystemTestAttributeSet>()->Health;
+		float ActualHealthValue = DestComponent->GetSet<UAbilitySystemTestAttributeSet>()->Health;
 		float ExpectedDuration = BaseDuration;
 		float ActualDuration = DestComponent->GetGameplayEffectDuration(AppliedDamageHandle);
 		
@@ -747,7 +747,7 @@ bool GameplayEffectsTest_DurationBuff(UWorld *World, FAutomationTestBase * Test)
 
 	// Debuff the duration of the effect
 	{
-		SKILL_LOG_SCOPE(TEXT("Reduce damage mod during by 1 second"))
+		ABILITY_LOG_SCOPE(TEXT("Reduce damage mod during by 1 second"))
 
 		UGameplayEffect * DurationEffect = Cast<UGameplayEffect>(StaticConstructObject(UGameplayEffect::StaticClass(), GetTransientPackage(), FName(TEXT("Duration Debuff"))));
 		DurationEffect->Modifiers.SetNum(1);
@@ -774,7 +774,7 @@ bool GameplayEffectsTest_DurationBuff(UWorld *World, FAutomationTestBase * Test)
 
 	{
 		float ExpectedHealthValue = StartHealth;
-		float ActualHealthValue = DestComponent->GetSet<USkillSystemTestAttributeSet>()->Health;
+		float ActualHealthValue = DestComponent->GetSet<UAbilitySystemTestAttributeSet>()->Health;
 
 		Test->TestTrue(SKILL_TEST_TEXT("Damage Applied. %.2f == %.2f", ActualHealthValue, ExpectedHealthValue), ActualHealthValue == ExpectedHealthValue);
 	}
@@ -783,7 +783,7 @@ bool GameplayEffectsTest_DurationBuff(UWorld *World, FAutomationTestBase * Test)
 	// Duration Buff
 	//
 	{
-		SKILL_LOG_SCOPE(TEXT("Apply Damage mod that lasts 2 seconds"))
+		ABILITY_LOG_SCOPE(TEXT("Apply Damage mod that lasts 2 seconds"))
 
 		UGameplayEffect * BaseDmgEffect = Cast<UGameplayEffect>(StaticConstructObject(UGameplayEffect::StaticClass(), GetTransientPackage(), FName(TEXT("Damage"))));
 		BaseDmgEffect->Modifiers.SetNum(1);
@@ -798,7 +798,7 @@ bool GameplayEffectsTest_DurationBuff(UWorld *World, FAutomationTestBase * Test)
 		AppliedDamageHandle = SourceComponent->ApplyGameplayEffectToTarget(BaseDmgEffect, DestComponent, 1.f);
 
 		float ExpectedHealthValue = (StartHealth + (DamageValue));
-		float ActualHealthValue = DestComponent->GetSet<USkillSystemTestAttributeSet>()->Health;
+		float ActualHealthValue = DestComponent->GetSet<UAbilitySystemTestAttributeSet>()->Health;
 		float ExpectedDuration = BaseDuration;
 		float ActualDuration = DestComponent->GetGameplayEffectDuration(AppliedDamageHandle);
 
@@ -812,7 +812,7 @@ bool GameplayEffectsTest_DurationBuff(UWorld *World, FAutomationTestBase * Test)
 
 	// Buff the duration of the effect
 	{
-		SKILL_LOG_SCOPE(TEXT("Increase damage mod during by 1 second"))
+		ABILITY_LOG_SCOPE(TEXT("Increase damage mod during by 1 second"))
 
 		UGameplayEffect * DurationEffect = Cast<UGameplayEffect>(StaticConstructObject(UGameplayEffect::StaticClass(), GetTransientPackage(), FName(TEXT("Duration Buff"))));
 		DurationEffect->Modifiers.SetNum(1);
@@ -839,7 +839,7 @@ bool GameplayEffectsTest_DurationBuff(UWorld *World, FAutomationTestBase * Test)
 
 	{
 		float ExpectedHealthValue = (StartHealth + (DamageValue));
-		float ActualHealthValue = DestComponent->GetSet<USkillSystemTestAttributeSet>()->Health;
+		float ActualHealthValue = DestComponent->GetSet<UAbilitySystemTestAttributeSet>()->Health;
 
 		Test->TestTrue(SKILL_TEST_TEXT("Damage Applied. %.2f == %.2f", ActualHealthValue, ExpectedHealthValue), ActualHealthValue == ExpectedHealthValue);
 	}
@@ -849,7 +849,7 @@ bool GameplayEffectsTest_DurationBuff(UWorld *World, FAutomationTestBase * Test)
 
 	{
 		float ExpectedHealthValue = StartHealth;
-		float ActualHealthValue = DestComponent->GetSet<USkillSystemTestAttributeSet>()->Health;
+		float ActualHealthValue = DestComponent->GetSet<UAbilitySystemTestAttributeSet>()->Health;
 
 		Test->TestTrue(SKILL_TEST_TEXT("Damage Applied. %.2f == %.2f", ActualHealthValue, ExpectedHealthValue), ActualHealthValue == ExpectedHealthValue);
 	}
@@ -859,7 +859,7 @@ bool GameplayEffectsTest_DurationBuff(UWorld *World, FAutomationTestBase * Test)
 	// Removing Duration buff
 	//
 	{
-		SKILL_LOG_SCOPE(TEXT("Apply Damage mod that lasts 2 seconds"))
+		ABILITY_LOG_SCOPE(TEXT("Apply Damage mod that lasts 2 seconds"))
 
 		UGameplayEffect * BaseDmgEffect = Cast<UGameplayEffect>(StaticConstructObject(UGameplayEffect::StaticClass(), GetTransientPackage(), FName(TEXT("Damage"))));
 		BaseDmgEffect->Modifiers.SetNum(1);
@@ -874,7 +874,7 @@ bool GameplayEffectsTest_DurationBuff(UWorld *World, FAutomationTestBase * Test)
 		AppliedDamageHandle = SourceComponent->ApplyGameplayEffectToTarget(BaseDmgEffect, DestComponent, 1.f);
 
 		float ExpectedHealthValue = (StartHealth + (DamageValue));
-		float ActualHealthValue = DestComponent->GetSet<USkillSystemTestAttributeSet>()->Health;
+		float ActualHealthValue = DestComponent->GetSet<UAbilitySystemTestAttributeSet>()->Health;
 		float ExpectedDuration = BaseDuration;
 		float ActualDuration = DestComponent->GetGameplayEffectDuration(AppliedDamageHandle);
 
@@ -888,7 +888,7 @@ bool GameplayEffectsTest_DurationBuff(UWorld *World, FAutomationTestBase * Test)
 
 	// Buff the duration of the effect
 	{
-		SKILL_LOG_SCOPE(TEXT("Increase damage mod during by 1 second"))
+		ABILITY_LOG_SCOPE(TEXT("Increase damage mod during by 1 second"))
 
 		UGameplayEffect * DurationEffect = Cast<UGameplayEffect>(StaticConstructObject(UGameplayEffect::StaticClass(), GetTransientPackage(), FName(TEXT("Duration Buff"))));
 		DurationEffect->Modifiers.SetNum(1);
@@ -920,7 +920,7 @@ bool GameplayEffectsTest_DurationBuff(UWorld *World, FAutomationTestBase * Test)
 		float ExpectedDuration = BaseDuration;
 		float ActualDuration = DestComponent->GetGameplayEffectDuration(AppliedDamageHandle);
 		float ExpectedHealthValue = StartHealth + DamageValue;
-		float ActualHealthValue = DestComponent->GetSet<USkillSystemTestAttributeSet>()->Health;
+		float ActualHealthValue = DestComponent->GetSet<UAbilitySystemTestAttributeSet>()->Health;
 
 		Test->TestTrue(SKILL_TEST_TEXT("Duration of GameplayEffect Post Mod Remove. %.2f == %.2f", ActualDuration, ExpectedDuration), ActualDuration == ExpectedDuration);
 		Test->TestTrue(SKILL_TEST_TEXT("Damage Applied. %.2f == %.2f", ActualHealthValue, ExpectedHealthValue), ActualHealthValue == ExpectedHealthValue);
@@ -932,7 +932,7 @@ bool GameplayEffectsTest_DurationBuff(UWorld *World, FAutomationTestBase * Test)
 
 	{
 		float ExpectedHealthValue = StartHealth;
-		float ActualHealthValue = DestComponent->GetSet<USkillSystemTestAttributeSet>()->Health;
+		float ActualHealthValue = DestComponent->GetSet<UAbilitySystemTestAttributeSet>()->Health;
 
 		Test->TestTrue(SKILL_TEST_TEXT("Damage Applied. %.2f == %.2f", ActualHealthValue, ExpectedHealthValue), ActualHealthValue == ExpectedHealthValue);
 	}
@@ -954,15 +954,15 @@ bool GameplayEffectsTest_DamageBuffBuff_Basic(UWorld *World, FAutomationTestBase
 	const float DamageBuffMultiplier = 2.f;			// Damage is buff and multiplied by 2
 	const float DamageBuffMultiplierBonus = 1.f;	// The above multiplier receives a +1 bonus (we expect a final multiplier of 3)
 
-	ASkillSystemTestPawn *SourceActor = World->SpawnActor<ASkillSystemTestPawn>();
-	ASkillSystemTestPawn *DestActor = World->SpawnActor<ASkillSystemTestPawn>();
+	AAbilitySystemTestPawn *SourceActor = World->SpawnActor<AAbilitySystemTestPawn>();
+	AAbilitySystemTestPawn *DestActor = World->SpawnActor<AAbilitySystemTestPawn>();
 
-	UProperty *HealthProperty = FindFieldChecked<UProperty>(USkillSystemTestAttributeSet::StaticClass(), GET_MEMBER_NAME_CHECKED(USkillSystemTestAttributeSet, Health));
+	UProperty *HealthProperty = FindFieldChecked<UProperty>(UAbilitySystemTestAttributeSet::StaticClass(), GET_MEMBER_NAME_CHECKED(UAbilitySystemTestAttributeSet, Health));
 
 	UAttributeComponent * SourceComponent = SourceActor->AttributeComponent;
 	UAttributeComponent * DestComponent = DestActor->AttributeComponent;
-	SourceComponent->GetSet<USkillSystemTestAttributeSet>()->Health = StartHealth;
-	DestComponent->GetSet<USkillSystemTestAttributeSet>()->Health = StartHealth;
+	SourceComponent->GetSet<UAbilitySystemTestAttributeSet>()->Health = StartHealth;
+	DestComponent->GetSet<UAbilitySystemTestAttributeSet>()->Health = StartHealth;
 
 	// Apply "Damage" but set to INFINITE_DURATION
 	// (An odd example for damage, but would make sense for something like run speed, etc)
@@ -971,7 +971,7 @@ bool GameplayEffectsTest_DamageBuffBuff_Basic(UWorld *World, FAutomationTestBase
 	FActiveGameplayEffectHandle DamageHandle;
 
 	{
-		SKILL_LOG_SCOPE(TEXT("Apply Buff Buff"))
+		ABILITY_LOG_SCOPE(TEXT("Apply Buff Buff"))
 
 		// Here we are choosing to do this by adding an perm IncomingGE buff first. There are other ways to do this.
 
@@ -991,7 +991,7 @@ bool GameplayEffectsTest_DamageBuffBuff_Basic(UWorld *World, FAutomationTestBase
 	}
 
 	{
-		SKILL_LOG_SCOPE(TEXT("Apply Damage Buff"))
+		ABILITY_LOG_SCOPE(TEXT("Apply Damage Buff"))
 
 		UGameplayEffect * BuffEffect = Cast<UGameplayEffect>(StaticConstructObject(UGameplayEffect::StaticClass(), GetTransientPackage(), FName(TEXT("Buff"))));
 		BuffEffect->Modifiers.SetNum(1);
@@ -1014,7 +1014,7 @@ bool GameplayEffectsTest_DamageBuffBuff_Basic(UWorld *World, FAutomationTestBase
 	}
 
 	{
-		SKILL_LOG_SCOPE(TEXT("Apply Permanent (infinite duration) Damage"))
+		ABILITY_LOG_SCOPE(TEXT("Apply Permanent (infinite duration) Damage"))
 
 		UGameplayEffect * BaseDmgEffect = Cast<UGameplayEffect>(StaticConstructObject(UGameplayEffect::StaticClass(), GetTransientPackage(), FName(TEXT("Damage"))));
 		BaseDmgEffect->Modifiers.SetNum(1);
@@ -1029,7 +1029,7 @@ bool GameplayEffectsTest_DamageBuffBuff_Basic(UWorld *World, FAutomationTestBase
 		DamageHandle = SourceComponent->ApplyGameplayEffectToTarget(BaseDmgEffect, DestComponent, 1.f);
 
 		float ExpectedValue = (StartHealth + (DamageValue * (DamageBuffMultiplier + DamageBuffMultiplierBonus)));
-		float ActualValue = DestComponent->GetSet<USkillSystemTestAttributeSet>()->Health;
+		float ActualValue = DestComponent->GetSet<UAbilitySystemTestAttributeSet>()->Health;
 
 		Test->TestTrue(SKILL_TEST_TEXT("Damaged Applied. ActualValue: %.2f. ExpectedValue: %.2f.", ActualValue, ExpectedValue), (ActualValue == ExpectedValue));
 		Test->TestTrue(SKILL_TEST_TEXT("NumberOfActive Dest GameplayEffects %d == 1", DestComponent->GetNumActiveGameplayEffect()), DestComponent->GetNumActiveGameplayEffect() == 1);
@@ -1043,7 +1043,7 @@ bool GameplayEffectsTest_DamageBuffBuff_Basic(UWorld *World, FAutomationTestBase
 		// Remove the original buff buff
 		bool RemovedEffect = SourceComponent->RemoveActiveGameplayEffect(BuffBuffHandle);
 
-		float ActualValue = DestComponent->GetSet<USkillSystemTestAttributeSet>()->Health;
+		float ActualValue = DestComponent->GetSet<UAbilitySystemTestAttributeSet>()->Health;
 		float ExpectedValue = (StartHealth + (DamageValue * (DamageBuffMultiplier + DamageBuffMultiplierBonus)));
 
 		Test->TestTrue(SKILL_TEST_TEXT("Damaged Applied. ActualValue: %.2f. ExpectedValue: %.2f.", ActualValue, ExpectedValue), (ActualValue == ExpectedValue));
@@ -1055,7 +1055,7 @@ bool GameplayEffectsTest_DamageBuffBuff_Basic(UWorld *World, FAutomationTestBase
 	}
 
 	{
-		SKILL_LOG_SCOPE(TEXT("Remove Buff"));
+		ABILITY_LOG_SCOPE(TEXT("Remove Buff"));
 #if SKILL_SYSTEM_AGGREGATOR_DEBUG
 		FAggregator::AllocationStats.DependantsUpdated = 0;
 #endif
@@ -1064,7 +1064,7 @@ bool GameplayEffectsTest_DamageBuffBuff_Basic(UWorld *World, FAutomationTestBase
 
 		// No change to health since we applied a snapshot of the buff to the damage GE
 		float ExpectedValue = (StartHealth + (DamageValue * (DamageBuffMultiplier + DamageBuffMultiplierBonus)));
-		float ActualValue = DestComponent->GetSet<USkillSystemTestAttributeSet>()->Health;
+		float ActualValue = DestComponent->GetSet<UAbilitySystemTestAttributeSet>()->Health;
 
 		Test->TestTrue(SKILL_TEST_TEXT("Damaged Applied. ActualValue: %.2f. ExpectedValue: %.2f.", ActualValue, ExpectedValue), (ActualValue == ExpectedValue));
 
@@ -1095,15 +1095,15 @@ bool GameplayEffectsTest_DamageBuffBuff_FullLink(UWorld *World, FAutomationTestB
 	const float DamageBuffMultiplier = 2.f;			// Damage is buff and multiplied by 2
 	const float DamageBuffMultiplierBonus = 1.f;	// The above multiplier receives a +1 bonus (we expect a final multiplier of 3)
 
-	ASkillSystemTestPawn *SourceActor = World->SpawnActor<ASkillSystemTestPawn>();
-	ASkillSystemTestPawn *DestActor = World->SpawnActor<ASkillSystemTestPawn>();
+	AAbilitySystemTestPawn *SourceActor = World->SpawnActor<AAbilitySystemTestPawn>();
+	AAbilitySystemTestPawn *DestActor = World->SpawnActor<AAbilitySystemTestPawn>();
 
-	UProperty *HealthProperty = FindFieldChecked<UProperty>(USkillSystemTestAttributeSet::StaticClass(), GET_MEMBER_NAME_CHECKED(USkillSystemTestAttributeSet, Health));
+	UProperty *HealthProperty = FindFieldChecked<UProperty>(UAbilitySystemTestAttributeSet::StaticClass(), GET_MEMBER_NAME_CHECKED(UAbilitySystemTestAttributeSet, Health));
 
 	UAttributeComponent * SourceComponent = SourceActor->AttributeComponent;
 	UAttributeComponent * DestComponent = DestActor->AttributeComponent;
-	SourceComponent->GetSet<USkillSystemTestAttributeSet>()->Health = StartHealth;
-	DestComponent->GetSet<USkillSystemTestAttributeSet>()->Health = StartHealth;
+	SourceComponent->GetSet<UAbilitySystemTestAttributeSet>()->Health = StartHealth;
+	DestComponent->GetSet<UAbilitySystemTestAttributeSet>()->Health = StartHealth;
 
 	// Apply "Damage" but set to INFINITE_DURATION
 	// (An odd example for damage, but would make sense for something like run speed, etc)
@@ -1112,7 +1112,7 @@ bool GameplayEffectsTest_DamageBuffBuff_FullLink(UWorld *World, FAutomationTestB
 	FActiveGameplayEffectHandle DamageHandle;
 
 	{
-		SKILL_LOG_SCOPE(TEXT("Apply Buff Buff"))
+		ABILITY_LOG_SCOPE(TEXT("Apply Buff Buff"))
 
 		// Here we are choosing to do this by adding an perm IncomingGE buff first. There are other ways to do this.
 
@@ -1132,7 +1132,7 @@ bool GameplayEffectsTest_DamageBuffBuff_FullLink(UWorld *World, FAutomationTestB
 	}
 
 	{
-		SKILL_LOG_SCOPE(TEXT("Apply Damage Buff"))
+		ABILITY_LOG_SCOPE(TEXT("Apply Damage Buff"))
 
 		UGameplayEffect * BuffEffect = Cast<UGameplayEffect>(StaticConstructObject(UGameplayEffect::StaticClass(), GetTransientPackage(), FName(TEXT("Buff"))));
 		BuffEffect->Modifiers.SetNum(1);
@@ -1155,7 +1155,7 @@ bool GameplayEffectsTest_DamageBuffBuff_FullLink(UWorld *World, FAutomationTestB
 	}
 
 	{
-		SKILL_LOG_SCOPE(TEXT("Apply Permanent (infinite duration) Damage"))
+		ABILITY_LOG_SCOPE(TEXT("Apply Permanent (infinite duration) Damage"))
 
 		UGameplayEffect * BaseDmgEffect = Cast<UGameplayEffect>(StaticConstructObject(UGameplayEffect::StaticClass(), GetTransientPackage(), FName(TEXT("Damage"))));
 		BaseDmgEffect->Modifiers.SetNum(1);
@@ -1170,7 +1170,7 @@ bool GameplayEffectsTest_DamageBuffBuff_FullLink(UWorld *World, FAutomationTestB
 		DamageHandle = SourceComponent->ApplyGameplayEffectToTarget(BaseDmgEffect, DestComponent, 1.f);
 
 		float ExpectedValue = (StartHealth + (DamageValue * (DamageBuffMultiplier + DamageBuffMultiplierBonus)));
-		float ActualValue = DestComponent->GetSet<USkillSystemTestAttributeSet>()->Health;
+		float ActualValue = DestComponent->GetSet<UAbilitySystemTestAttributeSet>()->Health;
 
 		Test->TestTrue(SKILL_TEST_TEXT("Damaged Applied. ActualValue: %.2f. ExpectedValue: %.2f.", ActualValue, ExpectedValue), (ActualValue == ExpectedValue));
 		Test->TestTrue(SKILL_TEST_TEXT("NumberOfActive Dest GameplayEffects %d == 1", DestComponent->GetNumActiveGameplayEffect()), DestComponent->GetNumActiveGameplayEffect() == 1);
@@ -1184,7 +1184,7 @@ bool GameplayEffectsTest_DamageBuffBuff_FullLink(UWorld *World, FAutomationTestB
 		// Remove the original buff buff
 		bool RemovedEffect = SourceComponent->RemoveActiveGameplayEffect(BuffBuffHandle);
 
-		float ActualValue = DestComponent->GetSet<USkillSystemTestAttributeSet>()->Health;
+		float ActualValue = DestComponent->GetSet<UAbilitySystemTestAttributeSet>()->Health;
 		float ExpectedValue = (StartHealth + (DamageValue * (DamageBuffMultiplier)));
 
 		Test->TestTrue(SKILL_TEST_TEXT("Damaged Applied. ActualValue: %.2f. ExpectedValue: %.2f.", ActualValue, ExpectedValue), (ActualValue == ExpectedValue));
@@ -1196,14 +1196,14 @@ bool GameplayEffectsTest_DamageBuffBuff_FullLink(UWorld *World, FAutomationTestB
 	}
 
 	{
-		SKILL_LOG_SCOPE(TEXT("Remove Buff"));
+		ABILITY_LOG_SCOPE(TEXT("Remove Buff"));
 #if SKILL_SYSTEM_AGGREGATOR_DEBUG
 		FAggregator::AllocationStats.DependantsUpdated = 0;
 #endif
 		// Remove the buff
 		bool RemovedEffect = SourceComponent->RemoveActiveGameplayEffect(BuffHandle);
 
-		float ActualValue = DestComponent->GetSet<USkillSystemTestAttributeSet>()->Health;
+		float ActualValue = DestComponent->GetSet<UAbilitySystemTestAttributeSet>()->Health;
 		float ExpectedValue = StartHealth + DamageValue;
 
 		Test->TestTrue(SKILL_TEST_TEXT("Damaged Applied. ActualValue: %.2f. ExpectedValue: %.2f.", ActualValue, ExpectedValue), (ActualValue == ExpectedValue));
@@ -1233,15 +1233,15 @@ bool GameplayEffectsTest_DamageBuffBuff_FullSnapshot(UWorld *World, FAutomationT
 	const float DamageBuffMultiplier = 2.f;			// Damage is buff and multiplied by 2
 	const float DamageBuffMultiplierBonus = 1.f;	// The above multiplier receives a +1 bonus (we expect a final multiplier of 3)
 
-	ASkillSystemTestPawn *SourceActor = World->SpawnActor<ASkillSystemTestPawn>();
-	ASkillSystemTestPawn *DestActor = World->SpawnActor<ASkillSystemTestPawn>();
+	AAbilitySystemTestPawn *SourceActor = World->SpawnActor<AAbilitySystemTestPawn>();
+	AAbilitySystemTestPawn *DestActor = World->SpawnActor<AAbilitySystemTestPawn>();
 
-	UProperty *HealthProperty = FindFieldChecked<UProperty>(USkillSystemTestAttributeSet::StaticClass(), GET_MEMBER_NAME_CHECKED(USkillSystemTestAttributeSet, Health));
+	UProperty *HealthProperty = FindFieldChecked<UProperty>(UAbilitySystemTestAttributeSet::StaticClass(), GET_MEMBER_NAME_CHECKED(UAbilitySystemTestAttributeSet, Health));
 
 	UAttributeComponent * SourceComponent = SourceActor->AttributeComponent;
 	UAttributeComponent * DestComponent = DestActor->AttributeComponent;
-	SourceComponent->GetSet<USkillSystemTestAttributeSet>()->Health = StartHealth;
-	DestComponent->GetSet<USkillSystemTestAttributeSet>()->Health = StartHealth;
+	SourceComponent->GetSet<UAbilitySystemTestAttributeSet>()->Health = StartHealth;
+	DestComponent->GetSet<UAbilitySystemTestAttributeSet>()->Health = StartHealth;
 
 	// Apply "Damage" but set to INFINITE_DURATION
 	// (An odd example for damage, but would make sense for something like run speed, etc)
@@ -1250,7 +1250,7 @@ bool GameplayEffectsTest_DamageBuffBuff_FullSnapshot(UWorld *World, FAutomationT
 	FActiveGameplayEffectHandle DamageHandle;
 
 	{
-		SKILL_LOG_SCOPE(TEXT("Apply Buff Buff"))
+		ABILITY_LOG_SCOPE(TEXT("Apply Buff Buff"))
 
 		// Here we are choosing to do this by adding an perm IncomingGE buff first. There are other ways to do this.
 
@@ -1270,7 +1270,7 @@ bool GameplayEffectsTest_DamageBuffBuff_FullSnapshot(UWorld *World, FAutomationT
 	}
 
 	{
-		SKILL_LOG_SCOPE(TEXT("Apply Damage Buff"))
+		ABILITY_LOG_SCOPE(TEXT("Apply Damage Buff"))
 
 		UGameplayEffect * BuffEffect = Cast<UGameplayEffect>(StaticConstructObject(UGameplayEffect::StaticClass(), GetTransientPackage(), FName(TEXT("Buff"))));
 		BuffEffect->Modifiers.SetNum(1);
@@ -1293,7 +1293,7 @@ bool GameplayEffectsTest_DamageBuffBuff_FullSnapshot(UWorld *World, FAutomationT
 	}
 
 	{
-		SKILL_LOG_SCOPE(TEXT("Apply Permanent (infinite duration) Damage"))
+		ABILITY_LOG_SCOPE(TEXT("Apply Permanent (infinite duration) Damage"))
 
 		UGameplayEffect * BaseDmgEffect = Cast<UGameplayEffect>(StaticConstructObject(UGameplayEffect::StaticClass(), GetTransientPackage(), FName(TEXT("Damage"))));
 		BaseDmgEffect->Modifiers.SetNum(1);
@@ -1308,7 +1308,7 @@ bool GameplayEffectsTest_DamageBuffBuff_FullSnapshot(UWorld *World, FAutomationT
 		DamageHandle = SourceComponent->ApplyGameplayEffectToTarget(BaseDmgEffect, DestComponent, 1.f);
 
 		float ExpectedValue = (StartHealth + (DamageValue * (DamageBuffMultiplier + DamageBuffMultiplierBonus)));
-		float ActualValue = DestComponent->GetSet<USkillSystemTestAttributeSet>()->Health;
+		float ActualValue = DestComponent->GetSet<UAbilitySystemTestAttributeSet>()->Health;
 
 		Test->TestTrue(SKILL_TEST_TEXT("Damaged Applied. ActualValue: %.2f. ExpectedValue: %.2f.", ActualValue, ExpectedValue), (ActualValue == ExpectedValue));
 		Test->TestTrue(SKILL_TEST_TEXT("NumberOfActive Dest GameplayEffects %d == 1", DestComponent->GetNumActiveGameplayEffect()), DestComponent->GetNumActiveGameplayEffect() == 1);
@@ -1322,7 +1322,7 @@ bool GameplayEffectsTest_DamageBuffBuff_FullSnapshot(UWorld *World, FAutomationT
 		// Remove the original buff buff
 		bool RemovedEffect = SourceComponent->RemoveActiveGameplayEffect(BuffBuffHandle);
 
-		float ActualValue = DestComponent->GetSet<USkillSystemTestAttributeSet>()->Health;
+		float ActualValue = DestComponent->GetSet<UAbilitySystemTestAttributeSet>()->Health;
 		float ExpectedValue = (StartHealth + (DamageValue * (DamageBuffMultiplier + DamageBuffMultiplierBonus)));
 
 		Test->TestTrue(SKILL_TEST_TEXT("Damaged Applied. ActualValue: %.2f. ExpectedValue: %.2f.", ActualValue, ExpectedValue), (ActualValue == ExpectedValue));
@@ -1334,14 +1334,14 @@ bool GameplayEffectsTest_DamageBuffBuff_FullSnapshot(UWorld *World, FAutomationT
 	}
 
 	{
-		SKILL_LOG_SCOPE(TEXT("Remove Buff"));
+		ABILITY_LOG_SCOPE(TEXT("Remove Buff"));
 #if SKILL_SYSTEM_AGGREGATOR_DEBUG
 		FAggregator::AllocationStats.DependantsUpdated = 0;
 #endif
 		// Remove the buff
 		bool RemovedEffect = SourceComponent->RemoveActiveGameplayEffect(BuffHandle);
 
-		float ActualValue = DestComponent->GetSet<USkillSystemTestAttributeSet>()->Health;
+		float ActualValue = DestComponent->GetSet<UAbilitySystemTestAttributeSet>()->Health;
 		float ExpectedValue = (StartHealth + (DamageValue * (DamageBuffMultiplier + DamageBuffMultiplierBonus)));
 
 		Test->TestTrue(SKILL_TEST_TEXT("Damaged Applied. ActualValue: %.2f. ExpectedValue: %.2f.", ActualValue, ExpectedValue), (ActualValue == ExpectedValue));
@@ -1371,15 +1371,15 @@ bool GameplayEffectsTest_DamageBuffBuff_SnapshotLink(UWorld *World, FAutomationT
 	const float DamageBuffMultiplier = 2.f;			// Damage is buff and multiplied by 2
 	const float DamageBuffMultiplierBonus = 1.f;	// The above multiplier receives a +1 bonus (we expect a final multiplier of 3)
 
-	ASkillSystemTestPawn *SourceActor = World->SpawnActor<ASkillSystemTestPawn>();
-	ASkillSystemTestPawn *DestActor = World->SpawnActor<ASkillSystemTestPawn>();
+	AAbilitySystemTestPawn *SourceActor = World->SpawnActor<AAbilitySystemTestPawn>();
+	AAbilitySystemTestPawn *DestActor = World->SpawnActor<AAbilitySystemTestPawn>();
 
-	UProperty *HealthProperty = FindFieldChecked<UProperty>(USkillSystemTestAttributeSet::StaticClass(), GET_MEMBER_NAME_CHECKED(USkillSystemTestAttributeSet, Health));
+	UProperty *HealthProperty = FindFieldChecked<UProperty>(UAbilitySystemTestAttributeSet::StaticClass(), GET_MEMBER_NAME_CHECKED(UAbilitySystemTestAttributeSet, Health));
 
 	UAttributeComponent * SourceComponent = SourceActor->AttributeComponent;
 	UAttributeComponent * DestComponent = DestActor->AttributeComponent;
-	SourceComponent->GetSet<USkillSystemTestAttributeSet>()->Health = StartHealth;
-	DestComponent->GetSet<USkillSystemTestAttributeSet>()->Health = StartHealth;
+	SourceComponent->GetSet<UAbilitySystemTestAttributeSet>()->Health = StartHealth;
+	DestComponent->GetSet<UAbilitySystemTestAttributeSet>()->Health = StartHealth;
 
 	// Apply "Damage" but set to INFINITE_DURATION
 	// (An odd example for damage, but would make sense for something like run speed, etc)
@@ -1388,7 +1388,7 @@ bool GameplayEffectsTest_DamageBuffBuff_SnapshotLink(UWorld *World, FAutomationT
 	FActiveGameplayEffectHandle DamageHandle;
 
 	{
-		SKILL_LOG_SCOPE(TEXT("Apply Buff Buff"))
+		ABILITY_LOG_SCOPE(TEXT("Apply Buff Buff"))
 
 		// Here we are choosing to do this by adding an perm IncomingGE buff first. There are other ways to do this.
 
@@ -1408,7 +1408,7 @@ bool GameplayEffectsTest_DamageBuffBuff_SnapshotLink(UWorld *World, FAutomationT
 	}
 
 	{
-		SKILL_LOG_SCOPE(TEXT("Apply Damage Buff"))
+		ABILITY_LOG_SCOPE(TEXT("Apply Damage Buff"))
 
 		UGameplayEffect * BuffEffect = Cast<UGameplayEffect>(StaticConstructObject(UGameplayEffect::StaticClass(), GetTransientPackage(), FName(TEXT("Buff"))));
 		BuffEffect->Modifiers.SetNum(1);
@@ -1431,7 +1431,7 @@ bool GameplayEffectsTest_DamageBuffBuff_SnapshotLink(UWorld *World, FAutomationT
 	}
 
 	{
-		SKILL_LOG_SCOPE(TEXT("Apply Permanent (infinite duration) Damage"))
+		ABILITY_LOG_SCOPE(TEXT("Apply Permanent (infinite duration) Damage"))
 
 		UGameplayEffect * BaseDmgEffect = Cast<UGameplayEffect>(StaticConstructObject(UGameplayEffect::StaticClass(), GetTransientPackage(), FName(TEXT("Damage"))));
 		BaseDmgEffect->Modifiers.SetNum(1);
@@ -1446,21 +1446,21 @@ bool GameplayEffectsTest_DamageBuffBuff_SnapshotLink(UWorld *World, FAutomationT
 		DamageHandle = SourceComponent->ApplyGameplayEffectToTarget(BaseDmgEffect, DestComponent, 1.f);
 
 		float ExpectedValue = (StartHealth + (DamageValue * (DamageBuffMultiplier + DamageBuffMultiplierBonus)));
-		float ActualValue = DestComponent->GetSet<USkillSystemTestAttributeSet>()->Health;
+		float ActualValue = DestComponent->GetSet<UAbilitySystemTestAttributeSet>()->Health;
 
 		Test->TestTrue(SKILL_TEST_TEXT("Damaged Applied. ActualValue: %.2f. ExpectedValue: %.2f.", ActualValue, ExpectedValue), (ActualValue == ExpectedValue));
 		Test->TestTrue(SKILL_TEST_TEXT("NumberOfActive Dest GameplayEffects %d == 1", DestComponent->GetNumActiveGameplayEffect()), DestComponent->GetNumActiveGameplayEffect() == 1);
 	}
 
 	{
-		SKILL_LOG_SCOPE(TEXT("Remove Buff Buff"));
+		ABILITY_LOG_SCOPE(TEXT("Remove Buff Buff"));
 #if SKILL_SYSTEM_AGGREGATOR_DEBUG
 		FAggregator::AllocationStats.DependantsUpdated = 0;
 #endif
 		// Remove the original buff buff
 		bool RemovedEffect = SourceComponent->RemoveActiveGameplayEffect(BuffBuffHandle);
 
-		float ActualValue = DestComponent->GetSet<USkillSystemTestAttributeSet>()->Health;
+		float ActualValue = DestComponent->GetSet<UAbilitySystemTestAttributeSet>()->Health;
 		float ExpectedValue = (StartHealth + (DamageValue * (DamageBuffMultiplier + DamageBuffMultiplierBonus)));
 
 		Test->TestTrue(SKILL_TEST_TEXT("Damaged Applied. ActualValue: %.2f. ExpectedValue: %.2f.", ActualValue, ExpectedValue), (ActualValue == ExpectedValue));
@@ -1472,14 +1472,14 @@ bool GameplayEffectsTest_DamageBuffBuff_SnapshotLink(UWorld *World, FAutomationT
 	}
 
 	{
-		SKILL_LOG_SCOPE(TEXT("Remove Buff"));
+		ABILITY_LOG_SCOPE(TEXT("Remove Buff"));
 #if SKILL_SYSTEM_AGGREGATOR_DEBUG
 		FAggregator::AllocationStats.DependantsUpdated = 0;
 #endif
 		// Remove the buff
 		bool RemovedEffect = SourceComponent->RemoveActiveGameplayEffect(BuffHandle);
 
-		float ActualValue = DestComponent->GetSet<USkillSystemTestAttributeSet>()->Health;
+		float ActualValue = DestComponent->GetSet<UAbilitySystemTestAttributeSet>()->Health;
 		float ExpectedValue = StartHealth + DamageValue;
 
 		Test->TestTrue(SKILL_TEST_TEXT("Damaged Applied. ActualValue: %.2f. ExpectedValue: %.2f.", ActualValue, ExpectedValue), (ActualValue == ExpectedValue));
@@ -1505,19 +1505,19 @@ bool GameplayEffectsTest_DamageAppliesBuff(UWorld *World, FAutomationTestBase *
 	const float DamageValue = 5.f;
 	const float DamageProtectionDivisor = 2.f;
 
-	ASkillSystemTestPawn *SourceActor = World->SpawnActor<ASkillSystemTestPawn>();
-	ASkillSystemTestPawn *DestActor = World->SpawnActor<ASkillSystemTestPawn>();
+	AAbilitySystemTestPawn *SourceActor = World->SpawnActor<AAbilitySystemTestPawn>();
+	AAbilitySystemTestPawn *DestActor = World->SpawnActor<AAbilitySystemTestPawn>();
 
-	UProperty *DamageProperty = FindFieldChecked<UProperty>(USkillSystemTestAttributeSet::StaticClass(), GET_MEMBER_NAME_CHECKED(USkillSystemTestAttributeSet, Damage));
+	UProperty *DamageProperty = FindFieldChecked<UProperty>(UAbilitySystemTestAttributeSet::StaticClass(), GET_MEMBER_NAME_CHECKED(UAbilitySystemTestAttributeSet, Damage));
 
 	UAttributeComponent * SourceComponent = SourceActor->AttributeComponent;
 	UAttributeComponent * DestComponent = DestActor->AttributeComponent;
-	SourceComponent->GetSet<USkillSystemTestAttributeSet>()->Health = StartHealth;
-	DestComponent->GetSet<USkillSystemTestAttributeSet>()->Health = StartHealth;
+	SourceComponent->GetSet<UAbilitySystemTestAttributeSet>()->Health = StartHealth;
+	DestComponent->GetSet<UAbilitySystemTestAttributeSet>()->Health = StartHealth;
 
 	// Apply damage and a buff that reduces incoming damage
 	{
-		SKILL_LOG_SCOPE(TEXT("Apply DamageBuff and InstantDamage"))
+		ABILITY_LOG_SCOPE(TEXT("Apply DamageBuff and InstantDamage"))
 
 		UGameplayEffect * BuffEffect = Cast<UGameplayEffect>(StaticConstructObject(UGameplayEffect::StaticClass(), GetTransientPackage(), FName(TEXT("DamageBuff"))));
 		BuffEffect->Modifiers.SetNum(1);
@@ -1542,10 +1542,10 @@ bool GameplayEffectsTest_DamageAppliesBuff(UWorld *World, FAutomationTestBase *
 		SourceComponent->ApplyGameplayEffectToTarget(BaseDmgEffect, DestComponent, 1.f);
 
 		float ExpectedValue = (StartHealth - DamageValue);
-		float ActualValue = DestComponent->GetSet<USkillSystemTestAttributeSet>()->Health;
+		float ActualValue = DestComponent->GetSet<UAbilitySystemTestAttributeSet>()->Health;
 
 		Test->TestTrue(SKILL_TEST_TEXT("Buff Instant Damage Applied"), (ActualValue == ExpectedValue));
-		SKILL_LOG(Log, TEXT("Final Health: %.2f"), DestComponent->GetSet<USkillSystemTestAttributeSet>()->Health);
+		ABILITY_LOG(Log, TEXT("Final Health: %.2f"), DestComponent->GetSet<UAbilitySystemTestAttributeSet>()->Health);
 	}
 
 	// Apply Damage
@@ -1563,10 +1563,10 @@ bool GameplayEffectsTest_DamageAppliesBuff(UWorld *World, FAutomationTestBase *
 		SourceComponent->ApplyGameplayEffectToTarget(BaseDmgEffect, DestComponent, 1.f);
 
 		float ExpectedValue = (StartHealth - DamageValue - (DamageValue / DamageProtectionDivisor));
-		float ActualValue = DestComponent->GetSet<USkillSystemTestAttributeSet>()->Health;
+		float ActualValue = DestComponent->GetSet<UAbilitySystemTestAttributeSet>()->Health;
 
 		Test->TestTrue(SKILL_TEST_TEXT("Buff Instant Damage Applied"), (ActualValue == ExpectedValue));
-		SKILL_LOG(Log, TEXT("Final Health: %.2f"), DestComponent->GetSet<USkillSystemTestAttributeSet>()->Health);
+		ABILITY_LOG(Log, TEXT("Final Health: %.2f"), DestComponent->GetSet<UAbilitySystemTestAttributeSet>()->Health);
 	}
 
 	World->EditorDestroyActor(SourceActor, false);
@@ -1582,19 +1582,19 @@ bool GameplayEffectsTest_BuffAppliesBuff(UWorld *World, FAutomationTestBase * Te
 	const float DamageValue = 5.f;
 	const float DamageProtectionDivisor = 2.f;
 
-	ASkillSystemTestPawn *SourceActor = World->SpawnActor<ASkillSystemTestPawn>();
-	ASkillSystemTestPawn *DestActor = World->SpawnActor<ASkillSystemTestPawn>();
+	AAbilitySystemTestPawn *SourceActor = World->SpawnActor<AAbilitySystemTestPawn>();
+	AAbilitySystemTestPawn *DestActor = World->SpawnActor<AAbilitySystemTestPawn>();
 
-	UProperty *DamageProperty = FindFieldChecked<UProperty>(USkillSystemTestAttributeSet::StaticClass(), GET_MEMBER_NAME_CHECKED(USkillSystemTestAttributeSet, Damage));
+	UProperty *DamageProperty = FindFieldChecked<UProperty>(UAbilitySystemTestAttributeSet::StaticClass(), GET_MEMBER_NAME_CHECKED(UAbilitySystemTestAttributeSet, Damage));
 
 	UAttributeComponent * SourceComponent = SourceActor->AttributeComponent;
 	UAttributeComponent * DestComponent = DestActor->AttributeComponent;
-	SourceComponent->GetSet<USkillSystemTestAttributeSet>()->Health = StartHealth;
-	DestComponent->GetSet<USkillSystemTestAttributeSet>()->Health = StartHealth;
+	SourceComponent->GetSet<UAbilitySystemTestAttributeSet>()->Health = StartHealth;
+	DestComponent->GetSet<UAbilitySystemTestAttributeSet>()->Health = StartHealth;
 
 	// Apply damage and a buff that reduces incoming damage
 	{
-		SKILL_LOG_SCOPE(TEXT("Apply DamageBuff"))
+		ABILITY_LOG_SCOPE(TEXT("Apply DamageBuff"))
 
 		UGameplayEffect * BuffEffect = Cast<UGameplayEffect>(StaticConstructObject(UGameplayEffect::StaticClass(), GetTransientPackage(), FName(TEXT("DamageBuff"))));
 		BuffEffect->Modifiers.SetNum(1);
@@ -1632,10 +1632,10 @@ bool GameplayEffectsTest_BuffAppliesBuff(UWorld *World, FAutomationTestBase * Te
 		DestComponent->ApplyGameplayEffectToTarget(BaseDmgEffect, SourceComponent, 1.f);
 
 		float ExpectedValue = StartHealth - DamageValue;
-		float ActualValue = SourceComponent->GetSet<USkillSystemTestAttributeSet>()->Health;
+		float ActualValue = SourceComponent->GetSet<UAbilitySystemTestAttributeSet>()->Health;
 
 		Test->TestTrue(SKILL_TEST_TEXT("Sending buffs test"), (ActualValue == ExpectedValue));
-		SKILL_LOG(Log, TEXT("Final Health: %.2f"), SourceComponent->GetSet<USkillSystemTestAttributeSet>()->Health);
+		ABILITY_LOG(Log, TEXT("Final Health: %.2f"), SourceComponent->GetSet<UAbilitySystemTestAttributeSet>()->Health);
 	}
 
 	// Apply Damage
@@ -1653,10 +1653,10 @@ bool GameplayEffectsTest_BuffAppliesBuff(UWorld *World, FAutomationTestBase * Te
 		SourceComponent->ApplyGameplayEffectToTarget(BaseDmgEffect, DestComponent, 1.f);
 
 		float ExpectedValue = StartHealth - DamageValue;
-		float ActualValue = DestComponent->GetSet<USkillSystemTestAttributeSet>()->Health;
+		float ActualValue = DestComponent->GetSet<UAbilitySystemTestAttributeSet>()->Health;
 
 		Test->TestTrue(SKILL_TEST_TEXT("Buff Instant Damage Applied"), (ActualValue == ExpectedValue));
-		SKILL_LOG(Log, TEXT("Final Health: %.2f"), DestComponent->GetSet<USkillSystemTestAttributeSet>()->Health);
+		ABILITY_LOG(Log, TEXT("Final Health: %.2f"), DestComponent->GetSet<UAbilitySystemTestAttributeSet>()->Health);
 	}
 
 	// Apply Damage
@@ -1674,10 +1674,10 @@ bool GameplayEffectsTest_BuffAppliesBuff(UWorld *World, FAutomationTestBase * Te
 		SourceComponent->ApplyGameplayEffectToTarget(BaseDmgEffect, DestComponent, 1.f);
 
 		float ExpectedValue = (StartHealth - DamageValue - (DamageValue / DamageProtectionDivisor));
-		float ActualValue = DestComponent->GetSet<USkillSystemTestAttributeSet>()->Health;
+		float ActualValue = DestComponent->GetSet<UAbilitySystemTestAttributeSet>()->Health;
 
 		Test->TestTrue(SKILL_TEST_TEXT("Instant Damage Applied"), (ActualValue == ExpectedValue));
-		SKILL_LOG(Log, TEXT("Final Health: %.2f"), DestComponent->GetSet<USkillSystemTestAttributeSet>()->Health);
+		ABILITY_LOG(Log, TEXT("Final Health: %.2f"), DestComponent->GetSet<UAbilitySystemTestAttributeSet>()->Health);
 	}
 
 	// Apply Damage again to make sure that the buff only applied once
@@ -1695,10 +1695,10 @@ bool GameplayEffectsTest_BuffAppliesBuff(UWorld *World, FAutomationTestBase * Te
 		SourceComponent->ApplyGameplayEffectToTarget(BaseDmgEffect, DestComponent, 1.f);
 
 		float ExpectedValue = (StartHealth - DamageValue - (2 * DamageValue / DamageProtectionDivisor));
-		float ActualValue = DestComponent->GetSet<USkillSystemTestAttributeSet>()->Health;
+		float ActualValue = DestComponent->GetSet<UAbilitySystemTestAttributeSet>()->Health;
 
 		Test->TestTrue(SKILL_TEST_TEXT("Instant Damage Applied"), (ActualValue == ExpectedValue));
-		SKILL_LOG(Log, TEXT("Final Health: %.2f"), DestComponent->GetSet<USkillSystemTestAttributeSet>()->Health);
+		ABILITY_LOG(Log, TEXT("Final Health: %.2f"), DestComponent->GetSet<UAbilitySystemTestAttributeSet>()->Health);
 	}
 
 	World->EditorDestroyActor(SourceActor, false);
@@ -1713,19 +1713,19 @@ bool GameplayEffectsTest_BuffIndirection(UWorld *World, FAutomationTestBase * Te
 	const float DamageValue = 5.f;
 	const float DamageProtectionDivisor = 2.f;
 
-	ASkillSystemTestPawn *SourceActor = World->SpawnActor<ASkillSystemTestPawn>();
-	ASkillSystemTestPawn *DestActor = World->SpawnActor<ASkillSystemTestPawn>();
+	AAbilitySystemTestPawn *SourceActor = World->SpawnActor<AAbilitySystemTestPawn>();
+	AAbilitySystemTestPawn *DestActor = World->SpawnActor<AAbilitySystemTestPawn>();
 
-	UProperty *DamageProperty = FindFieldChecked<UProperty>(USkillSystemTestAttributeSet::StaticClass(), GET_MEMBER_NAME_CHECKED(USkillSystemTestAttributeSet, Damage));
+	UProperty *DamageProperty = FindFieldChecked<UProperty>(UAbilitySystemTestAttributeSet::StaticClass(), GET_MEMBER_NAME_CHECKED(UAbilitySystemTestAttributeSet, Damage));
 
 	UAttributeComponent * SourceComponent = SourceActor->AttributeComponent;
 	UAttributeComponent * DestComponent = DestActor->AttributeComponent;
-	SourceComponent->GetSet<USkillSystemTestAttributeSet>()->Health = StartHealth;
-	DestComponent->GetSet<USkillSystemTestAttributeSet>()->Health = StartHealth;
+	SourceComponent->GetSet<UAbilitySystemTestAttributeSet>()->Health = StartHealth;
+	DestComponent->GetSet<UAbilitySystemTestAttributeSet>()->Health = StartHealth;
 
 	// Apply damage and a buff that reduces incoming damage
 	{
-		SKILL_LOG_SCOPE(TEXT("Apply DamageBuff and InstantDamage"))
+		ABILITY_LOG_SCOPE(TEXT("Apply DamageBuff and InstantDamage"))
 
 		UGameplayEffect * BuffEffect = Cast<UGameplayEffect>(StaticConstructObject(UGameplayEffect::StaticClass(), GetTransientPackage(), FName(TEXT("DamageBuff"))));
 		BuffEffect->Modifiers.SetNum(1);
@@ -1770,10 +1770,10 @@ bool GameplayEffectsTest_BuffIndirection(UWorld *World, FAutomationTestBase * Te
 		SourceComponent->ApplyGameplayEffectToTarget(BaseDmgEffect, DestComponent, 1.f);
 
 		float ExpectedValue = StartHealth - DamageValue;
-		float ActualValue = DestComponent->GetSet<USkillSystemTestAttributeSet>()->Health;
+		float ActualValue = DestComponent->GetSet<UAbilitySystemTestAttributeSet>()->Health;
 
 		Test->TestTrue(SKILL_TEST_TEXT("Buff Instant Damage Applied"), (ActualValue == ExpectedValue));
-		SKILL_LOG(Log, TEXT("Final Health: %.2f"), DestComponent->GetSet<USkillSystemTestAttributeSet>()->Health);
+		ABILITY_LOG(Log, TEXT("Final Health: %.2f"), DestComponent->GetSet<UAbilitySystemTestAttributeSet>()->Health);
 	}
 
 	// Apply Damage
@@ -1791,10 +1791,10 @@ bool GameplayEffectsTest_BuffIndirection(UWorld *World, FAutomationTestBase * Te
 		DestComponent->ApplyGameplayEffectToTarget(BaseDmgEffect, SourceComponent, 1.f);
 
 		float ExpectedValue = StartHealth - DamageValue;
-		float ActualValue = SourceComponent->GetSet<USkillSystemTestAttributeSet>()->Health;
+		float ActualValue = SourceComponent->GetSet<UAbilitySystemTestAttributeSet>()->Health;
 
 		Test->TestTrue(SKILL_TEST_TEXT("Instant Damage Applied"), (ActualValue == ExpectedValue));
-		SKILL_LOG(Log, TEXT("Final Health: %.2f"), SourceComponent->GetSet<USkillSystemTestAttributeSet>()->Health);
+		ABILITY_LOG(Log, TEXT("Final Health: %.2f"), SourceComponent->GetSet<UAbilitySystemTestAttributeSet>()->Health);
 	}
 
 	// Apply Damage again to make sure that the buff only applied once
@@ -1812,10 +1812,10 @@ bool GameplayEffectsTest_BuffIndirection(UWorld *World, FAutomationTestBase * Te
 		DestComponent->ApplyGameplayEffectToTarget(BaseDmgEffect, SourceComponent, 1.f);
 
 		float ExpectedValue = StartHealth - DamageValue - (DamageValue / DamageProtectionDivisor);
-		float ActualValue = SourceComponent->GetSet<USkillSystemTestAttributeSet>()->Health;
+		float ActualValue = SourceComponent->GetSet<UAbilitySystemTestAttributeSet>()->Health;
 
 		Test->TestTrue(SKILL_TEST_TEXT("Instant Damage Applied"), (ActualValue == ExpectedValue));
-		SKILL_LOG(Log, TEXT("Final Health: %.2f"), SourceComponent->GetSet<USkillSystemTestAttributeSet>()->Health);
+		ABILITY_LOG(Log, TEXT("Final Health: %.2f"), SourceComponent->GetSet<UAbilitySystemTestAttributeSet>()->Health);
 	}
 
 	World->EditorDestroyActor(SourceActor, false);
@@ -1831,19 +1831,19 @@ bool GameplayEffectsTest_DurationDamage(UWorld *World, FAutomationTestBase * Tes
 	float Duration = 5.f;
 	float StartTime = World->GetTimeSeconds();
 
-	ASkillSystemTestPawn *SourceActor = World->SpawnActor<ASkillSystemTestPawn>();
-	ASkillSystemTestPawn *DestActor = World->SpawnActor<ASkillSystemTestPawn>();
+	AAbilitySystemTestPawn *SourceActor = World->SpawnActor<AAbilitySystemTestPawn>();
+	AAbilitySystemTestPawn *DestActor = World->SpawnActor<AAbilitySystemTestPawn>();
 
-	UProperty *HealthProperty = FindFieldChecked<UProperty>(USkillSystemTestAttributeSet::StaticClass(), GET_MEMBER_NAME_CHECKED(USkillSystemTestAttributeSet, Health));
+	UProperty *HealthProperty = FindFieldChecked<UProperty>(UAbilitySystemTestAttributeSet::StaticClass(), GET_MEMBER_NAME_CHECKED(UAbilitySystemTestAttributeSet, Health));
 
 	UAttributeComponent * SourceComponent = SourceActor->AttributeComponent;
 	UAttributeComponent * DestComponent = DestActor->AttributeComponent;
 
-	SourceComponent->GetSet<USkillSystemTestAttributeSet>()->Health = StartHealth;
-	DestComponent->GetSet<USkillSystemTestAttributeSet>()->Health = StartHealth;
+	SourceComponent->GetSet<UAbilitySystemTestAttributeSet>()->Health = StartHealth;
+	DestComponent->GetSet<UAbilitySystemTestAttributeSet>()->Health = StartHealth;
 
 	{
-		SKILL_LOG_SCOPE(TEXT("Apply Temporary Damage"));
+		ABILITY_LOG_SCOPE(TEXT("Apply Temporary Damage"));
 
 		UGameplayEffect * BaseDmgEffect = Cast<UGameplayEffect>(StaticConstructObject(UGameplayEffect::StaticClass(), GetTransientPackage(), FName(TEXT("BaseDmgEffect"))));
 		BaseDmgEffect->Modifiers.SetNum(1);
@@ -1858,7 +1858,7 @@ bool GameplayEffectsTest_DurationDamage(UWorld *World, FAutomationTestBase * Tes
 		SourceComponent->ApplyGameplayEffectToTarget(BaseDmgEffect, DestComponent, 1.f);
 
 		// The effect should instantly apply without ticking (for now at least)
-		float ActualValue = DestComponent->GetSet<USkillSystemTestAttributeSet>()->Health;
+		float ActualValue = DestComponent->GetSet<UAbilitySystemTestAttributeSet>()->Health;
 		float ExpectedValue = StartHealth + (DamageValue);
 
 		Test->TestTrue(SKILL_TEST_TEXT("Damage Applied. Actual: %.2f == Exected: %.2f", ActualValue, ExpectedValue), ActualValue == ExpectedValue);
@@ -1877,7 +1877,7 @@ bool GameplayEffectsTest_DurationDamage(UWorld *World, FAutomationTestBase * Tes
 		}
 
 		// The temporary effect is still in place
-		float ActualValue = DestComponent->GetSet<USkillSystemTestAttributeSet>()->Health;
+		float ActualValue = DestComponent->GetSet<UAbilitySystemTestAttributeSet>()->Health;
 		float ExpectedValue = StartHealth + (DamageValue);
 
 		Test->TestTrue(SKILL_TEST_TEXT("Damage Applied. Duration (left) %.2f. Actual: %.2f == Exected: %.2f", Duration, ActualValue, ExpectedValue), ActualValue == ExpectedValue);
@@ -1886,7 +1886,7 @@ bool GameplayEffectsTest_DurationDamage(UWorld *World, FAutomationTestBase * Tes
 	// Ensure the effect expired
 	{
 		// The effect should instantly execute one time without ticking (for now at least)
-		float ActualValue = DestComponent->GetSet<USkillSystemTestAttributeSet>()->Health;
+		float ActualValue = DestComponent->GetSet<UAbilitySystemTestAttributeSet>()->Health;
 		float ExpectedValue = StartHealth;
 
 		Test->TestTrue(SKILL_TEST_TEXT("Damage Applied. Duration (left) %.2f. Actual: %.2f == Exected: %.2f", Duration, ActualValue, ExpectedValue), ActualValue == ExpectedValue);
@@ -1908,20 +1908,20 @@ bool GameplayEffectsTest_PeriodicDamage(UWorld *World, FAutomationTestBase * Tes
 	float Duration = 5.f;
 	float StartTime = World->GetTimeSeconds();
 
-	ASkillSystemTestPawn *SourceActor = World->SpawnActor<ASkillSystemTestPawn>();
-	ASkillSystemTestPawn *DestActor = World->SpawnActor<ASkillSystemTestPawn>();
+	AAbilitySystemTestPawn *SourceActor = World->SpawnActor<AAbilitySystemTestPawn>();
+	AAbilitySystemTestPawn *DestActor = World->SpawnActor<AAbilitySystemTestPawn>();
 
-	UProperty *HealthProperty = FindFieldChecked<UProperty>(USkillSystemTestAttributeSet::StaticClass(), GET_MEMBER_NAME_CHECKED(USkillSystemTestAttributeSet, Health));
+	UProperty *HealthProperty = FindFieldChecked<UProperty>(UAbilitySystemTestAttributeSet::StaticClass(), GET_MEMBER_NAME_CHECKED(UAbilitySystemTestAttributeSet, Health));
 
 	UAttributeComponent * SourceComponent = SourceActor->AttributeComponent;
 	UAttributeComponent * DestComponent = DestActor->AttributeComponent;
 
-	SourceComponent->GetSet<USkillSystemTestAttributeSet>()->Health = StartHealth;
-	DestComponent->GetSet<USkillSystemTestAttributeSet>()->Health = StartHealth;
+	SourceComponent->GetSet<UAbilitySystemTestAttributeSet>()->Health = StartHealth;
+	DestComponent->GetSet<UAbilitySystemTestAttributeSet>()->Health = StartHealth;
 
 	float ApplyCount = 0;
 	{
-		SKILL_LOG_SCOPE(TEXT("Apply Dot"))
+		ABILITY_LOG_SCOPE(TEXT("Apply Dot"))
 
 		UGameplayEffect * BaseDmgEffect = Cast<UGameplayEffect>(StaticConstructObject(UGameplayEffect::StaticClass(), GetTransientPackage(), FName(TEXT("BaseDmgEffect"))));
 		BaseDmgEffect->Modifiers.SetNum(1);
@@ -1937,7 +1937,7 @@ bool GameplayEffectsTest_PeriodicDamage(UWorld *World, FAutomationTestBase * Tes
 		SourceComponent->ApplyGameplayEffectToTarget(BaseDmgEffect, DestComponent, 5.f);
 
 		// The effect should execute on the next tick
-		float ActualValue = DestComponent->GetSet<USkillSystemTestAttributeSet>()->Health;
+		float ActualValue = DestComponent->GetSet<UAbilitySystemTestAttributeSet>()->Health;
 		float ExpectedValue = StartHealth + (DamageValue * ApplyCount);
 
 		Test->TestTrue(SKILL_TEST_TEXT("Damage Applied. Actual: %.2f == Exected: %.2f", ActualValue, ExpectedValue), ActualValue == ExpectedValue );
@@ -1957,7 +1957,7 @@ bool GameplayEffectsTest_PeriodicDamage(UWorld *World, FAutomationTestBase * Tes
 		}
 
 		// The effect should instantly execute one time without ticking (for now at least)
-		float ActualValue = DestComponent->GetSet<USkillSystemTestAttributeSet>()->Health;
+		float ActualValue = DestComponent->GetSet<UAbilitySystemTestAttributeSet>()->Health;
 		float ExpectedValue = StartHealth + (DamageValue * ApplyCount);
 
 		Test->TestTrue(SKILL_TEST_TEXT("Damage Applied. Duration (left) %.2f. Actual: %.2f == Exected: %.2f", Duration, ActualValue, ExpectedValue), ActualValue == ExpectedValue);
@@ -1982,20 +1982,20 @@ bool GameplayEffectsTest_LifestealExtension(UWorld *World, FAutomationTestBase *
 	const float LifestealPCT = 0.20f;
 	float StartTime = World->GetTimeSeconds();
 
-	ASkillSystemTestPawn *SourceActor = World->SpawnActor<ASkillSystemTestPawn>();
-	ASkillSystemTestPawn *DestActor = World->SpawnActor<ASkillSystemTestPawn>();
+	AAbilitySystemTestPawn *SourceActor = World->SpawnActor<AAbilitySystemTestPawn>();
+	AAbilitySystemTestPawn *DestActor = World->SpawnActor<AAbilitySystemTestPawn>();
 
-	UProperty *HealthProperty = FindFieldChecked<UProperty>(USkillSystemTestAttributeSet::StaticClass(), GET_MEMBER_NAME_CHECKED(USkillSystemTestAttributeSet, Health));
-	UProperty *DamageProperty = FindFieldChecked<UProperty>(USkillSystemTestAttributeSet::StaticClass(), GET_MEMBER_NAME_CHECKED(USkillSystemTestAttributeSet, Damage));
+	UProperty *HealthProperty = FindFieldChecked<UProperty>(UAbilitySystemTestAttributeSet::StaticClass(), GET_MEMBER_NAME_CHECKED(UAbilitySystemTestAttributeSet, Health));
+	UProperty *DamageProperty = FindFieldChecked<UProperty>(UAbilitySystemTestAttributeSet::StaticClass(), GET_MEMBER_NAME_CHECKED(UAbilitySystemTestAttributeSet, Damage));
 
 	UAttributeComponent * SourceComponent = SourceActor->AttributeComponent;
 	UAttributeComponent * DestComponent = DestActor->AttributeComponent;
 
-	SourceComponent->GetSet<USkillSystemTestAttributeSet>()->Health = StartHealth;
-	DestComponent->GetSet<USkillSystemTestAttributeSet>()->Health = StartHealth;
+	SourceComponent->GetSet<UAbilitySystemTestAttributeSet>()->Health = StartHealth;
+	DestComponent->GetSet<UAbilitySystemTestAttributeSet>()->Health = StartHealth;
 
 	{
-		SKILL_LOG_SCOPE(TEXT("Apply Lifesteal"));
+		ABILITY_LOG_SCOPE(TEXT("Apply Lifesteal"));
 
 		UGameplayEffect * LifestealEffect = Cast<UGameplayEffect>(StaticConstructObject(UGameplayEffect::StaticClass(), GetTransientPackage(), FName(TEXT("LifestealPassive"))));
 		LifestealEffect->Modifiers.SetNum(1);
@@ -2012,7 +2012,7 @@ bool GameplayEffectsTest_LifestealExtension(UWorld *World, FAutomationTestBase *
 	}
 	
 	{
-		SKILL_LOG_SCOPE(TEXT("Apply Damage"));
+		ABILITY_LOG_SCOPE(TEXT("Apply Damage"));
 
 		UGameplayEffect * BaseDmgEffect = Cast<UGameplayEffect>(StaticConstructObject(UGameplayEffect::StaticClass(), GetTransientPackage(), FName(TEXT("BaseDmgEffect"))));
 		BaseDmgEffect->Modifiers.SetNum(1);
@@ -2028,7 +2028,7 @@ bool GameplayEffectsTest_LifestealExtension(UWorld *World, FAutomationTestBase *
 
 		// The effect should instantly execute one time without ticking (for now at least)
 		{
-			float ActualValue = DestComponent->GetSet<USkillSystemTestAttributeSet>()->Health;
+			float ActualValue = DestComponent->GetSet<UAbilitySystemTestAttributeSet>()->Health;
 			float ExpectedValue = StartHealth + (DamageValue);
 
 			Test->TestTrue(SKILL_TEST_TEXT("Damage Applied. Actual: %.2f == Exected: %.2f", ActualValue, ExpectedValue), ActualValue == ExpectedValue);
@@ -2036,7 +2036,7 @@ bool GameplayEffectsTest_LifestealExtension(UWorld *World, FAutomationTestBase *
 
 		// Test that the source received extra health back
 		{
-			float ActualValue = SourceComponent->GetSet<USkillSystemTestAttributeSet>()->Health;
+			float ActualValue = SourceComponent->GetSet<UAbilitySystemTestAttributeSet>()->Health;
 			float ExpectedValue = StartHealth + (-DamageValue * LifestealPCT);
 
 			Test->TestTrue(SKILL_TEST_TEXT("Health after lifesteal. Actual: %.2f == Exected: %.2f", ActualValue, ExpectedValue), ActualValue == ExpectedValue);
@@ -2057,21 +2057,21 @@ bool GameplayEffectsTest_ShieldExtension(UWorld *World, FAutomationTestBase * Te
 	const float ShieldAmount = 20.0f;
 	float StartTime = World->GetTimeSeconds();
 
-	ASkillSystemTestPawn *SourceActor = World->SpawnActor<ASkillSystemTestPawn>();
-	ASkillSystemTestPawn *DestActor = World->SpawnActor<ASkillSystemTestPawn>();
+	AAbilitySystemTestPawn *SourceActor = World->SpawnActor<AAbilitySystemTestPawn>();
+	AAbilitySystemTestPawn *DestActor = World->SpawnActor<AAbilitySystemTestPawn>();
 
-	UProperty *HealthProperty = FindFieldChecked<UProperty>(USkillSystemTestAttributeSet::StaticClass(), GET_MEMBER_NAME_CHECKED(USkillSystemTestAttributeSet, Health));
-	UProperty *DamageProperty = FindFieldChecked<UProperty>(USkillSystemTestAttributeSet::StaticClass(), GET_MEMBER_NAME_CHECKED(USkillSystemTestAttributeSet, Damage));
+	UProperty *HealthProperty = FindFieldChecked<UProperty>(UAbilitySystemTestAttributeSet::StaticClass(), GET_MEMBER_NAME_CHECKED(UAbilitySystemTestAttributeSet, Health));
+	UProperty *DamageProperty = FindFieldChecked<UProperty>(UAbilitySystemTestAttributeSet::StaticClass(), GET_MEMBER_NAME_CHECKED(UAbilitySystemTestAttributeSet, Damage));
 
 	UAttributeComponent * SourceComponent = SourceActor->AttributeComponent;
 	UAttributeComponent * DestComponent = DestActor->AttributeComponent;
 
-	SourceComponent->GetSet<USkillSystemTestAttributeSet>()->Health = StartHealth;
-	DestComponent->GetSet<USkillSystemTestAttributeSet>()->Health = StartHealth;
+	SourceComponent->GetSet<UAbilitySystemTestAttributeSet>()->Health = StartHealth;
+	DestComponent->GetSet<UAbilitySystemTestAttributeSet>()->Health = StartHealth;
 
 	FActiveGameplayEffectHandle AppliedHandle;
 	{
-		SKILL_LOG_SCOPE(TEXT("Apply Shield"));
+		ABILITY_LOG_SCOPE(TEXT("Apply Shield"));
 
 		UGameplayEffect * ShieldEffect = Cast<UGameplayEffect>(StaticConstructObject(UGameplayEffect::StaticClass(), GetTransientPackage(), FName(TEXT("ShieldPassive"))));
 		ShieldEffect->Modifiers.SetNum(1);
@@ -2088,7 +2088,7 @@ bool GameplayEffectsTest_ShieldExtension(UWorld *World, FAutomationTestBase * Te
 	}
 
 	{
-		SKILL_LOG_SCOPE(TEXT("Apply Damage"));
+		ABILITY_LOG_SCOPE(TEXT("Apply Damage"));
 
 		UGameplayEffect * BaseDmgEffect = Cast<UGameplayEffect>(StaticConstructObject(UGameplayEffect::StaticClass(), GetTransientPackage(), FName(TEXT("BaseDmgEffect"))));
 		BaseDmgEffect->Modifiers.SetNum(1);
@@ -2107,7 +2107,7 @@ bool GameplayEffectsTest_ShieldExtension(UWorld *World, FAutomationTestBase * Te
 
 			// Health should be the same
 			{
-				float ActualValue = DestComponent->GetSet<USkillSystemTestAttributeSet>()->Health;
+				float ActualValue = DestComponent->GetSet<UAbilitySystemTestAttributeSet>()->Health;
 				float ExpectedValue = StartHealth;
 				Test->TestTrue(SKILL_TEST_TEXT("Damage Applied. Actual: %.2f == Exected: %.2f", ActualValue, ExpectedValue), ActualValue == ExpectedValue);
 			}
@@ -2126,7 +2126,7 @@ bool GameplayEffectsTest_ShieldExtension(UWorld *World, FAutomationTestBase * Te
 
 			// Health should be the same
 			{
-				float ActualValue = DestComponent->GetSet<USkillSystemTestAttributeSet>()->Health;
+				float ActualValue = DestComponent->GetSet<UAbilitySystemTestAttributeSet>()->Health;
 				float ExpectedValue = StartHealth;
 				Test->TestTrue(SKILL_TEST_TEXT("Damage Applied. Actual: %.2f == Exected: %.2f", ActualValue, ExpectedValue), ActualValue == ExpectedValue);
 			}
@@ -2144,7 +2144,7 @@ bool GameplayEffectsTest_ShieldExtension(UWorld *World, FAutomationTestBase * Te
 			SourceComponent->ApplyGameplayEffectToTarget(BaseDmgEffect, DestComponent, 5.f);
 
 			// Now we should have lost some health
-			float ActualValue = DestComponent->GetSet<USkillSystemTestAttributeSet>()->Health;
+			float ActualValue = DestComponent->GetSet<UAbilitySystemTestAttributeSet>()->Health;
 			float ExpectedValue = StartHealth + DamageValue;;
 			Test->TestTrue(SKILL_TEST_TEXT("Damage Applied. Actual: %.2f == Exected: %.2f", ActualValue, ExpectedValue), ActualValue == ExpectedValue);
 
@@ -2172,22 +2172,22 @@ bool GameplayEffectsTest_ShieldExtensionMultiple(UWorld *World, FAutomationTestB
 	const float ShieldAmount = 20.0f;
 	float StartTime = World->GetTimeSeconds();
 
-	ASkillSystemTestPawn *SourceActor = World->SpawnActor<ASkillSystemTestPawn>();
-	ASkillSystemTestPawn *DestActor = World->SpawnActor<ASkillSystemTestPawn>();
+	AAbilitySystemTestPawn *SourceActor = World->SpawnActor<AAbilitySystemTestPawn>();
+	AAbilitySystemTestPawn *DestActor = World->SpawnActor<AAbilitySystemTestPawn>();
 
-	UProperty *HealthProperty = FindFieldChecked<UProperty>(USkillSystemTestAttributeSet::StaticClass(), GET_MEMBER_NAME_CHECKED(USkillSystemTestAttributeSet, Health));
-	UProperty *DamageProperty = FindFieldChecked<UProperty>(USkillSystemTestAttributeSet::StaticClass(), GET_MEMBER_NAME_CHECKED(USkillSystemTestAttributeSet, Damage));
+	UProperty *HealthProperty = FindFieldChecked<UProperty>(UAbilitySystemTestAttributeSet::StaticClass(), GET_MEMBER_NAME_CHECKED(UAbilitySystemTestAttributeSet, Health));
+	UProperty *DamageProperty = FindFieldChecked<UProperty>(UAbilitySystemTestAttributeSet::StaticClass(), GET_MEMBER_NAME_CHECKED(UAbilitySystemTestAttributeSet, Damage));
 
 	UAttributeComponent * SourceComponent = SourceActor->AttributeComponent;
 	UAttributeComponent * DestComponent = DestActor->AttributeComponent;
 
-	SourceComponent->GetSet<USkillSystemTestAttributeSet>()->Health = StartHealth;
-	DestComponent->GetSet<USkillSystemTestAttributeSet>()->Health = StartHealth;
+	SourceComponent->GetSet<UAbilitySystemTestAttributeSet>()->Health = StartHealth;
+	DestComponent->GetSet<UAbilitySystemTestAttributeSet>()->Health = StartHealth;
 
 	FActiveGameplayEffectHandle AppliedHandle_1;
 	FActiveGameplayEffectHandle AppliedHandle_2;
 	{
-		SKILL_LOG_SCOPE(TEXT("Apply Shields"));
+		ABILITY_LOG_SCOPE(TEXT("Apply Shields"));
 
 		UGameplayEffect * ShieldEffect = Cast<UGameplayEffect>(StaticConstructObject(UGameplayEffect::StaticClass(), GetTransientPackage(), FName(TEXT("ShieldPassive"))));
 		ShieldEffect->Modifiers.SetNum(1);
@@ -2205,7 +2205,7 @@ bool GameplayEffectsTest_ShieldExtensionMultiple(UWorld *World, FAutomationTestB
 	}
 
 	{
-		SKILL_LOG_SCOPE(TEXT("Apply Damage"));
+		ABILITY_LOG_SCOPE(TEXT("Apply Damage"));
 
 		UGameplayEffect * SmallDmgEffect = Cast<UGameplayEffect>(StaticConstructObject(UGameplayEffect::StaticClass(), GetTransientPackage(), FName(TEXT("SmallDmgEffect"))));
 		SmallDmgEffect->Modifiers.SetNum(1);
@@ -2234,7 +2234,7 @@ bool GameplayEffectsTest_ShieldExtensionMultiple(UWorld *World, FAutomationTestB
 
 			// Health should be the same
 			{
-				float ActualValue = DestComponent->GetSet<USkillSystemTestAttributeSet>()->Health;
+				float ActualValue = DestComponent->GetSet<UAbilitySystemTestAttributeSet>()->Health;
 				float ExpectedValue = StartHealth;
 				Test->TestTrue(SKILL_TEST_TEXT("Damage Applied. Actual: %.2f == Exected: %.2f", ActualValue, ExpectedValue), ActualValue == ExpectedValue);
 			}
@@ -2260,7 +2260,7 @@ bool GameplayEffectsTest_ShieldExtensionMultiple(UWorld *World, FAutomationTestB
 
 			// Health should still be the same
 			{
-				float ActualValue = DestComponent->GetSet<USkillSystemTestAttributeSet>()->Health;
+				float ActualValue = DestComponent->GetSet<UAbilitySystemTestAttributeSet>()->Health;
 				float ExpectedValue = StartHealth;
 				Test->TestTrue(SKILL_TEST_TEXT("Damage Applied. Actual: %.2f == Exected: %.2f", ActualValue, ExpectedValue), ActualValue == ExpectedValue);
 			}
@@ -2289,7 +2289,7 @@ bool GameplayEffectsTest_ShieldExtensionMultiple(UWorld *World, FAutomationTestB
 			float HealthDelta = ShieldAmount + ShieldAmount + DamageValueSmall + DamageValueLarge + DamageValueLarge;
 
 			// Now we should have lost some health
-			float ActualValue = DestComponent->GetSet<USkillSystemTestAttributeSet>()->Health;
+			float ActualValue = DestComponent->GetSet<UAbilitySystemTestAttributeSet>()->Health;
 			float ExpectedValue = StartHealth + HealthDelta;
 			Test->TestTrue(SKILL_TEST_TEXT("Damage Applied. Actual: %.2f == Exected: %.2f", ActualValue, ExpectedValue), ActualValue == ExpectedValue);
 
@@ -2323,13 +2323,13 @@ UCurveTable * SetGlobalCurveTable()
 		check(Value == 5.f);
 	}
 
-	IGameplayAbilitiesModule::Get().GetSkillSystemGlobals()->AutomationTestOnly_SetGlobalCurveTable(CurveTable);
+	IGameplayAbilitiesModule::Get().GetAbilitySystemGlobals()->AutomationTestOnly_SetGlobalCurveTable(CurveTable);
 	return CurveTable;
 }
 
 void ClearGlobalCurveTable()
 {
-	IGameplayAbilitiesModule::Get().GetSkillSystemGlobals()->AutomationTestOnly_SetGlobalCurveTable(NULL);
+	IGameplayAbilitiesModule::Get().GetAbilitySystemGlobals()->AutomationTestOnly_SetGlobalCurveTable(NULL);
 }
 
 UCurveTable * GetStandardDamageOverrideCurveTable(float Factor)
@@ -2369,13 +2369,13 @@ UDataTable * SetGlobalDataTable()
 		check(!Row->bCanStack);
 	}
 
-	IGameplayAbilitiesModule::Get().GetSkillSystemGlobals()->AutomationTestOnly_SetGlobalAttributeDataTable(DataTable);
+	IGameplayAbilitiesModule::Get().GetAbilitySystemGlobals()->AutomationTestOnly_SetGlobalAttributeDataTable(DataTable);
 	return DataTable;
 }
 
 void ClearGlobalDataTable()
 {
-	IGameplayAbilitiesModule::Get().GetSkillSystemGlobals()->AutomationTestOnly_SetGlobalAttributeDataTable(NULL);
+	IGameplayAbilitiesModule::Get().GetAbilitySystemGlobals()->AutomationTestOnly_SetGlobalAttributeDataTable(NULL);
 }
 
 bool GameplayEffectsTest_InstantDamage_ScalingExplicit(UWorld *World, FAutomationTestBase * Test)
@@ -2392,21 +2392,21 @@ bool GameplayEffectsTest_InstantDamage_ScalingExplicit(UWorld *World, FAutomatio
 	// Sets up a linear curve table f(x)=x for StandardDamage
 	UCurveTable * SourceCurveTableOverrides = GetStandardDamageOverrideCurveTable(SourceDamageScale);
 
-	ASkillSystemTestPawn *SourceActor = World->SpawnActor<ASkillSystemTestPawn>();
-	ASkillSystemTestPawn *DestActor = World->SpawnActor<ASkillSystemTestPawn>();
+	AAbilitySystemTestPawn *SourceActor = World->SpawnActor<AAbilitySystemTestPawn>();
+	AAbilitySystemTestPawn *DestActor = World->SpawnActor<AAbilitySystemTestPawn>();
 
-	UProperty *DamageProperty = FindFieldChecked<UProperty>(USkillSystemTestAttributeSet::StaticClass(), GET_MEMBER_NAME_CHECKED(USkillSystemTestAttributeSet, Damage));
+	UProperty *DamageProperty = FindFieldChecked<UProperty>(UAbilitySystemTestAttributeSet::StaticClass(), GET_MEMBER_NAME_CHECKED(UAbilitySystemTestAttributeSet, Damage));
 
 	UAttributeComponent * SourceComponent = SourceActor->AttributeComponent;
 	UAttributeComponent * DestComponent = DestActor->AttributeComponent;
-	SourceComponent->GetSet<USkillSystemTestAttributeSet>()->Health = StartHealth;
-	DestComponent->GetSet<USkillSystemTestAttributeSet>()->Health = StartHealth;
+	SourceComponent->GetSet<UAbilitySystemTestAttributeSet>()->Health = StartHealth;
+	DestComponent->GetSet<UAbilitySystemTestAttributeSet>()->Health = StartHealth;
 
 	// Source now has SourceDamageScale (2x) damage over standard damage 
 	SourceComponent->PushGlobalCurveOveride(SourceCurveTableOverrides);
 
 	{
-		SKILL_LOG_SCOPE(TEXT("Apply InstantDamage"));
+		ABILITY_LOG_SCOPE(TEXT("Apply InstantDamage"));
 
 		UGameplayEffect * BaseDmgEffect = Cast<UGameplayEffect>(StaticConstructObject(UGameplayEffect::StaticClass(), GetTransientPackage(), FName(TEXT("BaseDmgEffect"))));
 		BaseDmgEffect->Modifiers.SetNum(1);
@@ -2419,7 +2419,7 @@ bool GameplayEffectsTest_InstantDamage_ScalingExplicit(UWorld *World, FAutomatio
 
 		SourceComponent->ApplyGameplayEffectToTarget(BaseDmgEffect, DestComponent, LevelOfDamage);
 
-		float ActualValue = DestComponent->GetSet<USkillSystemTestAttributeSet>()->Health;
+		float ActualValue = DestComponent->GetSet<UAbilitySystemTestAttributeSet>()->Health;
 		float ExpectedValue = StartHealth - (LevelOfDamage * SourceDamageScale);
 		Test->TestTrue(SKILL_TEST_TEXT("Damage Applied. Health: Actual: %.2f == Exected: %.2f", ActualValue, ExpectedValue), ActualValue == ExpectedValue);
 	}
@@ -2441,18 +2441,18 @@ bool GameplayEffectsTest_InstantDamage_ScalingGlobal(UWorld *World, FAutomationT
 
 	SetGlobalCurveTable();
 
-	ASkillSystemTestPawn *SourceActor = World->SpawnActor<ASkillSystemTestPawn>();
-	ASkillSystemTestPawn *DestActor = World->SpawnActor<ASkillSystemTestPawn>();
+	AAbilitySystemTestPawn *SourceActor = World->SpawnActor<AAbilitySystemTestPawn>();
+	AAbilitySystemTestPawn *DestActor = World->SpawnActor<AAbilitySystemTestPawn>();
 	
-	UProperty *DamageProperty = FindFieldChecked<UProperty>(USkillSystemTestAttributeSet::StaticClass(), GET_MEMBER_NAME_CHECKED(USkillSystemTestAttributeSet, Damage));
+	UProperty *DamageProperty = FindFieldChecked<UProperty>(UAbilitySystemTestAttributeSet::StaticClass(), GET_MEMBER_NAME_CHECKED(UAbilitySystemTestAttributeSet, Damage));
 
 	UAttributeComponent * SourceComponent = SourceActor->AttributeComponent;
 	UAttributeComponent * DestComponent = DestActor->AttributeComponent;
-	SourceComponent->GetSet<USkillSystemTestAttributeSet>()->Health = StartHealth;
-	DestComponent->GetSet<USkillSystemTestAttributeSet>()->Health = StartHealth;
+	SourceComponent->GetSet<UAbilitySystemTestAttributeSet>()->Health = StartHealth;
+	DestComponent->GetSet<UAbilitySystemTestAttributeSet>()->Health = StartHealth;
 
 	{
-		SKILL_LOG_SCOPE(TEXT("Apply InstantDamage"));
+		ABILITY_LOG_SCOPE(TEXT("Apply InstantDamage"));
 
 		UGameplayEffect * BaseDmgEffect = Cast<UGameplayEffect>(StaticConstructObject(UGameplayEffect::StaticClass(), GetTransientPackage(), FName(TEXT("BaseDmgEffect"))));
 		BaseDmgEffect->Modifiers.SetNum(1);
@@ -2465,7 +2465,7 @@ bool GameplayEffectsTest_InstantDamage_ScalingGlobal(UWorld *World, FAutomationT
 
 		SourceComponent->ApplyGameplayEffectToTarget(BaseDmgEffect, DestComponent, LevelOfDamage);
 
-		float ActualValue = DestComponent->GetSet<USkillSystemTestAttributeSet>()->Health;
+		float ActualValue = DestComponent->GetSet<UAbilitySystemTestAttributeSet>()->Health;
 		float ExpectedValue = StartHealth - (LevelOfDamage);
 		Test->TestTrue(SKILL_TEST_TEXT("Damage Applied. Health: Actual: %.2f == Exected: %.2f", ActualValue, ExpectedValue), ActualValue == ExpectedValue);
 	}
@@ -2490,21 +2490,21 @@ bool GameplayEffectsTest_InstantDamage_OverrideScaling(UWorld *World, FAutomatio
 	SetGlobalCurveTable();
 	UCurveTable * SourceCurveTableOverrides = GetStandardDamageOverrideCurveTable(SourceDamageScale);
 
-	ASkillSystemTestPawn *SourceActor = World->SpawnActor<ASkillSystemTestPawn>();
-	ASkillSystemTestPawn *DestActor = World->SpawnActor<ASkillSystemTestPawn>();
+	AAbilitySystemTestPawn *SourceActor = World->SpawnActor<AAbilitySystemTestPawn>();
+	AAbilitySystemTestPawn *DestActor = World->SpawnActor<AAbilitySystemTestPawn>();
 
-	UProperty *DamageProperty = FindFieldChecked<UProperty>(USkillSystemTestAttributeSet::StaticClass(), GET_MEMBER_NAME_CHECKED(USkillSystemTestAttributeSet, Damage));
+	UProperty *DamageProperty = FindFieldChecked<UProperty>(UAbilitySystemTestAttributeSet::StaticClass(), GET_MEMBER_NAME_CHECKED(UAbilitySystemTestAttributeSet, Damage));
 
 	UAttributeComponent * SourceComponent = SourceActor->AttributeComponent;
 	UAttributeComponent * DestComponent = DestActor->AttributeComponent;
-	SourceComponent->GetSet<USkillSystemTestAttributeSet>()->Health = StartHealth;
-	DestComponent->GetSet<USkillSystemTestAttributeSet>()->Health = StartHealth;
+	SourceComponent->GetSet<UAbilitySystemTestAttributeSet>()->Health = StartHealth;
+	DestComponent->GetSet<UAbilitySystemTestAttributeSet>()->Health = StartHealth;
 
 	// Source now has SourceDamageScale (2x) damage over standard damage 
 	SourceComponent->PushGlobalCurveOveride(SourceCurveTableOverrides);
 
 	{
-		SKILL_LOG_SCOPE(TEXT("Apply InstantDamage"));
+		ABILITY_LOG_SCOPE(TEXT("Apply InstantDamage"));
 
 		UGameplayEffect * BaseDmgEffect = Cast<UGameplayEffect>(StaticConstructObject(UGameplayEffect::StaticClass(), GetTransientPackage(), FName(TEXT("BaseDmgEffect"))));
 		BaseDmgEffect->Modifiers.SetNum(1);
@@ -2517,7 +2517,7 @@ bool GameplayEffectsTest_InstantDamage_OverrideScaling(UWorld *World, FAutomatio
 
 		SourceComponent->ApplyGameplayEffectToTarget(BaseDmgEffect, DestComponent, LevelOfDamage);
 
-		float ActualValue = DestComponent->GetSet<USkillSystemTestAttributeSet>()->Health;
+		float ActualValue = DestComponent->GetSet<UAbilitySystemTestAttributeSet>()->Health;
 		float ExpectedValue = StartHealth - (LevelOfDamage * SourceDamageScale);
 		Test->TestTrue(SKILL_TEST_TEXT("Damage Applied. Health: Actual: %.2f == Exected: %.2f", ActualValue, ExpectedValue), ActualValue == ExpectedValue);
 	}
@@ -2536,19 +2536,19 @@ bool GameplayEffectsTest_InstantDamageRequiredTag(UWorld *World, FAutomationTest
 	const float DamageValue = -5.f;
 	const float DamageProtectionDivisor = 2.f;
 
-	ASkillSystemTestPawn *SourceActor = World->SpawnActor<ASkillSystemTestPawn>();
-	ASkillSystemTestPawn *DestActor = World->SpawnActor<ASkillSystemTestPawn>();
+	AAbilitySystemTestPawn *SourceActor = World->SpawnActor<AAbilitySystemTestPawn>();
+	AAbilitySystemTestPawn *DestActor = World->SpawnActor<AAbilitySystemTestPawn>();
 
-	UProperty *HealthProperty = FindFieldChecked<UProperty>(USkillSystemTestAttributeSet::StaticClass(), GET_MEMBER_NAME_CHECKED(USkillSystemTestAttributeSet, Health));
+	UProperty *HealthProperty = FindFieldChecked<UProperty>(UAbilitySystemTestAttributeSet::StaticClass(), GET_MEMBER_NAME_CHECKED(UAbilitySystemTestAttributeSet, Health));
 
 	UAttributeComponent * SourceComponent = SourceActor->AttributeComponent;
 	UAttributeComponent * DestComponent = DestActor->AttributeComponent;
-	SourceComponent->GetSet<USkillSystemTestAttributeSet>()->Health = StartHealth;
-	DestComponent->GetSet<USkillSystemTestAttributeSet>()->Health = StartHealth;
+	SourceComponent->GetSet<UAbilitySystemTestAttributeSet>()->Health = StartHealth;
+	DestComponent->GetSet<UAbilitySystemTestAttributeSet>()->Health = StartHealth;
 
 	// Setup a GE to modify IncomingGEs
 	{
-		SKILL_LOG_SCOPE(TEXT("Apply ProtectionBuff"))
+		ABILITY_LOG_SCOPE(TEXT("Apply ProtectionBuff"))
 
 		UGameplayEffect* BaseProtectEffect = Cast<UGameplayEffect>(StaticConstructObject(UGameplayEffect::StaticClass(), GetTransientPackage(), FName(TEXT("ProtectBuff"))));
 		BaseProtectEffect->Modifiers.SetNum(1);
@@ -2567,7 +2567,7 @@ bool GameplayEffectsTest_InstantDamageRequiredTag(UWorld *World, FAutomationTest
 
 	// Apply Damage
 	{
-		SKILL_LOG_SCOPE(TEXT("Apply InstantDamage"))
+		ABILITY_LOG_SCOPE(TEXT("Apply InstantDamage"))
 
 		UGameplayEffect * BaseDmgEffect = Cast<UGameplayEffect>(StaticConstructObject(UGameplayEffect::StaticClass(), GetTransientPackage(), FName(TEXT("Damage"))));
 		BaseDmgEffect->Modifiers.SetNum(1);
@@ -2584,16 +2584,16 @@ bool GameplayEffectsTest_InstantDamageRequiredTag(UWorld *World, FAutomationTest
 
 		float ExpectedValue = (StartHealth + DamageValue);
 		
-		Test->TestTrue(SKILL_TEST_TEXT("Instant Damage Required Tag No Protection"), (DestComponent->GetSet<USkillSystemTestAttributeSet>()->Health == ExpectedValue));
-		SKILL_LOG(Log, TEXT("Final Health: %.2f"), DestComponent->GetSet<USkillSystemTestAttributeSet>()->Health);
+		Test->TestTrue(SKILL_TEST_TEXT("Instant Damage Required Tag No Protection"), (DestComponent->GetSet<UAbilitySystemTestAttributeSet>()->Health == ExpectedValue));
+		ABILITY_LOG(Log, TEXT("Final Health: %.2f"), DestComponent->GetSet<UAbilitySystemTestAttributeSet>()->Health);
 	}
 
 	// reset health
-	DestComponent->GetSet<USkillSystemTestAttributeSet>()->Health = StartHealth;
+	DestComponent->GetSet<UAbilitySystemTestAttributeSet>()->Health = StartHealth;
 
 	// Apply Damage
 	{
-		SKILL_LOG_SCOPE(TEXT("Apply InstantDamage"))
+		ABILITY_LOG_SCOPE(TEXT("Apply InstantDamage"))
 
 		UGameplayEffect * BaseDmgEffect = Cast<UGameplayEffect>(StaticConstructObject(UGameplayEffect::StaticClass(), GetTransientPackage(), FName(TEXT("Damage"))));
 		BaseDmgEffect->Modifiers.SetNum(1);
@@ -2609,7 +2609,7 @@ bool GameplayEffectsTest_InstantDamageRequiredTag(UWorld *World, FAutomationTest
 		SourceComponent->ApplyGameplayEffectToTarget(BaseDmgEffect, DestComponent, 1.f);
 
 		float ExpectedValue = (StartHealth + (DamageValue / DamageProtectionDivisor));
-		float ActualValue = DestComponent->GetSet<USkillSystemTestAttributeSet>()->Health;
+		float ActualValue = DestComponent->GetSet<UAbilitySystemTestAttributeSet>()->Health;
 		
 		Test->TestTrue(SKILL_TEST_TEXT("Damage Applied. Health: Actual: %.2f == Exected: %.2f", ActualValue, ExpectedValue), ActualValue == ExpectedValue);
 	}
@@ -2626,19 +2626,19 @@ bool GameplayEffectsTest_InstantDamageIgnoreTag(UWorld *World, FAutomationTestBa
 	const float DamageValue = -5.f;
 	const float DamageProtectionDivisor = 2.f;
 
-	ASkillSystemTestPawn *SourceActor = World->SpawnActor<ASkillSystemTestPawn>();
-	ASkillSystemTestPawn *DestActor = World->SpawnActor<ASkillSystemTestPawn>();
+	AAbilitySystemTestPawn *SourceActor = World->SpawnActor<AAbilitySystemTestPawn>();
+	AAbilitySystemTestPawn *DestActor = World->SpawnActor<AAbilitySystemTestPawn>();
 
-	UProperty *HealthProperty = FindFieldChecked<UProperty>(USkillSystemTestAttributeSet::StaticClass(), GET_MEMBER_NAME_CHECKED(USkillSystemTestAttributeSet, Health));
+	UProperty *HealthProperty = FindFieldChecked<UProperty>(UAbilitySystemTestAttributeSet::StaticClass(), GET_MEMBER_NAME_CHECKED(UAbilitySystemTestAttributeSet, Health));
 
 	UAttributeComponent * SourceComponent = SourceActor->AttributeComponent;
 	UAttributeComponent * DestComponent = DestActor->AttributeComponent;
-	SourceComponent->GetSet<USkillSystemTestAttributeSet>()->Health = StartHealth;
-	DestComponent->GetSet<USkillSystemTestAttributeSet>()->Health = StartHealth;
+	SourceComponent->GetSet<UAbilitySystemTestAttributeSet>()->Health = StartHealth;
+	DestComponent->GetSet<UAbilitySystemTestAttributeSet>()->Health = StartHealth;
 
 	// Setup a GE to modify IncomingGEs
 	{
-		SKILL_LOG_SCOPE(TEXT("Apply ProtectionBuff"))
+		ABILITY_LOG_SCOPE(TEXT("Apply ProtectionBuff"))
 
 		UGameplayEffect* BaseProtectEffect = Cast<UGameplayEffect>(StaticConstructObject(UGameplayEffect::StaticClass(), GetTransientPackage(), FName(TEXT("ProtectBuff"))));
 		BaseProtectEffect->Modifiers.SetNum(1);
@@ -2657,7 +2657,7 @@ bool GameplayEffectsTest_InstantDamageIgnoreTag(UWorld *World, FAutomationTestBa
 
 	// Apply Damage
 	{
-		SKILL_LOG_SCOPE(TEXT("Apply InstantDamage"))
+		ABILITY_LOG_SCOPE(TEXT("Apply InstantDamage"))
 
 		UGameplayEffect * BaseDmgEffect = Cast<UGameplayEffect>(StaticConstructObject(UGameplayEffect::StaticClass(), GetTransientPackage(), FName(TEXT("Damage"))));
 		BaseDmgEffect->Modifiers.SetNum(1);
@@ -2674,16 +2674,16 @@ bool GameplayEffectsTest_InstantDamageIgnoreTag(UWorld *World, FAutomationTestBa
 
 		float ExpectedValue = (StartHealth + DamageValue);
 
-		Test->TestTrue(SKILL_TEST_TEXT("Instant Damage Ignore Tag No Protection"), (DestComponent->GetSet<USkillSystemTestAttributeSet>()->Health == ExpectedValue));
-		SKILL_LOG(Log, TEXT("Final Health: %.2f"), DestComponent->GetSet<USkillSystemTestAttributeSet>()->Health);
+		Test->TestTrue(SKILL_TEST_TEXT("Instant Damage Ignore Tag No Protection"), (DestComponent->GetSet<UAbilitySystemTestAttributeSet>()->Health == ExpectedValue));
+		ABILITY_LOG(Log, TEXT("Final Health: %.2f"), DestComponent->GetSet<UAbilitySystemTestAttributeSet>()->Health);
 	}
 
 	// reset health
-	DestComponent->GetSet<USkillSystemTestAttributeSet>()->Health = StartHealth;
+	DestComponent->GetSet<UAbilitySystemTestAttributeSet>()->Health = StartHealth;
 
 	// Apply Damage
 	{
-		SKILL_LOG_SCOPE(TEXT("Apply InstantDamage"))
+		ABILITY_LOG_SCOPE(TEXT("Apply InstantDamage"))
 
 		UGameplayEffect * BaseDmgEffect = Cast<UGameplayEffect>(StaticConstructObject(UGameplayEffect::StaticClass(), GetTransientPackage(), FName(TEXT("Damage"))));
 		BaseDmgEffect->Modifiers.SetNum(1);
@@ -2700,8 +2700,8 @@ bool GameplayEffectsTest_InstantDamageIgnoreTag(UWorld *World, FAutomationTestBa
 
 		float ExpectedValue = (StartHealth + (DamageValue / DamageProtectionDivisor));
 
-		Test->TestTrue(SKILL_TEST_TEXT("Instant Damage Ignore Tag Protected"), (DestComponent->GetSet<USkillSystemTestAttributeSet>()->Health == ExpectedValue));
-		SKILL_LOG(Log, TEXT("Final Health: %.2f"), DestComponent->GetSet<USkillSystemTestAttributeSet>()->Health);
+		Test->TestTrue(SKILL_TEST_TEXT("Instant Damage Ignore Tag Protected"), (DestComponent->GetSet<UAbilitySystemTestAttributeSet>()->Health == ExpectedValue));
+		ABILITY_LOG(Log, TEXT("Final Health: %.2f"), DestComponent->GetSet<UAbilitySystemTestAttributeSet>()->Health);
 	}
 
 	World->EditorDestroyActor(SourceActor, false);
@@ -2717,19 +2717,19 @@ bool GameplayEffectsTest_InstantDamageModifierPassesTag(UWorld *World, FAutomati
 	const float BonusDamageMultiplier = 2.f;
 	const float DamageProtectionDivisor = 2.f;
 
-	ASkillSystemTestPawn *SourceActor = World->SpawnActor<ASkillSystemTestPawn>();
-	ASkillSystemTestPawn *DestActor = World->SpawnActor<ASkillSystemTestPawn>();
+	AAbilitySystemTestPawn *SourceActor = World->SpawnActor<AAbilitySystemTestPawn>();
+	AAbilitySystemTestPawn *DestActor = World->SpawnActor<AAbilitySystemTestPawn>();
 
-	UProperty *HealthProperty = FindFieldChecked<UProperty>(USkillSystemTestAttributeSet::StaticClass(), GET_MEMBER_NAME_CHECKED(USkillSystemTestAttributeSet, Health));
+	UProperty *HealthProperty = FindFieldChecked<UProperty>(UAbilitySystemTestAttributeSet::StaticClass(), GET_MEMBER_NAME_CHECKED(UAbilitySystemTestAttributeSet, Health));
 
 	UAttributeComponent * SourceComponent = SourceActor->AttributeComponent;
 	UAttributeComponent * DestComponent = DestActor->AttributeComponent;
-	SourceComponent->GetSet<USkillSystemTestAttributeSet>()->Health = StartHealth;
-	DestComponent->GetSet<USkillSystemTestAttributeSet>()->Health = StartHealth;
+	SourceComponent->GetSet<UAbilitySystemTestAttributeSet>()->Health = StartHealth;
+	DestComponent->GetSet<UAbilitySystemTestAttributeSet>()->Health = StartHealth;
 
 	// Setup a GE to modify OutgoingGEs
 	{
-		SKILL_LOG_SCOPE(TEXT("Apply DamageBuff"))
+		ABILITY_LOG_SCOPE(TEXT("Apply DamageBuff"))
 
 		UGameplayEffect * BaseDmgEffect = Cast<UGameplayEffect>(StaticConstructObject(UGameplayEffect::StaticClass(), GetTransientPackage(), FName(TEXT("DamageBuff"))));
 		BaseDmgEffect->Modifiers.SetNum(1);
@@ -2746,7 +2746,7 @@ bool GameplayEffectsTest_InstantDamageModifierPassesTag(UWorld *World, FAutomati
 
 	// Setup a GE to modify IncomingGEs
 	{
-		SKILL_LOG_SCOPE(TEXT("Apply ProtectionBuff"))
+		ABILITY_LOG_SCOPE(TEXT("Apply ProtectionBuff"))
 
 		UGameplayEffect* BaseProtectEffect = Cast<UGameplayEffect>(StaticConstructObject(UGameplayEffect::StaticClass(), GetTransientPackage(), FName(TEXT("ProtectBuff"))));
 		BaseProtectEffect->Modifiers.SetNum(1);
@@ -2765,7 +2765,7 @@ bool GameplayEffectsTest_InstantDamageModifierPassesTag(UWorld *World, FAutomati
 
 	// Apply Damage
 	{
-		SKILL_LOG_SCOPE(TEXT("Apply InstantDamage"))
+		ABILITY_LOG_SCOPE(TEXT("Apply InstantDamage"))
 
 		UGameplayEffect * BaseDmgEffect = Cast<UGameplayEffect>(StaticConstructObject(UGameplayEffect::StaticClass(), GetTransientPackage(), FName(TEXT("Damage"))));
 		BaseDmgEffect->Modifiers.SetNum(1);
@@ -2780,10 +2780,10 @@ bool GameplayEffectsTest_InstantDamageModifierPassesTag(UWorld *World, FAutomati
 		SourceComponent->ApplyGameplayEffectToTarget(BaseDmgEffect, DestComponent, 1.f);
 
 		float ExpectedValue = (StartHealth + ((DamageValue * BonusDamageMultiplier) / DamageProtectionDivisor));
-		float ActualValue = DestComponent->GetSet<USkillSystemTestAttributeSet>()->Health;
+		float ActualValue = DestComponent->GetSet<UAbilitySystemTestAttributeSet>()->Health;
 
 		Test->TestTrue(SKILL_TEST_TEXT("Instant Damage Required Tag No Protection.  Actual: %.2f == Exected: %.2f", ActualValue, ExpectedValue), (ActualValue == ExpectedValue));
-		SKILL_LOG(Log, TEXT("Final Health: %.2f"), DestComponent->GetSet<USkillSystemTestAttributeSet>()->Health);
+		ABILITY_LOG(Log, TEXT("Final Health: %.2f"), DestComponent->GetSet<UAbilitySystemTestAttributeSet>()->Health);
 	}
 
 	World->EditorDestroyActor(SourceActor, false);
@@ -2799,19 +2799,19 @@ bool GameplayEffectsTest_InstantDamageModifierTag(UWorld *World, FAutomationTest
 	const float BonusDamageValue = -10.f;
 	const float DamageProtectionDivisor = 2.f;
 
-	ASkillSystemTestPawn *SourceActor = World->SpawnActor<ASkillSystemTestPawn>();
-	ASkillSystemTestPawn *DestActor = World->SpawnActor<ASkillSystemTestPawn>();
+	AAbilitySystemTestPawn *SourceActor = World->SpawnActor<AAbilitySystemTestPawn>();
+	AAbilitySystemTestPawn *DestActor = World->SpawnActor<AAbilitySystemTestPawn>();
 
-	UProperty *HealthProperty = FindFieldChecked<UProperty>(USkillSystemTestAttributeSet::StaticClass(), GET_MEMBER_NAME_CHECKED(USkillSystemTestAttributeSet, Health));
+	UProperty *HealthProperty = FindFieldChecked<UProperty>(UAbilitySystemTestAttributeSet::StaticClass(), GET_MEMBER_NAME_CHECKED(UAbilitySystemTestAttributeSet, Health));
 
 	UAttributeComponent * SourceComponent = SourceActor->AttributeComponent;
 	UAttributeComponent * DestComponent = DestActor->AttributeComponent;
-	SourceComponent->GetSet<USkillSystemTestAttributeSet>()->Health = StartHealth;
-	DestComponent->GetSet<USkillSystemTestAttributeSet>()->Health = StartHealth;
+	SourceComponent->GetSet<UAbilitySystemTestAttributeSet>()->Health = StartHealth;
+	DestComponent->GetSet<UAbilitySystemTestAttributeSet>()->Health = StartHealth;
 
 	// Setup a GE to modify OutgoingGEs
 	{
-		SKILL_LOG_SCOPE(TEXT("Apply DamageBuff"))
+		ABILITY_LOG_SCOPE(TEXT("Apply DamageBuff"))
 
 		UGameplayEffect * BaseDmgEffect = Cast<UGameplayEffect>(StaticConstructObject(UGameplayEffect::StaticClass(), GetTransientPackage(), FName(TEXT("DamageBuff"))));
 		BaseDmgEffect->Modifiers.SetNum(1);
@@ -2828,7 +2828,7 @@ bool GameplayEffectsTest_InstantDamageModifierTag(UWorld *World, FAutomationTest
 
 	// Setup a GE to modify IncomingGEs
 	{
-		SKILL_LOG_SCOPE(TEXT("Apply ProtectionBuff"))
+		ABILITY_LOG_SCOPE(TEXT("Apply ProtectionBuff"))
 
 		UGameplayEffect* BaseProtectEffect = Cast<UGameplayEffect>(StaticConstructObject(UGameplayEffect::StaticClass(), GetTransientPackage(), FName(TEXT("ProtectBuff"))));
 		BaseProtectEffect->Modifiers.SetNum(1);
@@ -2847,7 +2847,7 @@ bool GameplayEffectsTest_InstantDamageModifierTag(UWorld *World, FAutomationTest
 
 	// Apply Damage
 	{
-		SKILL_LOG_SCOPE(TEXT("Apply InstantDamage"))
+		ABILITY_LOG_SCOPE(TEXT("Apply InstantDamage"))
 
 		UGameplayEffect * BaseDmgEffect = Cast<UGameplayEffect>(StaticConstructObject(UGameplayEffect::StaticClass(), GetTransientPackage(), FName(TEXT("Damage"))));
 		BaseDmgEffect->Modifiers.SetNum(1);
@@ -2862,7 +2862,7 @@ bool GameplayEffectsTest_InstantDamageModifierTag(UWorld *World, FAutomationTest
 		SourceComponent->ApplyGameplayEffectToTarget(BaseDmgEffect, DestComponent, 1.f);
 
 		float ExpectedValue = (StartHealth + ((DamageValue + BonusDamageValue) / DamageProtectionDivisor));
-		float ActualValue = DestComponent->GetSet<USkillSystemTestAttributeSet>()->Health;
+		float ActualValue = DestComponent->GetSet<UAbilitySystemTestAttributeSet>()->Health;
 
 		Test->TestTrue(SKILL_TEST_TEXT("Buff Instant Damage Applied.  Actual: %.2f == Exected: %.2f", ActualValue, ExpectedValue), (ActualValue == ExpectedValue));
 	}
@@ -2884,21 +2884,21 @@ bool GameplayEffectsTest_InstantDamage_ScalingProperty(UWorld *World, FAutomatio
 
 	SetGlobalCurveTable();
 
-	ASkillSystemTestPawn *SourceActor = World->SpawnActor<ASkillSystemTestPawn>();
-	ASkillSystemTestPawn *DestActor = World->SpawnActor<ASkillSystemTestPawn>();
+	AAbilitySystemTestPawn *SourceActor = World->SpawnActor<AAbilitySystemTestPawn>();
+	AAbilitySystemTestPawn *DestActor = World->SpawnActor<AAbilitySystemTestPawn>();
 
-	UProperty *DamageProperty = FindFieldChecked<UProperty>(USkillSystemTestAttributeSet::StaticClass(), GET_MEMBER_NAME_CHECKED(USkillSystemTestAttributeSet, Damage));
-	UProperty *PhysicalDamageProperty = FindFieldChecked<UProperty>(USkillSystemTestAttributeSet::StaticClass(), GET_MEMBER_NAME_CHECKED(USkillSystemTestAttributeSet, PhysicalDamage));
+	UProperty *DamageProperty = FindFieldChecked<UProperty>(UAbilitySystemTestAttributeSet::StaticClass(), GET_MEMBER_NAME_CHECKED(UAbilitySystemTestAttributeSet, Damage));
+	UProperty *PhysicalDamageProperty = FindFieldChecked<UProperty>(UAbilitySystemTestAttributeSet::StaticClass(), GET_MEMBER_NAME_CHECKED(UAbilitySystemTestAttributeSet, PhysicalDamage));
 
 	UAttributeComponent * SourceComponent = SourceActor->AttributeComponent;
 	UAttributeComponent * DestComponent = DestActor->AttributeComponent;
-	SourceComponent->GetSet<USkillSystemTestAttributeSet>()->Health = StartHealth;
-	DestComponent->GetSet<USkillSystemTestAttributeSet>()->Health = StartHealth;
+	SourceComponent->GetSet<UAbilitySystemTestAttributeSet>()->Health = StartHealth;
+	DestComponent->GetSet<UAbilitySystemTestAttributeSet>()->Health = StartHealth;
 
-	SourceComponent->GetSet<USkillSystemTestAttributeSet>()->PhysicalDamage = PhysicalDamage;
+	SourceComponent->GetSet<UAbilitySystemTestAttributeSet>()->PhysicalDamage = PhysicalDamage;
 
 	{
-		SKILL_LOG_SCOPE(TEXT("Apply InstantDamage"));
+		ABILITY_LOG_SCOPE(TEXT("Apply InstantDamage"));
 
 		// This effects do Damage = 1.f * LinearCurve[LevelOfGameplayEffect].
 		// This translate into 1.f * PhysicalDamage.
@@ -2917,7 +2917,7 @@ bool GameplayEffectsTest_InstantDamage_ScalingProperty(UWorld *World, FAutomatio
 
 		SourceComponent->ApplyGameplayEffectToTarget(BaseDmgEffect, DestComponent);
 
-		float ActualValue = DestComponent->GetSet<USkillSystemTestAttributeSet>()->Health;
+		float ActualValue = DestComponent->GetSet<UAbilitySystemTestAttributeSet>()->Health;
 		float ExpectedValue = StartHealth - (PhysicalDamage * GameplayEffectScaling);
 		Test->TestTrue(SKILL_TEST_TEXT("Damage Applied. Health: Actual: %.2f == Exected: %.2f", ActualValue, ExpectedValue), ActualValue == ExpectedValue);
 	}
@@ -2940,21 +2940,21 @@ bool GameplayEffectsTest_InstantDamage_ScalingPropertyNested(UWorld *World, FAut
 
 	SetGlobalCurveTable();
 
-	ASkillSystemTestPawn *SourceActor = World->SpawnActor<ASkillSystemTestPawn>();
-	ASkillSystemTestPawn *DestActor = World->SpawnActor<ASkillSystemTestPawn>();
+	AAbilitySystemTestPawn *SourceActor = World->SpawnActor<AAbilitySystemTestPawn>();
+	AAbilitySystemTestPawn *DestActor = World->SpawnActor<AAbilitySystemTestPawn>();
 
-	UProperty *DamageProperty = FindFieldChecked<UProperty>(USkillSystemTestAttributeSet::StaticClass(), GET_MEMBER_NAME_CHECKED(USkillSystemTestAttributeSet, Damage));
-	UProperty *PhysicalDamageProperty = FindFieldChecked<UProperty>(USkillSystemTestAttributeSet::StaticClass(), GET_MEMBER_NAME_CHECKED(USkillSystemTestAttributeSet, PhysicalDamage));
+	UProperty *DamageProperty = FindFieldChecked<UProperty>(UAbilitySystemTestAttributeSet::StaticClass(), GET_MEMBER_NAME_CHECKED(UAbilitySystemTestAttributeSet, Damage));
+	UProperty *PhysicalDamageProperty = FindFieldChecked<UProperty>(UAbilitySystemTestAttributeSet::StaticClass(), GET_MEMBER_NAME_CHECKED(UAbilitySystemTestAttributeSet, PhysicalDamage));
 
 	UAttributeComponent * SourceComponent = SourceActor->AttributeComponent;
 	UAttributeComponent * DestComponent = DestActor->AttributeComponent;
-	SourceComponent->GetSet<USkillSystemTestAttributeSet>()->Health = StartHealth;
-	DestComponent->GetSet<USkillSystemTestAttributeSet>()->Health = StartHealth;
+	SourceComponent->GetSet<UAbilitySystemTestAttributeSet>()->Health = StartHealth;
+	DestComponent->GetSet<UAbilitySystemTestAttributeSet>()->Health = StartHealth;
 
-	SourceComponent->GetSet<USkillSystemTestAttributeSet>()->PhysicalDamage = PhysicalDamage;
+	SourceComponent->GetSet<UAbilitySystemTestAttributeSet>()->PhysicalDamage = PhysicalDamage;
 
 	{
-		SKILL_LOG_SCOPE(TEXT("Apply InstantDamage"));
+		ABILITY_LOG_SCOPE(TEXT("Apply InstantDamage"));
 
 		// This effects do Damage = 1.f * LinearCurve[LevelOfGameplayEffect].
 		// This translate into 1.f * PhysicalDamage.
@@ -2972,7 +2972,7 @@ bool GameplayEffectsTest_InstantDamage_ScalingPropertyNested(UWorld *World, FAut
 
 		SourceComponent->ApplyGameplayEffectToTarget(BaseDmgEffect, DestComponent);
 
-		float ActualValue = DestComponent->GetSet<USkillSystemTestAttributeSet>()->Health;
+		float ActualValue = DestComponent->GetSet<UAbilitySystemTestAttributeSet>()->Health;
 		float ExpectedValue = StartHealth - (PhysicalDamage * GameplayEffectScaling);
 		Test->TestTrue(SKILL_TEST_TEXT("Damage Applied. Health: Actual: %.2f == Exected: %.2f", ActualValue, ExpectedValue), ActualValue == ExpectedValue);
 	}
@@ -2997,21 +2997,21 @@ bool GameplayEffectsTest_DotDamage_ScalingProperty_Snapshot(UWorld *World, FAuto
 
 	SetGlobalCurveTable();
 
-	ASkillSystemTestPawn *SourceActor = World->SpawnActor<ASkillSystemTestPawn>();
-	ASkillSystemTestPawn *DestActor = World->SpawnActor<ASkillSystemTestPawn>();
+	AAbilitySystemTestPawn *SourceActor = World->SpawnActor<AAbilitySystemTestPawn>();
+	AAbilitySystemTestPawn *DestActor = World->SpawnActor<AAbilitySystemTestPawn>();
 
-	UProperty *DamageProperty = FindFieldChecked<UProperty>(USkillSystemTestAttributeSet::StaticClass(), GET_MEMBER_NAME_CHECKED(USkillSystemTestAttributeSet, Damage));
-	UProperty *SpellDamageProperty = FindFieldChecked<UProperty>(USkillSystemTestAttributeSet::StaticClass(), GET_MEMBER_NAME_CHECKED(USkillSystemTestAttributeSet, SpellDamage));
+	UProperty *DamageProperty = FindFieldChecked<UProperty>(UAbilitySystemTestAttributeSet::StaticClass(), GET_MEMBER_NAME_CHECKED(UAbilitySystemTestAttributeSet, Damage));
+	UProperty *SpellDamageProperty = FindFieldChecked<UProperty>(UAbilitySystemTestAttributeSet::StaticClass(), GET_MEMBER_NAME_CHECKED(UAbilitySystemTestAttributeSet, SpellDamage));
 
 	UAttributeComponent * SourceComponent = SourceActor->AttributeComponent;
 	UAttributeComponent * DestComponent = DestActor->AttributeComponent;
-	SourceComponent->GetSet<USkillSystemTestAttributeSet>()->Health = StartHealth;
-	DestComponent->GetSet<USkillSystemTestAttributeSet>()->Health = StartHealth;
+	SourceComponent->GetSet<UAbilitySystemTestAttributeSet>()->Health = StartHealth;
+	DestComponent->GetSet<UAbilitySystemTestAttributeSet>()->Health = StartHealth;
 
-	SourceComponent->GetSet<USkillSystemTestAttributeSet>()->SpellDamage = SpellDamage;
+	SourceComponent->GetSet<UAbilitySystemTestAttributeSet>()->SpellDamage = SpellDamage;
 
 	{
-		SKILL_LOG_SCOPE(TEXT("Apply InstantDamage"));
+		ABILITY_LOG_SCOPE(TEXT("Apply InstantDamage"));
 
 		// This effects do Damage = 1.f * LinearCurve[LevelOfGameplayEffect].
 		// This translate into 1.f * PhysicalDamage.
@@ -3057,7 +3057,7 @@ bool GameplayEffectsTest_DotDamage_ScalingProperty_Snapshot(UWorld *World, FAuto
 		// Tick once
 		GameplayTest_TickWorld(World, 1.f);
 		
-		float ActualValue = DestComponent->GetSet<USkillSystemTestAttributeSet>()->Health;
+		float ActualValue = DestComponent->GetSet<UAbilitySystemTestAttributeSet>()->Health;
 		float ExpectedValue = StartHealth - (SpellDamage * GameplayEffectScaling) * 2.f;	// we trigger twice on this tick
 		Test->TestTrue(SKILL_TEST_TEXT("Damage Applied. Health: Actual: %.2f == Exected: %.2f", ActualValue, ExpectedValue), ActualValue == ExpectedValue);
 	}	
@@ -3083,21 +3083,21 @@ bool GameplayEffectsTest_DotDamage_ScalingProperty_Dynamic(UWorld *World, FAutom
 
 	SetGlobalCurveTable();
 
-	ASkillSystemTestPawn *SourceActor = World->SpawnActor<ASkillSystemTestPawn>();
-	ASkillSystemTestPawn *DestActor = World->SpawnActor<ASkillSystemTestPawn>();
+	AAbilitySystemTestPawn *SourceActor = World->SpawnActor<AAbilitySystemTestPawn>();
+	AAbilitySystemTestPawn *DestActor = World->SpawnActor<AAbilitySystemTestPawn>();
 
-	UProperty *DamageProperty = FindFieldChecked<UProperty>(USkillSystemTestAttributeSet::StaticClass(), GET_MEMBER_NAME_CHECKED(USkillSystemTestAttributeSet, Damage));
-	UProperty *SpellDamageProperty = FindFieldChecked<UProperty>(USkillSystemTestAttributeSet::StaticClass(), GET_MEMBER_NAME_CHECKED(USkillSystemTestAttributeSet, SpellDamage));
+	UProperty *DamageProperty = FindFieldChecked<UProperty>(UAbilitySystemTestAttributeSet::StaticClass(), GET_MEMBER_NAME_CHECKED(UAbilitySystemTestAttributeSet, Damage));
+	UProperty *SpellDamageProperty = FindFieldChecked<UProperty>(UAbilitySystemTestAttributeSet::StaticClass(), GET_MEMBER_NAME_CHECKED(UAbilitySystemTestAttributeSet, SpellDamage));
 
 	UAttributeComponent * SourceComponent = SourceActor->AttributeComponent;
 	UAttributeComponent * DestComponent = DestActor->AttributeComponent;
-	SourceComponent->GetSet<USkillSystemTestAttributeSet>()->Health = StartHealth;
-	DestComponent->GetSet<USkillSystemTestAttributeSet>()->Health = StartHealth;
+	SourceComponent->GetSet<UAbilitySystemTestAttributeSet>()->Health = StartHealth;
+	DestComponent->GetSet<UAbilitySystemTestAttributeSet>()->Health = StartHealth;
 
-	SourceComponent->GetSet<USkillSystemTestAttributeSet>()->SpellDamage = SpellDamage;
+	SourceComponent->GetSet<UAbilitySystemTestAttributeSet>()->SpellDamage = SpellDamage;
 
 	{
-		SKILL_LOG_SCOPE(TEXT("Apply InstantDamage"));
+		ABILITY_LOG_SCOPE(TEXT("Apply InstantDamage"));
 
 		// This effects do Damage = 1.f * LinearCurve[LevelOfGameplayEffect].
 		// This translate into 1.f * PhysicalDamage.
@@ -3119,7 +3119,7 @@ bool GameplayEffectsTest_DotDamage_ScalingProperty_Dynamic(UWorld *World, FAutom
 		SourceComponent->ApplyGameplayEffectToTarget(BaseDmgEffect, DestComponent);
 		GameplayTest_TickWorld(World, 0.0001f); // Move our Effects from the pending stack to the active stack inside the timer manager, this starts the clock for effect execution
 
-		float SpellDamageTest = SourceComponent->GetSet<USkillSystemTestAttributeSet>()->SpellDamage;
+		float SpellDamageTest = SourceComponent->GetSet<UAbilitySystemTestAttributeSet>()->SpellDamage;
 		check(SpellDamageTest == SpellDamage);
 
 		GameplayTest_TickWorld(World, 0.1f); // Offset the current time from the start of the period so that floating point issues don't affect the test
@@ -3129,7 +3129,7 @@ bool GameplayEffectsTest_DotDamage_ScalingProperty_Dynamic(UWorld *World, FAutom
 		// Tick once
 		GameplayTest_TickWorld(World, 1.f);
 
-		float ActualValue = DestComponent->GetSet<USkillSystemTestAttributeSet>()->Health;
+		float ActualValue = DestComponent->GetSet<UAbilitySystemTestAttributeSet>()->Health;
 		float ExpectedValue = StartHealth - (SpellDamage * GameplayEffectScaling * 2);	// We've ticked twice
 		Test->TestTrue(SKILL_TEST_TEXT("Damage Applied. Health: Actual: %.2f == Exected: %.2f", ActualValue, ExpectedValue), ActualValue == ExpectedValue);
 	}
@@ -3148,7 +3148,7 @@ bool GameplayEffectsTest_DotDamage_ScalingProperty_Dynamic(UWorld *World, FAutom
 		SourceComponent->ApplyGameplayEffectToTarget(SpellDmgEffect, SourceComponent);
 		GameplayTest_TickWorld(World, 0.0001f); // Move our Effects from the pending stack to the active stack inside the timer manager, this starts the clock for effect execution
 
-		float ActualValue = SourceComponent->GetSet<USkillSystemTestAttributeSet>()->SpellDamage;
+		float ActualValue = SourceComponent->GetSet<UAbilitySystemTestAttributeSet>()->SpellDamage;
 		float ExpectedValue = SpellDamage2;
 		Test->TestTrue(SKILL_TEST_TEXT("Spell Damage Mod: Actual: %.2f == Exected: %.2f", ActualValue, ExpectedValue), ActualValue == ExpectedValue);
 
@@ -3159,7 +3159,7 @@ bool GameplayEffectsTest_DotDamage_ScalingProperty_Dynamic(UWorld *World, FAutom
 		// Tick once
 		GameplayTest_TickWorld(World, 1.f);
 
-		float ActualValue = DestComponent->GetSet<USkillSystemTestAttributeSet>()->Health;
+		float ActualValue = DestComponent->GetSet<UAbilitySystemTestAttributeSet>()->Health;
 		float ExpectedValue = StartHealth - (SpellDamage * GameplayEffectScaling * 2) - (SpellDamage2 * GameplayEffectScaling);	// two unbuffed ticks, one buffed tick
 		Test->TestTrue(SKILL_TEST_TEXT("Damage Applied. Health: Actual: %.2f == Exected: %.2f", ActualValue, ExpectedValue), ActualValue == ExpectedValue);
 	}
@@ -3182,18 +3182,18 @@ bool GameplayEffectsTest_MetaAttributes(UWorld *World, FAutomationTestBase * Tes
 
 	SetGlobalCurveTable();
 
-	ASkillSystemTestPawn *SourceActor = World->SpawnActor<ASkillSystemTestPawn>();
+	AAbilitySystemTestPawn *SourceActor = World->SpawnActor<AAbilitySystemTestPawn>();
 	
 
-	UProperty *MaxHealthProperty = FindFieldChecked<UProperty>(USkillSystemTestAttributeSet::StaticClass(), GET_MEMBER_NAME_CHECKED(USkillSystemTestAttributeSet, MaxHealth));
-	UProperty *StrengthProperty = FindFieldChecked<UProperty>(USkillSystemTestAttributeSet::StaticClass(), GET_MEMBER_NAME_CHECKED(USkillSystemTestAttributeSet, Strength));
+	UProperty *MaxHealthProperty = FindFieldChecked<UProperty>(UAbilitySystemTestAttributeSet::StaticClass(), GET_MEMBER_NAME_CHECKED(UAbilitySystemTestAttributeSet, MaxHealth));
+	UProperty *StrengthProperty = FindFieldChecked<UProperty>(UAbilitySystemTestAttributeSet::StaticClass(), GET_MEMBER_NAME_CHECKED(UAbilitySystemTestAttributeSet, Strength));
 
 	UAttributeComponent * SourceComponent = SourceActor->AttributeComponent;
-	SourceComponent->GetSet<USkillSystemTestAttributeSet>()->MaxHealth = StartHealth;
-	SourceComponent->GetSet<USkillSystemTestAttributeSet>()->Strength = 0.f;
+	SourceComponent->GetSet<UAbilitySystemTestAttributeSet>()->MaxHealth = StartHealth;
+	SourceComponent->GetSet<UAbilitySystemTestAttributeSet>()->Strength = 0.f;
 
 	{
-		SKILL_LOG_SCOPE(TEXT("Setup meta stat"));
+		ABILITY_LOG_SCOPE(TEXT("Setup meta stat"));
 
 		// This effects do Damage = 1.f * LinearCurve[LevelOfGameplayEffect].
 		// This translate into 1.f * PhysicalDamage.
@@ -3215,7 +3215,7 @@ bool GameplayEffectsTest_MetaAttributes(UWorld *World, FAutomationTestBase * Tes
 		SourceComponent->ApplyGameplayEffectToTarget(StrengthMaxHealhEffect, SourceComponent);
 
 		// Strength starts at 0, so confirm it did nothing yet.
-		float ActualValue = SourceComponent->GetSet<USkillSystemTestAttributeSet>()->MaxHealth;
+		float ActualValue = SourceComponent->GetSet<UAbilitySystemTestAttributeSet>()->MaxHealth;
 		float ExpectedValue = StartHealth;
 		Test->TestTrue(SKILL_TEST_TEXT("Damage Applied. Health: Actual: %.2f == Exected: %.2f", ActualValue, ExpectedValue), ActualValue == ExpectedValue);
 	}
@@ -3234,13 +3234,13 @@ bool GameplayEffectsTest_MetaAttributes(UWorld *World, FAutomationTestBase * Tes
 		SourceComponent->ApplyGameplayEffectToTarget(StrEffect, SourceComponent);
 
 		{
-			float ActualValue = SourceComponent->GetSet<USkillSystemTestAttributeSet>()->Strength;
+			float ActualValue = SourceComponent->GetSet<UAbilitySystemTestAttributeSet>()->Strength;
 			float ExpectedValue = StrengthValue;
 			Test->TestTrue(SKILL_TEST_TEXT("Strength: Actual: %.2f == Exected: %.2f", ActualValue, ExpectedValue), ActualValue == ExpectedValue);
 		}
 
 		{
-			float ActualValue = SourceComponent->GetSet<USkillSystemTestAttributeSet>()->MaxHealth;
+			float ActualValue = SourceComponent->GetSet<UAbilitySystemTestAttributeSet>()->MaxHealth;
 			float ExpectedValue = StartHealth + (MaxHealthPerStrength * StrengthValue);
 			Test->TestTrue(SKILL_TEST_TEXT("MaxHealth: Actual: %.2f == Exected: %.2f", ActualValue, ExpectedValue), ActualValue == ExpectedValue);
 		}
@@ -3259,20 +3259,20 @@ bool GameplayEffectsTest_TagOrdering(UWorld *World, FAutomationTestBase * Test)
 	const float DamageValue = 5.f;
 	const float BonusDamageMultiplier = 2.f;
 
-	ASkillSystemTestPawn *SourceActor = World->SpawnActor<ASkillSystemTestPawn>();
-	ASkillSystemTestPawn *DestActor = World->SpawnActor<ASkillSystemTestPawn>();
+	AAbilitySystemTestPawn *SourceActor = World->SpawnActor<AAbilitySystemTestPawn>();
+	AAbilitySystemTestPawn *DestActor = World->SpawnActor<AAbilitySystemTestPawn>();
 
-	UProperty *HealthProperty = FindFieldChecked<UProperty>(USkillSystemTestAttributeSet::StaticClass(), GET_MEMBER_NAME_CHECKED(USkillSystemTestAttributeSet, Health));
-	UProperty *DamageProperty = FindFieldChecked<UProperty>(USkillSystemTestAttributeSet::StaticClass(), GET_MEMBER_NAME_CHECKED(USkillSystemTestAttributeSet, Damage));
+	UProperty *HealthProperty = FindFieldChecked<UProperty>(UAbilitySystemTestAttributeSet::StaticClass(), GET_MEMBER_NAME_CHECKED(UAbilitySystemTestAttributeSet, Health));
+	UProperty *DamageProperty = FindFieldChecked<UProperty>(UAbilitySystemTestAttributeSet::StaticClass(), GET_MEMBER_NAME_CHECKED(UAbilitySystemTestAttributeSet, Damage));
 
 	UAttributeComponent * SourceComponent = SourceActor->AttributeComponent;
 	UAttributeComponent * DestComponent = DestActor->AttributeComponent;
-	SourceComponent->GetSet<USkillSystemTestAttributeSet>()->Health = StartHealth;
-	DestComponent->GetSet<USkillSystemTestAttributeSet>()->Health = StartHealth;
+	SourceComponent->GetSet<UAbilitySystemTestAttributeSet>()->Health = StartHealth;
+	DestComponent->GetSet<UAbilitySystemTestAttributeSet>()->Health = StartHealth;
 
 	// Setup a GE to modify OutgoingGEs
 	{
-		SKILL_LOG_SCOPE(TEXT("FireDamageBuff"));
+		ABILITY_LOG_SCOPE(TEXT("FireDamageBuff"));
 
 		UGameplayEffect * BaseDmgEffect = Cast<UGameplayEffect>(StaticConstructObject(UGameplayEffect::StaticClass(), GetTransientPackage(), FName(TEXT("FireDamageBuff"))));
 		BaseDmgEffect->Modifiers.SetNum(1);
@@ -3292,7 +3292,7 @@ bool GameplayEffectsTest_TagOrdering(UWorld *World, FAutomationTestBase * Test)
 	}
 
 	{
-		SKILL_LOG_SCOPE(TEXT("MakeFireDamage"));
+		ABILITY_LOG_SCOPE(TEXT("MakeFireDamage"));
 
 		UGameplayEffect * BaseDmgEffect = Cast<UGameplayEffect>(StaticConstructObject(UGameplayEffect::StaticClass(), GetTransientPackage(), FName(TEXT("MakeFireDamage"))));
 		BaseDmgEffect->Modifiers.SetNum(1);
@@ -3313,7 +3313,7 @@ bool GameplayEffectsTest_TagOrdering(UWorld *World, FAutomationTestBase * Test)
 
 	// Apply Damage
 	{
-		SKILL_LOG_SCOPE(TEXT("Apply InstantDamage"));
+		ABILITY_LOG_SCOPE(TEXT("Apply InstantDamage"));
 
 		UGameplayEffect * BaseDmgEffect = Cast<UGameplayEffect>(StaticConstructObject(UGameplayEffect::StaticClass(), GetTransientPackage(), FName(TEXT("Damage"))));
 		BaseDmgEffect->Modifiers.SetNum(1);
@@ -3330,7 +3330,7 @@ bool GameplayEffectsTest_TagOrdering(UWorld *World, FAutomationTestBase * Test)
 		SourceComponent->ApplyGameplayEffectToTarget(BaseDmgEffect, DestComponent, 1.f);
 
 		float ExpectedValue = (StartHealth - (DamageValue * BonusDamageMultiplier));
-		float ActualValue = DestComponent->GetSet<USkillSystemTestAttributeSet>()->Health;
+		float ActualValue = DestComponent->GetSet<UAbilitySystemTestAttributeSet>()->Health;
 		
 		Test->TestTrue(SKILL_TEST_TEXT("MaxHealth: Actual: %.2f == Exected: %.2f", ActualValue, ExpectedValue), ActualValue == ExpectedValue);
 	}
@@ -3349,16 +3349,16 @@ bool GameplayEffectsTest_StackingHighest(UWorld *World, FAutomationTestBase * Te
 	SetGlobalCurveTable();
 	SetGlobalDataTable();
 
-	ASkillSystemTestPawn *SourceActor = World->SpawnActor<ASkillSystemTestPawn>();
-	ASkillSystemTestPawn *DestActor = World->SpawnActor<ASkillSystemTestPawn>();
+	AAbilitySystemTestPawn *SourceActor = World->SpawnActor<AAbilitySystemTestPawn>();
+	AAbilitySystemTestPawn *DestActor = World->SpawnActor<AAbilitySystemTestPawn>();
 
-	UProperty *StackingProperty = FindFieldChecked<UProperty>(USkillSystemTestAttributeSet::StaticClass(), GET_MEMBER_NAME_CHECKED(USkillSystemTestAttributeSet, StackingAttribute1));
+	UProperty *StackingProperty = FindFieldChecked<UProperty>(UAbilitySystemTestAttributeSet::StaticClass(), GET_MEMBER_NAME_CHECKED(UAbilitySystemTestAttributeSet, StackingAttribute1));
 
 	UAttributeComponent * SourceComponent = SourceActor->AttributeComponent;
 	UAttributeComponent * DestComponent = DestActor->AttributeComponent;
 
 	{
-		SKILL_LOG_SCOPE(TEXT("Apply StackingEffect"));
+		ABILITY_LOG_SCOPE(TEXT("Apply StackingEffect"));
 
 		UGameplayEffect * BaseStackedEffect = Cast<UGameplayEffect>(StaticConstructObject(UGameplayEffect::StaticClass(), GetTransientPackage(), FName(TEXT("StackingEffect1"))));
 		BaseStackedEffect->Modifiers.SetNum(1);
@@ -3377,7 +3377,7 @@ bool GameplayEffectsTest_StackingHighest(UWorld *World, FAutomationTestBase * Te
 
 		GameplayTest_TickWorld(World, SMALL_NUMBER); // Move our Effects from the pending stack to the active stack inside the timer manager, this starts the clock for effect execution
 		{
-			float ActualValue = DestComponent->GetSet<USkillSystemTestAttributeSet>()->StackingAttribute1;
+			float ActualValue = DestComponent->GetSet<UAbilitySystemTestAttributeSet>()->StackingAttribute1;
 			float ExpectedValue = StackingValue;
 			Test->TestTrue(SKILL_TEST_TEXT("Stacking Attribute: Actual: %.2f == Exected: %.2f", ActualValue, ExpectedValue), ActualValue == ExpectedValue);
 		}
@@ -3388,13 +3388,13 @@ bool GameplayEffectsTest_StackingHighest(UWorld *World, FAutomationTestBase * Te
 		// Tick once
 		GameplayTest_TickWorld(World, 1.f);
 
-		float ActualValue = DestComponent->GetSet<USkillSystemTestAttributeSet>()->StackingAttribute1;
+		float ActualValue = DestComponent->GetSet<UAbilitySystemTestAttributeSet>()->StackingAttribute1;
 		float ExpectedValue = StackingValue * 2.f;	// effect will execute twice
 		Test->TestTrue(SKILL_TEST_TEXT("Stacking Attribute: Actual: %.2f == Exected: %.2f", ActualValue, ExpectedValue), ActualValue == ExpectedValue);
 	}
 
 	{
-		SKILL_LOG_SCOPE(TEXT("Apply StackingEffect"));
+		ABILITY_LOG_SCOPE(TEXT("Apply StackingEffect"));
 
 		UGameplayEffect * BaseStackedEffect = Cast<UGameplayEffect>(StaticConstructObject(UGameplayEffect::StaticClass(), GetTransientPackage(), FName(TEXT("StackingEffect2"))));
 		BaseStackedEffect->Modifiers.SetNum(1);
@@ -3419,7 +3419,7 @@ bool GameplayEffectsTest_StackingHighest(UWorld *World, FAutomationTestBase * Te
 		// Tick once
 		GameplayTest_TickWorld(World, 1.f);
 
-		float ActualValue = DestComponent->GetSet<USkillSystemTestAttributeSet>()->StackingAttribute1;
+		float ActualValue = DestComponent->GetSet<UAbilitySystemTestAttributeSet>()->StackingAttribute1;
 		float ExpectedValue = StackingValue * 4.f;	// 2 for first effect, 2 for the last tick with both effects
 		Test->TestTrue(SKILL_TEST_TEXT("Stacking Attribute: Actual: %.2f == Exected: %.2f", ActualValue, ExpectedValue), ActualValue == ExpectedValue);
 	}
@@ -3441,16 +3441,16 @@ bool GameplayEffectsTest_StackingLowest(UWorld *World, FAutomationTestBase * Tes
 	SetGlobalCurveTable();
 	SetGlobalDataTable();
 
-	ASkillSystemTestPawn *SourceActor = World->SpawnActor<ASkillSystemTestPawn>();
-	ASkillSystemTestPawn *DestActor = World->SpawnActor<ASkillSystemTestPawn>();
+	AAbilitySystemTestPawn *SourceActor = World->SpawnActor<AAbilitySystemTestPawn>();
+	AAbilitySystemTestPawn *DestActor = World->SpawnActor<AAbilitySystemTestPawn>();
 
-	UProperty *StackingProperty = FindFieldChecked<UProperty>(USkillSystemTestAttributeSet::StaticClass(), GET_MEMBER_NAME_CHECKED(USkillSystemTestAttributeSet, StackingAttribute1));
+	UProperty *StackingProperty = FindFieldChecked<UProperty>(UAbilitySystemTestAttributeSet::StaticClass(), GET_MEMBER_NAME_CHECKED(UAbilitySystemTestAttributeSet, StackingAttribute1));
 
 	UAttributeComponent * SourceComponent = SourceActor->AttributeComponent;
 	UAttributeComponent * DestComponent = DestActor->AttributeComponent;
 
 	{
-		SKILL_LOG_SCOPE(TEXT("Apply StackingEffect"));
+		ABILITY_LOG_SCOPE(TEXT("Apply StackingEffect"));
 
 		UGameplayEffect * BaseStackedEffect = Cast<UGameplayEffect>(StaticConstructObject(UGameplayEffect::StaticClass(), GetTransientPackage(), FName(TEXT("StackingEffect1"))));
 		BaseStackedEffect->Modifiers.SetNum(1);
@@ -3475,13 +3475,13 @@ bool GameplayEffectsTest_StackingLowest(UWorld *World, FAutomationTestBase * Tes
 		// Tick once
 		GameplayTest_TickWorld(World, 1.f);
 
-		float ActualValue = DestComponent->GetSet<USkillSystemTestAttributeSet>()->StackingAttribute1;
+		float ActualValue = DestComponent->GetSet<UAbilitySystemTestAttributeSet>()->StackingAttribute1;
 		float ExpectedValue = StackingValue * 2.f;	// the effect should execute twice here
 		Test->TestTrue(SKILL_TEST_TEXT("Stacking Attribute: Actual: %.2f == Exected: %.2f", ActualValue, ExpectedValue), ActualValue == ExpectedValue);
 	}
 
 	{
-		SKILL_LOG_SCOPE(TEXT("Apply StackingEffect"));
+		ABILITY_LOG_SCOPE(TEXT("Apply StackingEffect"));
 
 		UGameplayEffect * BaseStackedEffect = Cast<UGameplayEffect>(StaticConstructObject(UGameplayEffect::StaticClass(), GetTransientPackage(), FName(TEXT("StackingEffect2"))));
 		BaseStackedEffect->Modifiers.SetNum(1);
@@ -3506,7 +3506,7 @@ bool GameplayEffectsTest_StackingLowest(UWorld *World, FAutomationTestBase * Tes
 		// Tick once
 		GameplayTest_TickWorld(World, 1.f);
 
-		float ActualValue = DestComponent->GetSet<USkillSystemTestAttributeSet>()->StackingAttribute1;
+		float ActualValue = DestComponent->GetSet<UAbilitySystemTestAttributeSet>()->StackingAttribute1;
 		float ExpectedValue = StackingValue * 3.f;	// the first effect has executed 3 times, the second hasn't executed
 		Test->TestTrue(SKILL_TEST_TEXT("Stacking Attribute: Actual: %.2f == Exected: %.2f", ActualValue, ExpectedValue), ActualValue == ExpectedValue);
 	}
@@ -3528,16 +3528,16 @@ bool GameplayEffectsTest_StackingUnlimited(UWorld *World, FAutomationTestBase *
 	SetGlobalCurveTable();
 	SetGlobalDataTable();
 
-	ASkillSystemTestPawn *SourceActor = World->SpawnActor<ASkillSystemTestPawn>();
-	ASkillSystemTestPawn *DestActor = World->SpawnActor<ASkillSystemTestPawn>();
+	AAbilitySystemTestPawn *SourceActor = World->SpawnActor<AAbilitySystemTestPawn>();
+	AAbilitySystemTestPawn *DestActor = World->SpawnActor<AAbilitySystemTestPawn>();
 
-	UProperty *StackingProperty = FindFieldChecked<UProperty>(USkillSystemTestAttributeSet::StaticClass(), GET_MEMBER_NAME_CHECKED(USkillSystemTestAttributeSet, StackingAttribute1));
+	UProperty *StackingProperty = FindFieldChecked<UProperty>(UAbilitySystemTestAttributeSet::StaticClass(), GET_MEMBER_NAME_CHECKED(UAbilitySystemTestAttributeSet, StackingAttribute1));
 
 	UAttributeComponent * SourceComponent = SourceActor->AttributeComponent;
 	UAttributeComponent * DestComponent = DestActor->AttributeComponent;
 
 	{
-		SKILL_LOG_SCOPE(TEXT("Apply StackingEffect"));
+		ABILITY_LOG_SCOPE(TEXT("Apply StackingEffect"));
 
 		UGameplayEffect * BaseStackedEffect = Cast<UGameplayEffect>(StaticConstructObject(UGameplayEffect::StaticClass(), GetTransientPackage(), FName(TEXT("StackingEffect1"))));
 		BaseStackedEffect->Modifiers.SetNum(1);
@@ -3556,7 +3556,7 @@ bool GameplayEffectsTest_StackingUnlimited(UWorld *World, FAutomationTestBase *
 	}
 
 	{
-		SKILL_LOG_SCOPE(TEXT("Apply StackingEffect"));
+		ABILITY_LOG_SCOPE(TEXT("Apply StackingEffect"));
 
 		UGameplayEffect * BaseStackedEffect = Cast<UGameplayEffect>(StaticConstructObject(UGameplayEffect::StaticClass(), GetTransientPackage(), FName(TEXT("StackingEffect2"))));
 		BaseStackedEffect->Modifiers.SetNum(1);
@@ -3581,7 +3581,7 @@ bool GameplayEffectsTest_StackingUnlimited(UWorld *World, FAutomationTestBase *
 		// Tick once
 		GameplayTest_TickWorld(World, 1.f);
 
-		float ActualValue = DestComponent->GetSet<USkillSystemTestAttributeSet>()->StackingAttribute1;
+		float ActualValue = DestComponent->GetSet<UAbilitySystemTestAttributeSet>()->StackingAttribute1;
 		float ExpectedValue = StackingValue * 6.f;	// 1 for first GE, 2 for the second GE, 3 for both GEs during the tick
 		Test->TestTrue(SKILL_TEST_TEXT("Stacking Attribute: Actual: %.2f == Exected: %.2f", ActualValue, ExpectedValue), ActualValue == ExpectedValue);
 	}
@@ -3603,16 +3603,16 @@ bool GameplayEffectsTest_StackingRemoval(UWorld *World, FAutomationTestBase * Te
 	SetGlobalCurveTable();
 	SetGlobalDataTable();
 
-	ASkillSystemTestPawn *SourceActor = World->SpawnActor<ASkillSystemTestPawn>();
-	ASkillSystemTestPawn *DestActor = World->SpawnActor<ASkillSystemTestPawn>();
+	AAbilitySystemTestPawn *SourceActor = World->SpawnActor<AAbilitySystemTestPawn>();
+	AAbilitySystemTestPawn *DestActor = World->SpawnActor<AAbilitySystemTestPawn>();
 
-	UProperty *StackingProperty = FindFieldChecked<UProperty>(USkillSystemTestAttributeSet::StaticClass(), GET_MEMBER_NAME_CHECKED(USkillSystemTestAttributeSet, StackingAttribute1));
+	UProperty *StackingProperty = FindFieldChecked<UProperty>(UAbilitySystemTestAttributeSet::StaticClass(), GET_MEMBER_NAME_CHECKED(UAbilitySystemTestAttributeSet, StackingAttribute1));
 
 	UAttributeComponent * SourceComponent = SourceActor->AttributeComponent;
 	UAttributeComponent * DestComponent = DestActor->AttributeComponent;
 
 	{
-		SKILL_LOG_SCOPE(TEXT("Apply StackingEffect"));
+		ABILITY_LOG_SCOPE(TEXT("Apply StackingEffect"));
 
 		UGameplayEffect * BaseStackedEffect = Cast<UGameplayEffect>(StaticConstructObject(UGameplayEffect::StaticClass(), GetTransientPackage(), FName(TEXT("StackingEffect1"))));
 		BaseStackedEffect->Modifiers.SetNum(1);
@@ -3631,7 +3631,7 @@ bool GameplayEffectsTest_StackingRemoval(UWorld *World, FAutomationTestBase * Te
 	}
 
 	{
-		SKILL_LOG_SCOPE(TEXT("Apply StackingEffect"));
+		ABILITY_LOG_SCOPE(TEXT("Apply StackingEffect"));
 
 		UGameplayEffect * BaseStackedEffect = Cast<UGameplayEffect>(StaticConstructObject(UGameplayEffect::StaticClass(), GetTransientPackage(), FName(TEXT("StackingEffect2"))));
 		BaseStackedEffect->Modifiers.SetNum(1);
@@ -3656,14 +3656,14 @@ bool GameplayEffectsTest_StackingRemoval(UWorld *World, FAutomationTestBase * Te
 		// Tick once
 		GameplayTest_TickWorld(World, 1.f);
 
-		float ActualValue = DestComponent->GetSet<USkillSystemTestAttributeSet>()->StackingAttribute1;
+		float ActualValue = DestComponent->GetSet<UAbilitySystemTestAttributeSet>()->StackingAttribute1;
 		float ExpectedValue = StackingValue * 4.f;	// 2 for second effect, 2 for second effect during tick
 		Test->TestTrue(SKILL_TEST_TEXT("Stacking Attribute: Actual: %.2f == Exected: %.2f", ActualValue, ExpectedValue), ActualValue == ExpectedValue);
 
 		// Tick once
 		GameplayTest_TickWorld(World, 1.f);
 
-		ActualValue = DestComponent->GetSet<USkillSystemTestAttributeSet>()->StackingAttribute1;
+		ActualValue = DestComponent->GetSet<UAbilitySystemTestAttributeSet>()->StackingAttribute1;
 		ExpectedValue = StackingValue * 5.f;	// 2 for second effect, 2 for second effect during first tick, 1 for first effect during second tick
 		Test->TestTrue(SKILL_TEST_TEXT("Stacking Attribute: Actual: %.2f == Exected: %.2f", ActualValue, ExpectedValue), ActualValue == ExpectedValue);
 	}
@@ -3685,16 +3685,16 @@ bool GameplayEffectsTest_StackingNoStack(UWorld *World, FAutomationTestBase * Te
 	SetGlobalCurveTable();
 	SetGlobalDataTable();
 
-	ASkillSystemTestPawn *SourceActor = World->SpawnActor<ASkillSystemTestPawn>();
-	ASkillSystemTestPawn *DestActor = World->SpawnActor<ASkillSystemTestPawn>();
+	AAbilitySystemTestPawn *SourceActor = World->SpawnActor<AAbilitySystemTestPawn>();
+	AAbilitySystemTestPawn *DestActor = World->SpawnActor<AAbilitySystemTestPawn>();
 
-	UProperty *NoStackProperty = FindFieldChecked<UProperty>(USkillSystemTestAttributeSet::StaticClass(), GET_MEMBER_NAME_CHECKED(USkillSystemTestAttributeSet, NoStackAttribute));
+	UProperty *NoStackProperty = FindFieldChecked<UProperty>(UAbilitySystemTestAttributeSet::StaticClass(), GET_MEMBER_NAME_CHECKED(UAbilitySystemTestAttributeSet, NoStackAttribute));
 
 	UAttributeComponent * SourceComponent = SourceActor->AttributeComponent;
 	UAttributeComponent * DestComponent = DestActor->AttributeComponent;
 
 	{
-		SKILL_LOG_SCOPE(TEXT("Apply StackingEffect"));
+		ABILITY_LOG_SCOPE(TEXT("Apply StackingEffect"));
 
 		UGameplayEffect * UnstackableEffect = Cast<UGameplayEffect>(StaticConstructObject(UGameplayEffect::StaticClass(), GetTransientPackage(), FName(TEXT("NoStackEffect1"))));
 		UnstackableEffect->Modifiers.SetNum(1);
@@ -3713,7 +3713,7 @@ bool GameplayEffectsTest_StackingNoStack(UWorld *World, FAutomationTestBase * Te
 	}
 
 	{
-		SKILL_LOG_SCOPE(TEXT("Apply StackingEffect"));
+		ABILITY_LOG_SCOPE(TEXT("Apply StackingEffect"));
 
 		UGameplayEffect * UnstackableEffect = Cast<UGameplayEffect>(StaticConstructObject(UGameplayEffect::StaticClass(), GetTransientPackage(), FName(TEXT("NoStackEffect2"))));
 		UnstackableEffect->Modifiers.SetNum(1);
@@ -3738,7 +3738,7 @@ bool GameplayEffectsTest_StackingNoStack(UWorld *World, FAutomationTestBase * Te
 		// Tick once
 		GameplayTest_TickWorld(World, 1.f);
 
-		float ActualValue = DestComponent->GetSet<USkillSystemTestAttributeSet>()->NoStackAttribute;
+		float ActualValue = DestComponent->GetSet<UAbilitySystemTestAttributeSet>()->NoStackAttribute;
 		float ExpectedValue = StackingValue * 6.f;	// 1 for first effect, 2 for second effect, 3 for both effects during tick
 		Test->TestTrue(SKILL_TEST_TEXT("Stacking Attribute: Actual: %.2f == Exected: %.2f", ActualValue, ExpectedValue), ActualValue == ExpectedValue);
 	}
@@ -3760,16 +3760,16 @@ bool GameplayEffectsTest_StackingCustomCapped(UWorld *World, FAutomationTestBase
 	SetGlobalCurveTable();
 	SetGlobalDataTable();
 
-	ASkillSystemTestPawn *SourceActor = World->SpawnActor<ASkillSystemTestPawn>();
-	ASkillSystemTestPawn *DestActor = World->SpawnActor<ASkillSystemTestPawn>();
+	AAbilitySystemTestPawn *SourceActor = World->SpawnActor<AAbilitySystemTestPawn>();
+	AAbilitySystemTestPawn *DestActor = World->SpawnActor<AAbilitySystemTestPawn>();
 
-	UProperty *StackingProperty = FindFieldChecked<UProperty>(USkillSystemTestAttributeSet::StaticClass(), GET_MEMBER_NAME_CHECKED(USkillSystemTestAttributeSet, StackingAttribute1));
+	UProperty *StackingProperty = FindFieldChecked<UProperty>(UAbilitySystemTestAttributeSet::StaticClass(), GET_MEMBER_NAME_CHECKED(UAbilitySystemTestAttributeSet, StackingAttribute1));
 
 	UAttributeComponent * SourceComponent = SourceActor->AttributeComponent;
 	UAttributeComponent * DestComponent = DestActor->AttributeComponent;
 
 	{
-		SKILL_LOG_SCOPE(TEXT("Apply StackingEffect"));
+		ABILITY_LOG_SCOPE(TEXT("Apply StackingEffect"));
 
 		UGameplayEffect * BaseStackedEffect = Cast<UGameplayEffect>(StaticConstructObject(UGameplayEffect::StaticClass(), GetTransientPackage(), FName(TEXT("StackingEffect1"))));
 		BaseStackedEffect->Modifiers.SetNum(1);
@@ -3790,7 +3790,7 @@ bool GameplayEffectsTest_StackingCustomCapped(UWorld *World, FAutomationTestBase
 	}
 
 	{
-		SKILL_LOG_SCOPE(TEXT("Apply StackingEffect"));
+		ABILITY_LOG_SCOPE(TEXT("Apply StackingEffect"));
 
 		UGameplayEffect * BaseStackedEffect = Cast<UGameplayEffect>(StaticConstructObject(UGameplayEffect::StaticClass(), GetTransientPackage(), FName(TEXT("StackingEffect2"))));
 		BaseStackedEffect->Modifiers.SetNum(1);
@@ -3811,7 +3811,7 @@ bool GameplayEffectsTest_StackingCustomCapped(UWorld *World, FAutomationTestBase
 	}
 
 	{
-		SKILL_LOG_SCOPE(TEXT("Apply StackingEffect"));
+		ABILITY_LOG_SCOPE(TEXT("Apply StackingEffect"));
 
 		UGameplayEffect * BaseStackedEffect = Cast<UGameplayEffect>(StaticConstructObject(UGameplayEffect::StaticClass(), GetTransientPackage(), FName(TEXT("StackingEffect3"))));
 		BaseStackedEffect->Modifiers.SetNum(1);
@@ -3838,13 +3838,13 @@ bool GameplayEffectsTest_StackingCustomCapped(UWorld *World, FAutomationTestBase
 		// Tick once
 		GameplayTest_TickWorld(World, 1.f);
 
-		float ActualValue = DestComponent->GetSet<USkillSystemTestAttributeSet>()->StackingAttribute1;
+		float ActualValue = DestComponent->GetSet<UAbilitySystemTestAttributeSet>()->StackingAttribute1;
 		float ExpectedValue = StackingValue * 4.f;	// 2 for effects being applied at the start of the tick, 2 for effects being applied at the end of this tick
 		Test->TestTrue(SKILL_TEST_TEXT("Stacking Attribute: Actual: %.2f == Exected: %.2f", ActualValue, ExpectedValue), ActualValue == ExpectedValue);
 	}
 
 	{
-		SKILL_LOG_SCOPE(TEXT("Apply StackingEffect"));
+		ABILITY_LOG_SCOPE(TEXT("Apply StackingEffect"));
 
 		UGameplayEffect * BaseStackedEffect = Cast<UGameplayEffect>(StaticConstructObject(UGameplayEffect::StaticClass(), GetTransientPackage(), FName(TEXT("StackingEffect4"))));
 		BaseStackedEffect->Modifiers.SetNum(1);
@@ -3871,7 +3871,7 @@ bool GameplayEffectsTest_StackingCustomCapped(UWorld *World, FAutomationTestBase
 		// Tick once
 		GameplayTest_TickWorld(World, 1.f);
 
-		float ActualValue = DestComponent->GetSet<USkillSystemTestAttributeSet>()->StackingAttribute1;
+		float ActualValue = DestComponent->GetSet<UAbilitySystemTestAttributeSet>()->StackingAttribute1;
 		float ExpectedValue = StackingValue * 6.f;	// 2 for effects being applied, 2 for tick, 2 for tick
 		Test->TestTrue(SKILL_TEST_TEXT("Stacking Attribute: Actual: %.2f == Exected: %.2f", ActualValue, ExpectedValue), ActualValue == ExpectedValue);
 	}
@@ -3880,7 +3880,7 @@ bool GameplayEffectsTest_StackingCustomCapped(UWorld *World, FAutomationTestBase
 		// Tick once
 		GameplayTest_TickWorld(World, 1.f);
 
-		float ActualValue = DestComponent->GetSet<USkillSystemTestAttributeSet>()->StackingAttribute1;
+		float ActualValue = DestComponent->GetSet<UAbilitySystemTestAttributeSet>()->StackingAttribute1;
 		float ExpectedValue = StackingValue * 8.f;	// the last GE should have refreshed the timer so we should have 2 for GEs applied and 3 * 2 for ticks
 		Test->TestTrue(SKILL_TEST_TEXT("Stacking Attribute: Actual: %.2f == Exected: %.2f", ActualValue, ExpectedValue), ActualValue == ExpectedValue);
 	}
@@ -3889,7 +3889,7 @@ bool GameplayEffectsTest_StackingCustomCapped(UWorld *World, FAutomationTestBase
 		// Tick once
 		GameplayTest_TickWorld(World, 1.f);
 
-		float ActualValue = DestComponent->GetSet<USkillSystemTestAttributeSet>()->StackingAttribute1;
+		float ActualValue = DestComponent->GetSet<UAbilitySystemTestAttributeSet>()->StackingAttribute1;
 		float ExpectedValue = StackingValue * 8.f;	// the effects should have timed out so we should have 2 for GEs applied and 3 * 2 for ticks
 		Test->TestTrue(SKILL_TEST_TEXT("Stacking Attribute: Actual: %.2f == Exected: %.2f", ActualValue, ExpectedValue), ActualValue == ExpectedValue);
 	}
@@ -3911,16 +3911,16 @@ bool GameplayEffectsTest_StackingCustomDiminishingReturns(UWorld *World, FAutoma
 	SetGlobalCurveTable();
 	SetGlobalDataTable();
 
-	ASkillSystemTestPawn *SourceActor = World->SpawnActor<ASkillSystemTestPawn>();
-	ASkillSystemTestPawn *DestActor = World->SpawnActor<ASkillSystemTestPawn>();
+	AAbilitySystemTestPawn *SourceActor = World->SpawnActor<AAbilitySystemTestPawn>();
+	AAbilitySystemTestPawn *DestActor = World->SpawnActor<AAbilitySystemTestPawn>();
 
-	UProperty *StackingProperty = FindFieldChecked<UProperty>(USkillSystemTestAttributeSet::StaticClass(), GET_MEMBER_NAME_CHECKED(USkillSystemTestAttributeSet, StackingAttribute1));
+	UProperty *StackingProperty = FindFieldChecked<UProperty>(UAbilitySystemTestAttributeSet::StaticClass(), GET_MEMBER_NAME_CHECKED(UAbilitySystemTestAttributeSet, StackingAttribute1));
 
 	UAttributeComponent * SourceComponent = SourceActor->AttributeComponent;
 	UAttributeComponent * DestComponent = DestActor->AttributeComponent;
 
 	{
-		SKILL_LOG_SCOPE(TEXT("Apply StackingEffect"));
+		ABILITY_LOG_SCOPE(TEXT("Apply StackingEffect"));
 
 		UGameplayEffect * BaseStackedEffect = Cast<UGameplayEffect>(StaticConstructObject(UGameplayEffect::StaticClass(), GetTransientPackage(), FName(TEXT("StackingEffect1"))));
 		BaseStackedEffect->Modifiers.SetNum(1);
@@ -3947,13 +3947,13 @@ bool GameplayEffectsTest_StackingCustomDiminishingReturns(UWorld *World, FAutoma
 		// Tick once
 		GameplayTest_TickWorld(World, 1.f);
 
-		float ActualValue = DestComponent->GetSet<USkillSystemTestAttributeSet>()->StackingAttribute1;
+		float ActualValue = DestComponent->GetSet<UAbilitySystemTestAttributeSet>()->StackingAttribute1;
 		float ExpectedValue = StackingValue * (5.f * 2);	// first application gets five times the result and will be applied at the start and end of the tick
 		Test->TestTrue(SKILL_TEST_TEXT("Stacking Attribute: Actual: %.2f == Exected: %.2f", ActualValue, ExpectedValue), ActualValue == ExpectedValue);
 	}
 
 	{
-		SKILL_LOG_SCOPE(TEXT("Apply StackingEffect"));
+		ABILITY_LOG_SCOPE(TEXT("Apply StackingEffect"));
 
 		UGameplayEffect * BaseStackedEffect = Cast<UGameplayEffect>(StaticConstructObject(UGameplayEffect::StaticClass(), GetTransientPackage(), FName(TEXT("StackingEffect2"))));
 		BaseStackedEffect->Modifiers.SetNum(1);
@@ -3980,13 +3980,13 @@ bool GameplayEffectsTest_StackingCustomDiminishingReturns(UWorld *World, FAutoma
 		// Tick once
 		GameplayTest_TickWorld(World, 1.f);
 
-		float ActualValue = DestComponent->GetSet<USkillSystemTestAttributeSet>()->StackingAttribute1;
+		float ActualValue = DestComponent->GetSet<UAbilitySystemTestAttributeSet>()->StackingAttribute1;
 		float ExpectedValue = StackingValue * ((5.f * 2) + 7.f);	// second application gets seven times the result
 		Test->TestTrue(SKILL_TEST_TEXT("Stacking Attribute: Actual: %.2f == Exected: %.2f", ActualValue, ExpectedValue), ActualValue == ExpectedValue);
 	}
 
 	{
-		SKILL_LOG_SCOPE(TEXT("Apply StackingEffect"));
+		ABILITY_LOG_SCOPE(TEXT("Apply StackingEffect"));
 
 		UGameplayEffect * BaseStackedEffect = Cast<UGameplayEffect>(StaticConstructObject(UGameplayEffect::StaticClass(), GetTransientPackage(), FName(TEXT("StackingEffect3"))));
 		BaseStackedEffect->Modifiers.SetNum(1);
@@ -4013,13 +4013,13 @@ bool GameplayEffectsTest_StackingCustomDiminishingReturns(UWorld *World, FAutoma
 		// Tick once
 		GameplayTest_TickWorld(World, 1.f);
 
-		float ActualValue = DestComponent->GetSet<USkillSystemTestAttributeSet>()->StackingAttribute1;
+		float ActualValue = DestComponent->GetSet<UAbilitySystemTestAttributeSet>()->StackingAttribute1;
 		float ExpectedValue = StackingValue * ((5.f * 2) + 7.f + 8.f);	// third application gets eight times the result
 		Test->TestTrue(SKILL_TEST_TEXT("Stacking Attribute: Actual: %.2f == Exected: %.2f", ActualValue, ExpectedValue), ActualValue == ExpectedValue);
 	}
 
 	{
-		SKILL_LOG_SCOPE(TEXT("Apply StackingEffect"));
+		ABILITY_LOG_SCOPE(TEXT("Apply StackingEffect"));
 
 		UGameplayEffect * BaseStackedEffect = Cast<UGameplayEffect>(StaticConstructObject(UGameplayEffect::StaticClass(), GetTransientPackage(), FName(TEXT("StackingEffect4"))));
 		BaseStackedEffect->Modifiers.SetNum(1);
@@ -4046,7 +4046,7 @@ bool GameplayEffectsTest_StackingCustomDiminishingReturns(UWorld *World, FAutoma
 		// Tick once
 		GameplayTest_TickWorld(World, 1.f);
 
-		float ActualValue = DestComponent->GetSet<USkillSystemTestAttributeSet>()->StackingAttribute1;
+		float ActualValue = DestComponent->GetSet<UAbilitySystemTestAttributeSet>()->StackingAttribute1;
 		float ExpectedValue = StackingValue * ((5.f * 2) + 7.f + 8.f + 9.f);	// fourth application gets nine times the result
 		Test->TestTrue(SKILL_TEST_TEXT("Stacking Attribute: Actual: %.2f == Exected: %.2f", ActualValue, ExpectedValue), ActualValue == ExpectedValue);
 	}
@@ -4068,16 +4068,16 @@ bool GameplayEffectsTest_StackingDifferentRules(UWorld *World, FAutomationTestBa
 	SetGlobalCurveTable();
 	SetGlobalDataTable();
 
-	ASkillSystemTestPawn *SourceActor = World->SpawnActor<ASkillSystemTestPawn>();
-	ASkillSystemTestPawn *DestActor = World->SpawnActor<ASkillSystemTestPawn>();
+	AAbilitySystemTestPawn *SourceActor = World->SpawnActor<AAbilitySystemTestPawn>();
+	AAbilitySystemTestPawn *DestActor = World->SpawnActor<AAbilitySystemTestPawn>();
 
-	UProperty *StackingProperty = FindFieldChecked<UProperty>(USkillSystemTestAttributeSet::StaticClass(), GET_MEMBER_NAME_CHECKED(USkillSystemTestAttributeSet, StackingAttribute1));
+	UProperty *StackingProperty = FindFieldChecked<UProperty>(UAbilitySystemTestAttributeSet::StaticClass(), GET_MEMBER_NAME_CHECKED(UAbilitySystemTestAttributeSet, StackingAttribute1));
 
 	UAttributeComponent * SourceComponent = SourceActor->AttributeComponent;
 	UAttributeComponent * DestComponent = DestActor->AttributeComponent;
 
 	{
-		SKILL_LOG_SCOPE(TEXT("Apply StackingEffect"));
+		ABILITY_LOG_SCOPE(TEXT("Apply StackingEffect"));
 
 		UGameplayEffect * BaseStackedEffect = Cast<UGameplayEffect>(StaticConstructObject(UGameplayEffect::StaticClass(), GetTransientPackage(), FName(TEXT("StackingEffect1"))));
 		BaseStackedEffect->Modifiers.SetNum(1);
@@ -4096,7 +4096,7 @@ bool GameplayEffectsTest_StackingDifferentRules(UWorld *World, FAutomationTestBa
 	}
 
 	{
-		SKILL_LOG_SCOPE(TEXT("Apply StackingEffect"));
+		ABILITY_LOG_SCOPE(TEXT("Apply StackingEffect"));
 
 		UGameplayEffect * BaseStackedEffect = Cast<UGameplayEffect>(StaticConstructObject(UGameplayEffect::StaticClass(), GetTransientPackage(), FName(TEXT("StackingEffect2"))));
 		BaseStackedEffect->Modifiers.SetNum(1);
@@ -4121,7 +4121,7 @@ bool GameplayEffectsTest_StackingDifferentRules(UWorld *World, FAutomationTestBa
 		// Tick once
 		GameplayTest_TickWorld(World, 1.f);
 
-		float ActualValue = DestComponent->GetSet<USkillSystemTestAttributeSet>()->StackingAttribute1;
+		float ActualValue = DestComponent->GetSet<UAbilitySystemTestAttributeSet>()->StackingAttribute1;
 		float ExpectedValue = StackingValue * 6.f;	// 1 for first effect, 2 for the second effect, 3 during tick
 		Test->TestTrue(SKILL_TEST_TEXT("Stacking Attribute: Actual: %.2f == Exected: %.2f", ActualValue, ExpectedValue), ActualValue == ExpectedValue);
 	}
@@ -4143,17 +4143,17 @@ bool GameplayEffectsTest_StackingDifferentAttributes(UWorld *World, FAutomationT
 	SetGlobalCurveTable();
 	SetGlobalDataTable();
 
-	ASkillSystemTestPawn *SourceActor = World->SpawnActor<ASkillSystemTestPawn>();
-	ASkillSystemTestPawn *DestActor = World->SpawnActor<ASkillSystemTestPawn>();
+	AAbilitySystemTestPawn *SourceActor = World->SpawnActor<AAbilitySystemTestPawn>();
+	AAbilitySystemTestPawn *DestActor = World->SpawnActor<AAbilitySystemTestPawn>();
 
-	UProperty *StackingProperty1 = FindFieldChecked<UProperty>(USkillSystemTestAttributeSet::StaticClass(), GET_MEMBER_NAME_CHECKED(USkillSystemTestAttributeSet, StackingAttribute1));
-	UProperty *StackingProperty2 = FindFieldChecked<UProperty>(USkillSystemTestAttributeSet::StaticClass(), GET_MEMBER_NAME_CHECKED(USkillSystemTestAttributeSet, StackingAttribute2));
+	UProperty *StackingProperty1 = FindFieldChecked<UProperty>(UAbilitySystemTestAttributeSet::StaticClass(), GET_MEMBER_NAME_CHECKED(UAbilitySystemTestAttributeSet, StackingAttribute1));
+	UProperty *StackingProperty2 = FindFieldChecked<UProperty>(UAbilitySystemTestAttributeSet::StaticClass(), GET_MEMBER_NAME_CHECKED(UAbilitySystemTestAttributeSet, StackingAttribute2));
 
 	UAttributeComponent * SourceComponent = SourceActor->AttributeComponent;
 	UAttributeComponent * DestComponent = DestActor->AttributeComponent;
 
 	{
-		SKILL_LOG_SCOPE(TEXT("Apply StackingEffect"));
+		ABILITY_LOG_SCOPE(TEXT("Apply StackingEffect"));
 
 		UGameplayEffect * BaseStackedEffect = Cast<UGameplayEffect>(StaticConstructObject(UGameplayEffect::StaticClass(), GetTransientPackage(), FName(TEXT("StackingEffect1"))));
 		BaseStackedEffect->Modifiers.SetNum(1);
@@ -4172,7 +4172,7 @@ bool GameplayEffectsTest_StackingDifferentAttributes(UWorld *World, FAutomationT
 	}
 
 	{
-		SKILL_LOG_SCOPE(TEXT("Apply StackingEffect"));
+		ABILITY_LOG_SCOPE(TEXT("Apply StackingEffect"));
 
 		UGameplayEffect * BaseStackedEffect = Cast<UGameplayEffect>(StaticConstructObject(UGameplayEffect::StaticClass(), GetTransientPackage(), FName(TEXT("StackingEffect2"))));
 		BaseStackedEffect->Modifiers.SetNum(1);
@@ -4197,10 +4197,10 @@ bool GameplayEffectsTest_StackingDifferentAttributes(UWorld *World, FAutomationT
 		// Tick once
 		GameplayTest_TickWorld(World, 1.f);
 
-		float ActualValue = DestComponent->GetSet<USkillSystemTestAttributeSet>()->StackingAttribute1;
+		float ActualValue = DestComponent->GetSet<UAbilitySystemTestAttributeSet>()->StackingAttribute1;
 		float ExpectedValue = StackingValue * 2.f;	// 1 for first effect, 1 for first effect during tick
 		Test->TestTrue(SKILL_TEST_TEXT("Stacking Attribute 1: %.2f == Exected: %.2f", ActualValue, ExpectedValue), ActualValue == ExpectedValue);
-		ActualValue = DestComponent->GetSet<USkillSystemTestAttributeSet>()->StackingAttribute2;
+		ActualValue = DestComponent->GetSet<UAbilitySystemTestAttributeSet>()->StackingAttribute2;
 		ExpectedValue = StackingValue * 4.f; // 2 for the second effect, 2 for the second effect during tick
 		Test->TestTrue(SKILL_TEST_TEXT("Stacking Attribute 2: Actual: %.2f == Exected: %.2f", ActualValue, ExpectedValue), ActualValue == ExpectedValue);
 	}
@@ -4222,16 +4222,16 @@ bool GameplayEffectsTest_StackingCustomTwoRules(UWorld *World, FAutomationTestBa
 	SetGlobalCurveTable();
 	SetGlobalDataTable();
 
-	ASkillSystemTestPawn *SourceActor = World->SpawnActor<ASkillSystemTestPawn>();
-	ASkillSystemTestPawn *DestActor = World->SpawnActor<ASkillSystemTestPawn>();
+	AAbilitySystemTestPawn *SourceActor = World->SpawnActor<AAbilitySystemTestPawn>();
+	AAbilitySystemTestPawn *DestActor = World->SpawnActor<AAbilitySystemTestPawn>();
 
-	UProperty *StackingProperty = FindFieldChecked<UProperty>(USkillSystemTestAttributeSet::StaticClass(), GET_MEMBER_NAME_CHECKED(USkillSystemTestAttributeSet, StackingAttribute1));
+	UProperty *StackingProperty = FindFieldChecked<UProperty>(UAbilitySystemTestAttributeSet::StaticClass(), GET_MEMBER_NAME_CHECKED(UAbilitySystemTestAttributeSet, StackingAttribute1));
 
 	UAttributeComponent * SourceComponent = SourceActor->AttributeComponent;
 	UAttributeComponent * DestComponent = DestActor->AttributeComponent;
 
 	{
-		SKILL_LOG_SCOPE(TEXT("Apply StackingEffect"));
+		ABILITY_LOG_SCOPE(TEXT("Apply StackingEffect"));
 
 		UGameplayEffect * BaseStackedEffect = Cast<UGameplayEffect>(StaticConstructObject(UGameplayEffect::StaticClass(), GetTransientPackage(), FName(TEXT("StackingEffect1"))));
 		BaseStackedEffect->Modifiers.SetNum(1);
@@ -4252,7 +4252,7 @@ bool GameplayEffectsTest_StackingCustomTwoRules(UWorld *World, FAutomationTestBa
 	}
 
 	{
-		SKILL_LOG_SCOPE(TEXT("Apply StackingEffect"));
+		ABILITY_LOG_SCOPE(TEXT("Apply StackingEffect"));
 
 		UGameplayEffect * BaseStackedEffect = Cast<UGameplayEffect>(StaticConstructObject(UGameplayEffect::StaticClass(), GetTransientPackage(), FName(TEXT("StackingEffect2"))));
 		BaseStackedEffect->Modifiers.SetNum(1);
@@ -4273,7 +4273,7 @@ bool GameplayEffectsTest_StackingCustomTwoRules(UWorld *World, FAutomationTestBa
 	}
 
 	{
-		SKILL_LOG_SCOPE(TEXT("Apply StackingEffect"));
+		ABILITY_LOG_SCOPE(TEXT("Apply StackingEffect"));
 
 		UGameplayEffect * BaseStackedEffect = Cast<UGameplayEffect>(StaticConstructObject(UGameplayEffect::StaticClass(), GetTransientPackage(), FName(TEXT("StackingEffect3"))));
 		BaseStackedEffect->Modifiers.SetNum(1);
@@ -4300,7 +4300,7 @@ bool GameplayEffectsTest_StackingCustomTwoRules(UWorld *World, FAutomationTestBa
 		// Tick once
 		GameplayTest_TickWorld(World, 1.f);
 
-		float ActualValue = DestComponent->GetSet<USkillSystemTestAttributeSet>()->StackingAttribute1;
+		float ActualValue = DestComponent->GetSet<UAbilitySystemTestAttributeSet>()->StackingAttribute1;
 		float ExpectedValue = StackingValue * ((1 + 7) * 2);	// the capped stacking rule will apply one, the diminishing returns rule will apply seven, both rules are applied twice
 		Test->TestTrue(SKILL_TEST_TEXT("Stacking Attribute: Actual: %.2f == Exected: %.2f", ActualValue, ExpectedValue), ActualValue == ExpectedValue);
 	}
@@ -4322,17 +4322,17 @@ bool GameplayEffectsTest_StackingCustomTwoAttributes(UWorld *World, FAutomationT
 	SetGlobalCurveTable();
 	SetGlobalDataTable();
 
-	ASkillSystemTestPawn *SourceActor = World->SpawnActor<ASkillSystemTestPawn>();
-	ASkillSystemTestPawn *DestActor = World->SpawnActor<ASkillSystemTestPawn>();
+	AAbilitySystemTestPawn *SourceActor = World->SpawnActor<AAbilitySystemTestPawn>();
+	AAbilitySystemTestPawn *DestActor = World->SpawnActor<AAbilitySystemTestPawn>();
 
-	UProperty *StackingProperty1 = FindFieldChecked<UProperty>(USkillSystemTestAttributeSet::StaticClass(), GET_MEMBER_NAME_CHECKED(USkillSystemTestAttributeSet, StackingAttribute1));
-	UProperty *StackingProperty2 = FindFieldChecked<UProperty>(USkillSystemTestAttributeSet::StaticClass(), GET_MEMBER_NAME_CHECKED(USkillSystemTestAttributeSet, StackingAttribute2));
+	UProperty *StackingProperty1 = FindFieldChecked<UProperty>(UAbilitySystemTestAttributeSet::StaticClass(), GET_MEMBER_NAME_CHECKED(UAbilitySystemTestAttributeSet, StackingAttribute1));
+	UProperty *StackingProperty2 = FindFieldChecked<UProperty>(UAbilitySystemTestAttributeSet::StaticClass(), GET_MEMBER_NAME_CHECKED(UAbilitySystemTestAttributeSet, StackingAttribute2));
 
 	UAttributeComponent * SourceComponent = SourceActor->AttributeComponent;
 	UAttributeComponent * DestComponent = DestActor->AttributeComponent;
 
 	{
-		SKILL_LOG_SCOPE(TEXT("Apply StackingEffect"));
+		ABILITY_LOG_SCOPE(TEXT("Apply StackingEffect"));
 
 		UGameplayEffect * BaseStackedEffect = Cast<UGameplayEffect>(StaticConstructObject(UGameplayEffect::StaticClass(), GetTransientPackage(), FName(TEXT("StackingEffect1"))));
 		BaseStackedEffect->Modifiers.SetNum(1);
@@ -4353,7 +4353,7 @@ bool GameplayEffectsTest_StackingCustomTwoAttributes(UWorld *World, FAutomationT
 	}
 
 	{
-		SKILL_LOG_SCOPE(TEXT("Apply StackingEffect"));
+		ABILITY_LOG_SCOPE(TEXT("Apply StackingEffect"));
 
 		UGameplayEffect * BaseStackedEffect = Cast<UGameplayEffect>(StaticConstructObject(UGameplayEffect::StaticClass(), GetTransientPackage(), FName(TEXT("StackingEffect2"))));
 		BaseStackedEffect->Modifiers.SetNum(1);
@@ -4374,7 +4374,7 @@ bool GameplayEffectsTest_StackingCustomTwoAttributes(UWorld *World, FAutomationT
 	}
 
 	{
-		SKILL_LOG_SCOPE(TEXT("Apply StackingEffect"));
+		ABILITY_LOG_SCOPE(TEXT("Apply StackingEffect"));
 
 		UGameplayEffect * BaseStackedEffect = Cast<UGameplayEffect>(StaticConstructObject(UGameplayEffect::StaticClass(), GetTransientPackage(), FName(TEXT("StackingEffect3"))));
 		BaseStackedEffect->Modifiers.SetNum(1);
@@ -4401,11 +4401,11 @@ bool GameplayEffectsTest_StackingCustomTwoAttributes(UWorld *World, FAutomationT
 		// Tick once
 		GameplayTest_TickWorld(World, 1.f);
 
-		float ActualValue = DestComponent->GetSet<USkillSystemTestAttributeSet>()->StackingAttribute1;
+		float ActualValue = DestComponent->GetSet<UAbilitySystemTestAttributeSet>()->StackingAttribute1;
 		float ExpectedValue = StackingValue * (2 * 2);	// two effects should be applied to the first attribute, both apply twice
 		Test->TestTrue(SKILL_TEST_TEXT("Stacking Attribute: Actual: %.2f == Exected: %.2f", ActualValue, ExpectedValue), ActualValue == ExpectedValue);
 
-		ActualValue = DestComponent->GetSet<USkillSystemTestAttributeSet>()->StackingAttribute2;
+		ActualValue = DestComponent->GetSet<UAbilitySystemTestAttributeSet>()->StackingAttribute2;
 		ExpectedValue = StackingValue * (1 * 2);	// one effect should be applied twice to the second attribute
 		Test->TestTrue(SKILL_TEST_TEXT("Stacking Attribute: Actual: %.2f == Exected: %.2f", ActualValue, ExpectedValue), ActualValue == ExpectedValue);
 	}
@@ -4427,17 +4427,17 @@ bool GameplayEffectsTest_StackingRemovingModifiers(UWorld *World, FAutomationTes
 	SetGlobalCurveTable();
 	SetGlobalDataTable();
 
-	ASkillSystemTestPawn *SourceActor = World->SpawnActor<ASkillSystemTestPawn>();
-	ASkillSystemTestPawn *DestActor = World->SpawnActor<ASkillSystemTestPawn>();
+	AAbilitySystemTestPawn *SourceActor = World->SpawnActor<AAbilitySystemTestPawn>();
+	AAbilitySystemTestPawn *DestActor = World->SpawnActor<AAbilitySystemTestPawn>();
 
-	UProperty *StackingProperty = FindFieldChecked<UProperty>(USkillSystemTestAttributeSet::StaticClass(), GET_MEMBER_NAME_CHECKED(USkillSystemTestAttributeSet, StackingAttribute1));
+	UProperty *StackingProperty = FindFieldChecked<UProperty>(UAbilitySystemTestAttributeSet::StaticClass(), GET_MEMBER_NAME_CHECKED(UAbilitySystemTestAttributeSet, StackingAttribute1));
 
 	UAttributeComponent * SourceComponent = SourceActor->AttributeComponent;
 	UAttributeComponent * DestComponent = DestActor->AttributeComponent;
 
 	// Setup a GE to modify IncomingGEs
 	{
-		SKILL_LOG_SCOPE(TEXT("Apply modifier to incoming, tagged GEs"));
+		ABILITY_LOG_SCOPE(TEXT("Apply modifier to incoming, tagged GEs"));
 
 		UGameplayEffect* BaseModEffect = Cast<UGameplayEffect>(StaticConstructObject(UGameplayEffect::StaticClass(), GetTransientPackage(), FName(TEXT("ProtectBuff"))));
 		BaseModEffect->Modifiers.SetNum(1);
@@ -4454,7 +4454,7 @@ bool GameplayEffectsTest_StackingRemovingModifiers(UWorld *World, FAutomationTes
 	}
 
 	{
-		SKILL_LOG_SCOPE(TEXT("Apply StackingEffect"));
+		ABILITY_LOG_SCOPE(TEXT("Apply StackingEffect"));
 
 		UGameplayEffect * BaseStackedEffect = Cast<UGameplayEffect>(StaticConstructObject(UGameplayEffect::StaticClass(), GetTransientPackage(), FName(TEXT("StackingEffect1"))));
 		BaseStackedEffect->Modifiers.SetNum(1);
@@ -4474,7 +4474,7 @@ bool GameplayEffectsTest_StackingRemovingModifiers(UWorld *World, FAutomationTes
 	}
 
 	{
-		SKILL_LOG_SCOPE(TEXT("Apply StackingEffect"));
+		ABILITY_LOG_SCOPE(TEXT("Apply StackingEffect"));
 
 		UGameplayEffect * BaseStackedEffect = Cast<UGameplayEffect>(StaticConstructObject(UGameplayEffect::StaticClass(), GetTransientPackage(), FName(TEXT("StackingEffect2"))));
 		BaseStackedEffect->Modifiers.SetNum(1);
@@ -4499,7 +4499,7 @@ bool GameplayEffectsTest_StackingRemovingModifiers(UWorld *World, FAutomationTes
 		// Tick once
 		GameplayTest_TickWorld(World, 1.f);
 
-		float ActualValue = DestComponent->GetSet<USkillSystemTestAttributeSet>()->StackingAttribute1;
+		float ActualValue = DestComponent->GetSet<UAbilitySystemTestAttributeSet>()->StackingAttribute1;
 		float ExpectedValue = StackingValue * 8.f;	// 4 for the first GE, 4 for first GE during tick
 		Test->TestTrue(SKILL_TEST_TEXT("Stacking Attribute: Actual: %.2f == Exected: %.2f", ActualValue, ExpectedValue), ActualValue == ExpectedValue);
 
@@ -4507,7 +4507,7 @@ bool GameplayEffectsTest_StackingRemovingModifiers(UWorld *World, FAutomationTes
 		// Tick once
 		GameplayTest_TickWorld(World, 1.f);
 
-		ActualValue = DestComponent->GetSet<USkillSystemTestAttributeSet>()->StackingAttribute1;
+		ActualValue = DestComponent->GetSet<UAbilitySystemTestAttributeSet>()->StackingAttribute1;
 		ExpectedValue = StackingValue * 10.f;	// 4 for the first GE, 4 for first GE during first tick, 2 for the second GE during second tick
 		Test->TestTrue(SKILL_TEST_TEXT("Stacking Attribute: Actual: %.2f == Exected: %.2f", ActualValue, ExpectedValue), ActualValue == ExpectedValue);
 	}
@@ -4529,16 +4529,16 @@ bool GameplayEffectsTest_StackingAddingModifiers(UWorld *World, FAutomationTestB
 	SetGlobalCurveTable();
 	SetGlobalDataTable();
 
-	ASkillSystemTestPawn *SourceActor = World->SpawnActor<ASkillSystemTestPawn>();
-	ASkillSystemTestPawn *DestActor = World->SpawnActor<ASkillSystemTestPawn>();
+	AAbilitySystemTestPawn *SourceActor = World->SpawnActor<AAbilitySystemTestPawn>();
+	AAbilitySystemTestPawn *DestActor = World->SpawnActor<AAbilitySystemTestPawn>();
 
-	UProperty *StackingProperty = FindFieldChecked<UProperty>(USkillSystemTestAttributeSet::StaticClass(), GET_MEMBER_NAME_CHECKED(USkillSystemTestAttributeSet, StackingAttribute1));
+	UProperty *StackingProperty = FindFieldChecked<UProperty>(UAbilitySystemTestAttributeSet::StaticClass(), GET_MEMBER_NAME_CHECKED(UAbilitySystemTestAttributeSet, StackingAttribute1));
 
 	UAttributeComponent * SourceComponent = SourceActor->AttributeComponent;
 	UAttributeComponent * DestComponent = DestActor->AttributeComponent;
 
 	{
-		SKILL_LOG_SCOPE(TEXT("Apply StackingEffect"));
+		ABILITY_LOG_SCOPE(TEXT("Apply StackingEffect"));
 
 		UGameplayEffect * BaseStackedEffect = Cast<UGameplayEffect>(StaticConstructObject(UGameplayEffect::StaticClass(), GetTransientPackage(), FName(TEXT("StackingEffect1"))));
 		BaseStackedEffect->Modifiers.SetNum(1);
@@ -4558,7 +4558,7 @@ bool GameplayEffectsTest_StackingAddingModifiers(UWorld *World, FAutomationTestB
 	}
 
 	{
-		SKILL_LOG_SCOPE(TEXT("Apply StackingEffect"));
+		ABILITY_LOG_SCOPE(TEXT("Apply StackingEffect"));
 
 		UGameplayEffect * BaseStackedEffect = Cast<UGameplayEffect>(StaticConstructObject(UGameplayEffect::StaticClass(), GetTransientPackage(), FName(TEXT("StackingEffect2"))));
 		BaseStackedEffect->Modifiers.SetNum(1);
@@ -4583,14 +4583,14 @@ bool GameplayEffectsTest_StackingAddingModifiers(UWorld *World, FAutomationTestB
 		// Tick once
 		GameplayTest_TickWorld(World, 1.f);
 
-		float ActualValue = DestComponent->GetSet<USkillSystemTestAttributeSet>()->StackingAttribute1;
+		float ActualValue = DestComponent->GetSet<UAbilitySystemTestAttributeSet>()->StackingAttribute1;
 		float ExpectedValue = StackingValue * 4.f;	// 4 for the second GE executing twice
 		Test->TestTrue(SKILL_TEST_TEXT("Stacking Attribute: Actual: %.2f == Exected: %.2f", ActualValue, ExpectedValue), ActualValue == ExpectedValue);
 	}
 
 	// Setup a GE to modify IncomingGEs
 	{
-		SKILL_LOG_SCOPE(TEXT("Apply modifier to incoming, tagged GEs"));
+		ABILITY_LOG_SCOPE(TEXT("Apply modifier to incoming, tagged GEs"));
 
 		UGameplayEffect* BaseModEffect = Cast<UGameplayEffect>(StaticConstructObject(UGameplayEffect::StaticClass(), GetTransientPackage(), FName(TEXT("ProtectBuff"))));
 		BaseModEffect->Modifiers.SetNum(1);
@@ -4614,7 +4614,7 @@ bool GameplayEffectsTest_StackingAddingModifiers(UWorld *World, FAutomationTestB
 		// Tick once
 		GameplayTest_TickWorld(World, 1.f);
 
-		float ActualValue = DestComponent->GetSet<USkillSystemTestAttributeSet>()->StackingAttribute1;
+		float ActualValue = DestComponent->GetSet<UAbilitySystemTestAttributeSet>()->StackingAttribute1;
 		float ExpectedValue = StackingValue * 7.f;	// 2 for the second GE, 2 for second GE during tick, 3 for the modified first GE during second tick
 		Test->TestTrue(SKILL_TEST_TEXT("Stacking Attribute: Actual: %.2f == Exected: %.2f", ActualValue, ExpectedValue), ActualValue == ExpectedValue);
 	}
@@ -4634,19 +4634,19 @@ bool GameplayEffectsTest_ImmunityIncoming(UWorld *World, FAutomationTestBase * T
 	const float StartHealth = 100.f;
 	const float DamageValue = 5.f;
 
-	ASkillSystemTestPawn *SourceActor = World->SpawnActor<ASkillSystemTestPawn>();
-	ASkillSystemTestPawn *DestActor = World->SpawnActor<ASkillSystemTestPawn>();
+	AAbilitySystemTestPawn *SourceActor = World->SpawnActor<AAbilitySystemTestPawn>();
+	AAbilitySystemTestPawn *DestActor = World->SpawnActor<AAbilitySystemTestPawn>();
 
-	UProperty *DamageProperty = FindFieldChecked<UProperty>(USkillSystemTestAttributeSet::StaticClass(), GET_MEMBER_NAME_CHECKED(USkillSystemTestAttributeSet, Damage));
+	UProperty *DamageProperty = FindFieldChecked<UProperty>(UAbilitySystemTestAttributeSet::StaticClass(), GET_MEMBER_NAME_CHECKED(UAbilitySystemTestAttributeSet, Damage));
 
 	UAttributeComponent * SourceComponent = SourceActor->AttributeComponent;
 	UAttributeComponent * DestComponent = DestActor->AttributeComponent;
-	SourceComponent->GetSet<USkillSystemTestAttributeSet>()->Health = StartHealth;
-	DestComponent->GetSet<USkillSystemTestAttributeSet>()->Health = StartHealth;
+	SourceComponent->GetSet<UAbilitySystemTestAttributeSet>()->Health = StartHealth;
+	DestComponent->GetSet<UAbilitySystemTestAttributeSet>()->Health = StartHealth;
 
 	// Setup a GE to provide immunity from incoming GEs
 	{
-		SKILL_LOG_SCOPE(TEXT("Apply ImmunityBuff"))
+		ABILITY_LOG_SCOPE(TEXT("Apply ImmunityBuff"))
 
 		UGameplayEffect * BaseImmunityEffect = Cast<UGameplayEffect>(StaticConstructObject(UGameplayEffect::StaticClass(), GetTransientPackage(), FName(TEXT("ImmunityBuff"))));
 		BaseImmunityEffect->Duration.SetValue(UGameplayEffect::INFINITE_DURATION);
@@ -4659,7 +4659,7 @@ bool GameplayEffectsTest_ImmunityIncoming(UWorld *World, FAutomationTestBase * T
 
 	// Apply Damage
 	{
-		SKILL_LOG_SCOPE(TEXT("Apply InstantDamage"))
+		ABILITY_LOG_SCOPE(TEXT("Apply InstantDamage"))
 
 		UGameplayEffect * BaseDmgEffect = Cast<UGameplayEffect>(StaticConstructObject(UGameplayEffect::StaticClass(), GetTransientPackage(), FName(TEXT("Damage"))));
 		BaseDmgEffect->Modifiers.SetNum(1);
@@ -4675,10 +4675,10 @@ bool GameplayEffectsTest_ImmunityIncoming(UWorld *World, FAutomationTestBase * T
 		SourceComponent->ApplyGameplayEffectToTarget(BaseDmgEffect, DestComponent, 1.f);
 
 		float ExpectedValue = StartHealth;
-		float ActualValue = DestComponent->GetSet<USkillSystemTestAttributeSet>()->Health;
+		float ActualValue = DestComponent->GetSet<UAbilitySystemTestAttributeSet>()->Health;
 
 		Test->TestTrue(SKILL_TEST_TEXT("Buff Instant Damage Applied"), (ActualValue == ExpectedValue));
-		SKILL_LOG(Log, TEXT("Final Health: %.2f"), DestComponent->GetSet<USkillSystemTestAttributeSet>()->Health);
+		ABILITY_LOG(Log, TEXT("Final Health: %.2f"), DestComponent->GetSet<UAbilitySystemTestAttributeSet>()->Health);
 	}
 
 	World->EditorDestroyActor(SourceActor, false);
@@ -4692,19 +4692,19 @@ bool GameplayEffectsTest_ImmunityOutgoing(UWorld *World, FAutomationTestBase * T
 	const float StartHealth = 100.f;
 	const float DamageValue = 5.f;
 
-	ASkillSystemTestPawn *SourceActor = World->SpawnActor<ASkillSystemTestPawn>();
-	ASkillSystemTestPawn *DestActor = World->SpawnActor<ASkillSystemTestPawn>();
+	AAbilitySystemTestPawn *SourceActor = World->SpawnActor<AAbilitySystemTestPawn>();
+	AAbilitySystemTestPawn *DestActor = World->SpawnActor<AAbilitySystemTestPawn>();
 
-	UProperty *DamageProperty = FindFieldChecked<UProperty>(USkillSystemTestAttributeSet::StaticClass(), GET_MEMBER_NAME_CHECKED(USkillSystemTestAttributeSet, Damage));
+	UProperty *DamageProperty = FindFieldChecked<UProperty>(UAbilitySystemTestAttributeSet::StaticClass(), GET_MEMBER_NAME_CHECKED(UAbilitySystemTestAttributeSet, Damage));
 
 	UAttributeComponent * SourceComponent = SourceActor->AttributeComponent;
 	UAttributeComponent * DestComponent = DestActor->AttributeComponent;
-	SourceComponent->GetSet<USkillSystemTestAttributeSet>()->Health = StartHealth;
-	DestComponent->GetSet<USkillSystemTestAttributeSet>()->Health = StartHealth;
+	SourceComponent->GetSet<UAbilitySystemTestAttributeSet>()->Health = StartHealth;
+	DestComponent->GetSet<UAbilitySystemTestAttributeSet>()->Health = StartHealth;
 
 	// Setup a GE to provide immunity from outgoing GEs
 	{
-		SKILL_LOG_SCOPE(TEXT("Apply ImmunityBuff"))
+		ABILITY_LOG_SCOPE(TEXT("Apply ImmunityBuff"))
 
 		UGameplayEffect * BaseImmunityEffect = Cast<UGameplayEffect>(StaticConstructObject(UGameplayEffect::StaticClass(), GetTransientPackage(), FName(TEXT("ImmunityBuff"))));
 		BaseImmunityEffect->Duration.SetValue(UGameplayEffect::INFINITE_DURATION);
@@ -4717,7 +4717,7 @@ bool GameplayEffectsTest_ImmunityOutgoing(UWorld *World, FAutomationTestBase * T
 
 	// Apply Damage
 	{
-		SKILL_LOG_SCOPE(TEXT("Apply InstantDamage"))
+		ABILITY_LOG_SCOPE(TEXT("Apply InstantDamage"))
 
 		UGameplayEffect * BaseDmgEffect = Cast<UGameplayEffect>(StaticConstructObject(UGameplayEffect::StaticClass(), GetTransientPackage(), FName(TEXT("Damage"))));
 		BaseDmgEffect->Modifiers.SetNum(1);
@@ -4733,10 +4733,10 @@ bool GameplayEffectsTest_ImmunityOutgoing(UWorld *World, FAutomationTestBase * T
 		SourceComponent->ApplyGameplayEffectToTarget(BaseDmgEffect, DestComponent, 1.f);
 
 		float ExpectedValue = StartHealth;
-		float ActualValue = DestComponent->GetSet<USkillSystemTestAttributeSet>()->Health;
+		float ActualValue = DestComponent->GetSet<UAbilitySystemTestAttributeSet>()->Health;
 
 		Test->TestTrue(SKILL_TEST_TEXT("Buff Instant Damage Applied"), (ActualValue == ExpectedValue));
-		SKILL_LOG(Log, TEXT("Final Health: %.2f"), DestComponent->GetSet<USkillSystemTestAttributeSet>()->Health);
+		ABILITY_LOG(Log, TEXT("Final Health: %.2f"), DestComponent->GetSet<UAbilitySystemTestAttributeSet>()->Health);
 	}
 
 	World->EditorDestroyActor(SourceActor, false);
@@ -4753,19 +4753,19 @@ bool GameplayEffectsTest_ImmunityMod(UWorld *World, FAutomationTestBase * Test)
 	const float DamageValue = 5.f;
 	const float BonusDamageMultiplier = 2.f;
 
-	ASkillSystemTestPawn *SourceActor = World->SpawnActor<ASkillSystemTestPawn>();
-	ASkillSystemTestPawn *DestActor = World->SpawnActor<ASkillSystemTestPawn>();
+	AAbilitySystemTestPawn *SourceActor = World->SpawnActor<AAbilitySystemTestPawn>();
+	AAbilitySystemTestPawn *DestActor = World->SpawnActor<AAbilitySystemTestPawn>();
 
-	UProperty *DamageProperty = FindFieldChecked<UProperty>(USkillSystemTestAttributeSet::StaticClass(), GET_MEMBER_NAME_CHECKED(USkillSystemTestAttributeSet, Damage));
+	UProperty *DamageProperty = FindFieldChecked<UProperty>(UAbilitySystemTestAttributeSet::StaticClass(), GET_MEMBER_NAME_CHECKED(UAbilitySystemTestAttributeSet, Damage));
 
 	UAttributeComponent * SourceComponent = SourceActor->AttributeComponent;
 	UAttributeComponent * DestComponent = DestActor->AttributeComponent;
-	SourceComponent->GetSet<USkillSystemTestAttributeSet>()->Health = StartHealth;
-	DestComponent->GetSet<USkillSystemTestAttributeSet>()->Health = StartHealth;
+	SourceComponent->GetSet<UAbilitySystemTestAttributeSet>()->Health = StartHealth;
+	DestComponent->GetSet<UAbilitySystemTestAttributeSet>()->Health = StartHealth;
 
 	// Setup a GE to provide immunity to buffed damage
 	{
-		SKILL_LOG_SCOPE(TEXT("Apply Immunity"))
+		ABILITY_LOG_SCOPE(TEXT("Apply Immunity"))
 
 		UGameplayEffect * BaseImmunityEffect = Cast<UGameplayEffect>(StaticConstructObject(UGameplayEffect::StaticClass(), GetTransientPackage(), FName(TEXT("ImmunityBuff"))));
 		BaseImmunityEffect->Duration.SetValue(UGameplayEffect::INFINITE_DURATION);
@@ -4778,7 +4778,7 @@ bool GameplayEffectsTest_ImmunityMod(UWorld *World, FAutomationTestBase * Test)
 
 	// Setup a GE to modify OutgoingGEs
 	{
-		SKILL_LOG_SCOPE(TEXT("Apply DamageBuff"))
+		ABILITY_LOG_SCOPE(TEXT("Apply DamageBuff"))
 
 		UGameplayEffect * BaseDmgEffect = Cast<UGameplayEffect>(StaticConstructObject(UGameplayEffect::StaticClass(), GetTransientPackage(), FName(TEXT("DamageBuff"))));
 		BaseDmgEffect->Modifiers.SetNum(1);
@@ -4795,7 +4795,7 @@ bool GameplayEffectsTest_ImmunityMod(UWorld *World, FAutomationTestBase * Test)
 
 	// Apply Damage
 	{
-		SKILL_LOG_SCOPE(TEXT("Apply InstantDamage"))
+		ABILITY_LOG_SCOPE(TEXT("Apply InstantDamage"))
 
 		UGameplayEffect * BaseDmgEffect = Cast<UGameplayEffect>(StaticConstructObject(UGameplayEffect::StaticClass(), GetTransientPackage(), FName(TEXT("Damage"))));
 		BaseDmgEffect->Modifiers.SetNum(1);
@@ -4810,10 +4810,10 @@ bool GameplayEffectsTest_ImmunityMod(UWorld *World, FAutomationTestBase * Test)
 		SourceComponent->ApplyGameplayEffectToTarget(BaseDmgEffect, DestComponent, 1.f);
 
 		float ExpectedValue = StartHealth;
-		float ActualValue = DestComponent->GetSet<USkillSystemTestAttributeSet>()->Health;
+		float ActualValue = DestComponent->GetSet<UAbilitySystemTestAttributeSet>()->Health;
 
 		Test->TestTrue(SKILL_TEST_TEXT("Buff Instant Damage Applied"), (ActualValue == ExpectedValue));
-		SKILL_LOG(Log, TEXT("Final Health: %.2f"), DestComponent->GetSet<USkillSystemTestAttributeSet>()->Health);
+		ABILITY_LOG(Log, TEXT("Final Health: %.2f"), DestComponent->GetSet<UAbilitySystemTestAttributeSet>()->Health);
 	}
 
 	World->EditorDestroyActor(SourceActor, false);
@@ -4829,19 +4829,19 @@ bool GameplayEffectsTest_ImmunityBlockedBuff(UWorld *World, FAutomationTestBase
 	const float StartHealth = 100.f;
 	const float DamageValue = 5.f;
 
-	ASkillSystemTestPawn *SourceActor = World->SpawnActor<ASkillSystemTestPawn>();
-	ASkillSystemTestPawn *DestActor = World->SpawnActor<ASkillSystemTestPawn>();
+	AAbilitySystemTestPawn *SourceActor = World->SpawnActor<AAbilitySystemTestPawn>();
+	AAbilitySystemTestPawn *DestActor = World->SpawnActor<AAbilitySystemTestPawn>();
 
-	UProperty *DamageProperty = FindFieldChecked<UProperty>(USkillSystemTestAttributeSet::StaticClass(), GET_MEMBER_NAME_CHECKED(USkillSystemTestAttributeSet, Damage));
+	UProperty *DamageProperty = FindFieldChecked<UProperty>(UAbilitySystemTestAttributeSet::StaticClass(), GET_MEMBER_NAME_CHECKED(UAbilitySystemTestAttributeSet, Damage));
 
 	UAttributeComponent * SourceComponent = SourceActor->AttributeComponent;
 	UAttributeComponent * DestComponent = DestActor->AttributeComponent;
-	SourceComponent->GetSet<USkillSystemTestAttributeSet>()->Health = StartHealth;
-	DestComponent->GetSet<USkillSystemTestAttributeSet>()->Health = StartHealth;
+	SourceComponent->GetSet<UAbilitySystemTestAttributeSet>()->Health = StartHealth;
+	DestComponent->GetSet<UAbilitySystemTestAttributeSet>()->Health = StartHealth;
 
 	// Setup a GE to provide immunity to buffed damage
 	{
-		SKILL_LOG_SCOPE(TEXT("Apply Immunity"))
+		ABILITY_LOG_SCOPE(TEXT("Apply Immunity"))
 
 		UGameplayEffect * BaseImmunityEffect = Cast<UGameplayEffect>(StaticConstructObject(UGameplayEffect::StaticClass(), GetTransientPackage(), FName(TEXT("ImmunityBuff"))));
 		BaseImmunityEffect->Duration.SetValue(UGameplayEffect::INFINITE_DURATION);
@@ -4854,7 +4854,7 @@ bool GameplayEffectsTest_ImmunityBlockedBuff(UWorld *World, FAutomationTestBase
 
 	// Apply base damage and a buff that causes extra damage
 	{
-		SKILL_LOG_SCOPE(TEXT("Apply DamageBuff and InstantDamage"))
+		ABILITY_LOG_SCOPE(TEXT("Apply DamageBuff and InstantDamage"))
 
 		UGameplayEffect * BuffEffect = Cast<UGameplayEffect>(StaticConstructObject(UGameplayEffect::StaticClass(), GetTransientPackage(), FName(TEXT("DamageBuff"))));
 		BuffEffect->Modifiers.SetNum(1);
@@ -4879,10 +4879,10 @@ bool GameplayEffectsTest_ImmunityBlockedBuff(UWorld *World, FAutomationTestBase
 		SourceComponent->ApplyGameplayEffectToTarget(BaseDmgEffect, DestComponent, 1.f);
 
 		float ExpectedValue = (StartHealth - DamageValue);
-		float ActualValue = DestComponent->GetSet<USkillSystemTestAttributeSet>()->Health;
+		float ActualValue = DestComponent->GetSet<UAbilitySystemTestAttributeSet>()->Health;
 
 		Test->TestTrue(SKILL_TEST_TEXT("Buff Instant Damage Applied. ActualValue: %.2f. ExpectedValue: %.2f.", ActualValue, ExpectedValue), (ActualValue == ExpectedValue));
-		SKILL_LOG(Log, TEXT("Final Health: %.2f"), DestComponent->GetSet<USkillSystemTestAttributeSet>()->Health);
+		ABILITY_LOG(Log, TEXT("Final Health: %.2f"), DestComponent->GetSet<UAbilitySystemTestAttributeSet>()->Health);
 	}
 
 	World->EditorDestroyActor(SourceActor, false);
@@ -4898,19 +4898,19 @@ bool GameplayEffectsTest_ImmunityBlockedBaseAndBuff(UWorld *World, FAutomationTe
 	const float StartHealth = 100.f;
 	const float DamageValue = 5.f;
 
-	ASkillSystemTestPawn *SourceActor = World->SpawnActor<ASkillSystemTestPawn>();
-	ASkillSystemTestPawn *DestActor = World->SpawnActor<ASkillSystemTestPawn>();
+	AAbilitySystemTestPawn *SourceActor = World->SpawnActor<AAbilitySystemTestPawn>();
+	AAbilitySystemTestPawn *DestActor = World->SpawnActor<AAbilitySystemTestPawn>();
 
-	UProperty *DamageProperty = FindFieldChecked<UProperty>(USkillSystemTestAttributeSet::StaticClass(), GET_MEMBER_NAME_CHECKED(USkillSystemTestAttributeSet, Damage));
+	UProperty *DamageProperty = FindFieldChecked<UProperty>(UAbilitySystemTestAttributeSet::StaticClass(), GET_MEMBER_NAME_CHECKED(UAbilitySystemTestAttributeSet, Damage));
 
 	UAttributeComponent * SourceComponent = SourceActor->AttributeComponent;
 	UAttributeComponent * DestComponent = DestActor->AttributeComponent;
-	SourceComponent->GetSet<USkillSystemTestAttributeSet>()->Health = StartHealth;
-	DestComponent->GetSet<USkillSystemTestAttributeSet>()->Health = StartHealth;
+	SourceComponent->GetSet<UAbilitySystemTestAttributeSet>()->Health = StartHealth;
+	DestComponent->GetSet<UAbilitySystemTestAttributeSet>()->Health = StartHealth;
 
 	// Setup a GE to provide immunity to buffed damage
 	{
-		SKILL_LOG_SCOPE(TEXT("Apply Immunity"))
+		ABILITY_LOG_SCOPE(TEXT("Apply Immunity"))
 
 		UGameplayEffect * BaseImmunityEffect = Cast<UGameplayEffect>(StaticConstructObject(UGameplayEffect::StaticClass(), GetTransientPackage(), FName(TEXT("ImmunityBuff"))));
 		BaseImmunityEffect->Duration.SetValue(UGameplayEffect::INFINITE_DURATION);
@@ -4923,7 +4923,7 @@ bool GameplayEffectsTest_ImmunityBlockedBaseAndBuff(UWorld *World, FAutomationTe
 
 	// Apply base damage and a buff that causes extra damage
 	{
-		SKILL_LOG_SCOPE(TEXT("Apply DamageBuff and InstantDamage"))
+		ABILITY_LOG_SCOPE(TEXT("Apply DamageBuff and InstantDamage"))
 
 		UGameplayEffect * BuffEffect = Cast<UGameplayEffect>(StaticConstructObject(UGameplayEffect::StaticClass(), GetTransientPackage(), FName(TEXT("DamageBuff"))));
 		BuffEffect->Modifiers.SetNum(1);
@@ -4948,10 +4948,10 @@ bool GameplayEffectsTest_ImmunityBlockedBaseAndBuff(UWorld *World, FAutomationTe
 		SourceComponent->ApplyGameplayEffectToTarget(BaseDmgEffect, DestComponent, 1.f);
 
 		float ExpectedValue = StartHealth;
-		float ActualValue = DestComponent->GetSet<USkillSystemTestAttributeSet>()->Health;
+		float ActualValue = DestComponent->GetSet<UAbilitySystemTestAttributeSet>()->Health;
 
 		Test->TestTrue(SKILL_TEST_TEXT("Buff Instant Damage Applied"), (ActualValue == ExpectedValue));
-		SKILL_LOG(Log, TEXT("Final Health: %.2f"), DestComponent->GetSet<USkillSystemTestAttributeSet>()->Health);
+		ABILITY_LOG(Log, TEXT("Final Health: %.2f"), DestComponent->GetSet<UAbilitySystemTestAttributeSet>()->Health);
 	}
 
 	World->EditorDestroyActor(SourceActor, false);
@@ -4967,15 +4967,15 @@ bool GameplayEffectsTest_ImmunityActiveGE(UWorld *World, FAutomationTestBase * T
 	const float DamageValue = 5.f;
 	const float DamageProtectionDivisor = 2.f;
 
-	ASkillSystemTestPawn *SourceActor = World->SpawnActor<ASkillSystemTestPawn>();
-	ASkillSystemTestPawn *DestActor = World->SpawnActor<ASkillSystemTestPawn>();
+	AAbilitySystemTestPawn *SourceActor = World->SpawnActor<AAbilitySystemTestPawn>();
+	AAbilitySystemTestPawn *DestActor = World->SpawnActor<AAbilitySystemTestPawn>();
 
-	UProperty *DamageProperty = FindFieldChecked<UProperty>(USkillSystemTestAttributeSet::StaticClass(), GET_MEMBER_NAME_CHECKED(USkillSystemTestAttributeSet, Damage));
+	UProperty *DamageProperty = FindFieldChecked<UProperty>(UAbilitySystemTestAttributeSet::StaticClass(), GET_MEMBER_NAME_CHECKED(UAbilitySystemTestAttributeSet, Damage));
 
 	UAttributeComponent * SourceComponent = SourceActor->AttributeComponent;
 	UAttributeComponent * DestComponent = DestActor->AttributeComponent;
-	SourceComponent->GetSet<USkillSystemTestAttributeSet>()->Health = StartHealth;
-	DestComponent->GetSet<USkillSystemTestAttributeSet>()->Health = StartHealth;
+	SourceComponent->GetSet<UAbilitySystemTestAttributeSet>()->Health = StartHealth;
+	DestComponent->GetSet<UAbilitySystemTestAttributeSet>()->Health = StartHealth;
 
 	// Apply a buff to reduce incoming damage on DestComponent
 	{
@@ -4994,7 +4994,7 @@ bool GameplayEffectsTest_ImmunityActiveGE(UWorld *World, FAutomationTestBase * T
 
 	// Apply Damage to verify the buff is working
 	{
-		SKILL_LOG_SCOPE(TEXT("Apply InstantDamage"))
+		ABILITY_LOG_SCOPE(TEXT("Apply InstantDamage"))
 
 		UGameplayEffect * BaseDmgEffect = Cast<UGameplayEffect>(StaticConstructObject(UGameplayEffect::StaticClass(), GetTransientPackage(), FName(TEXT("Damage"))));
 		BaseDmgEffect->Modifiers.SetNum(1);
@@ -5009,15 +5009,15 @@ bool GameplayEffectsTest_ImmunityActiveGE(UWorld *World, FAutomationTestBase * T
 		SourceComponent->ApplyGameplayEffectToTarget(BaseDmgEffect, DestComponent, 1.f);
 
 		float ExpectedValue = StartHealth - (DamageValue / DamageProtectionDivisor);
-		float ActualValue = DestComponent->GetSet<USkillSystemTestAttributeSet>()->Health;
+		float ActualValue = DestComponent->GetSet<UAbilitySystemTestAttributeSet>()->Health;
 
 		Test->TestTrue(SKILL_TEST_TEXT("Buff Instant Damage Applied"), (ActualValue == ExpectedValue));
-		SKILL_LOG(Log, TEXT("Final Health: %.2f"), DestComponent->GetSet<USkillSystemTestAttributeSet>()->Health);
+		ABILITY_LOG(Log, TEXT("Final Health: %.2f"), DestComponent->GetSet<UAbilitySystemTestAttributeSet>()->Health);
 	}
 
 	// Setup a GE to remove the buff
 	{
-		SKILL_LOG_SCOPE(TEXT("Apply Immunity"))
+		ABILITY_LOG_SCOPE(TEXT("Apply Immunity"))
 
 		UGameplayEffect * BaseImmunityEffect = Cast<UGameplayEffect>(StaticConstructObject(UGameplayEffect::StaticClass(), GetTransientPackage(), FName(TEXT("ImmunityBuff"))));
 		BaseImmunityEffect->Duration.SetValue(UGameplayEffect::INSTANT_APPLICATION);
@@ -5030,7 +5030,7 @@ bool GameplayEffectsTest_ImmunityActiveGE(UWorld *World, FAutomationTestBase * T
 
 	// Apply Damage
 	{
-		SKILL_LOG_SCOPE(TEXT("Apply InstantDamage"))
+		ABILITY_LOG_SCOPE(TEXT("Apply InstantDamage"))
 
 		UGameplayEffect * BaseDmgEffect = Cast<UGameplayEffect>(StaticConstructObject(UGameplayEffect::StaticClass(), GetTransientPackage(), FName(TEXT("Damage"))));
 		BaseDmgEffect->Modifiers.SetNum(1);
@@ -5045,10 +5045,10 @@ bool GameplayEffectsTest_ImmunityActiveGE(UWorld *World, FAutomationTestBase * T
 		SourceComponent->ApplyGameplayEffectToTarget(BaseDmgEffect, DestComponent, 1.f);
 
 		float ExpectedValue = StartHealth - (DamageValue / DamageProtectionDivisor) - DamageValue;
-		float ActualValue = DestComponent->GetSet<USkillSystemTestAttributeSet>()->Health;
+		float ActualValue = DestComponent->GetSet<UAbilitySystemTestAttributeSet>()->Health;
 
 		Test->TestTrue(SKILL_TEST_TEXT("Buff Instant Damage Applied"), (ActualValue == ExpectedValue));
-		SKILL_LOG(Log, TEXT("Final Health: %.2f"), DestComponent->GetSet<USkillSystemTestAttributeSet>()->Health);
+		ABILITY_LOG(Log, TEXT("Final Health: %.2f"), DestComponent->GetSet<UAbilitySystemTestAttributeSet>()->Health);
 	}
 
 	World->EditorDestroyActor(SourceActor, false);
@@ -5064,15 +5064,15 @@ bool GameplayEffectsTest_ImmunityActiveMod(UWorld *World, FAutomationTestBase *
 	const float DamageValue = 5.f;
 	const float DamageProtectionDivisor = 2.f;
 
-	ASkillSystemTestPawn *SourceActor = World->SpawnActor<ASkillSystemTestPawn>();
-	ASkillSystemTestPawn *DestActor = World->SpawnActor<ASkillSystemTestPawn>();
+	AAbilitySystemTestPawn *SourceActor = World->SpawnActor<AAbilitySystemTestPawn>();
+	AAbilitySystemTestPawn *DestActor = World->SpawnActor<AAbilitySystemTestPawn>();
 
-	UProperty *DamageProperty = FindFieldChecked<UProperty>(USkillSystemTestAttributeSet::StaticClass(), GET_MEMBER_NAME_CHECKED(USkillSystemTestAttributeSet, Damage));
+	UProperty *DamageProperty = FindFieldChecked<UProperty>(UAbilitySystemTestAttributeSet::StaticClass(), GET_MEMBER_NAME_CHECKED(UAbilitySystemTestAttributeSet, Damage));
 
 	UAttributeComponent * SourceComponent = SourceActor->AttributeComponent;
 	UAttributeComponent * DestComponent = DestActor->AttributeComponent;
-	SourceComponent->GetSet<USkillSystemTestAttributeSet>()->Health = StartHealth;
-	DestComponent->GetSet<USkillSystemTestAttributeSet>()->Health = StartHealth;
+	SourceComponent->GetSet<UAbilitySystemTestAttributeSet>()->Health = StartHealth;
+	DestComponent->GetSet<UAbilitySystemTestAttributeSet>()->Health = StartHealth;
 
 	// Apply a buff to reduce incoming damage on DestComponent
 	{
@@ -5091,7 +5091,7 @@ bool GameplayEffectsTest_ImmunityActiveMod(UWorld *World, FAutomationTestBase *
 
 	// Apply Damage to verify the buff is working
 	{
-		SKILL_LOG_SCOPE(TEXT("Apply InstantDamage"))
+		ABILITY_LOG_SCOPE(TEXT("Apply InstantDamage"))
 
 		UGameplayEffect * BaseDmgEffect = Cast<UGameplayEffect>(StaticConstructObject(UGameplayEffect::StaticClass(), GetTransientPackage(), FName(TEXT("Damage"))));
 		BaseDmgEffect->Modifiers.SetNum(1);
@@ -5106,15 +5106,15 @@ bool GameplayEffectsTest_ImmunityActiveMod(UWorld *World, FAutomationTestBase *
 		SourceComponent->ApplyGameplayEffectToTarget(BaseDmgEffect, DestComponent, 1.f);
 
 		float ExpectedValue = StartHealth - (DamageValue / DamageProtectionDivisor);
-		float ActualValue = DestComponent->GetSet<USkillSystemTestAttributeSet>()->Health;
+		float ActualValue = DestComponent->GetSet<UAbilitySystemTestAttributeSet>()->Health;
 
 		Test->TestTrue(SKILL_TEST_TEXT("Buff Instant Damage Applied"), (ActualValue == ExpectedValue));
-		SKILL_LOG(Log, TEXT("Final Health: %.2f"), DestComponent->GetSet<USkillSystemTestAttributeSet>()->Health);
+		ABILITY_LOG(Log, TEXT("Final Health: %.2f"), DestComponent->GetSet<UAbilitySystemTestAttributeSet>()->Health);
 	}
 
 	// Setup a GE to remove the buff
 	{
-		SKILL_LOG_SCOPE(TEXT("Apply Immunity"))
+		ABILITY_LOG_SCOPE(TEXT("Apply Immunity"))
 
 		UGameplayEffect * BaseImmunityEffect = Cast<UGameplayEffect>(StaticConstructObject(UGameplayEffect::StaticClass(), GetTransientPackage(), FName(TEXT("ImmunityBuff"))));
 		BaseImmunityEffect->Duration.SetValue(UGameplayEffect::INSTANT_APPLICATION);
@@ -5127,7 +5127,7 @@ bool GameplayEffectsTest_ImmunityActiveMod(UWorld *World, FAutomationTestBase *
 
 	// Apply Damage
 	{
-		SKILL_LOG_SCOPE(TEXT("Apply InstantDamage"))
+		ABILITY_LOG_SCOPE(TEXT("Apply InstantDamage"))
 
 		UGameplayEffect * BaseDmgEffect = Cast<UGameplayEffect>(StaticConstructObject(UGameplayEffect::StaticClass(), GetTransientPackage(), FName(TEXT("Damage"))));
 		BaseDmgEffect->Modifiers.SetNum(1);
@@ -5142,10 +5142,10 @@ bool GameplayEffectsTest_ImmunityActiveMod(UWorld *World, FAutomationTestBase *
 		SourceComponent->ApplyGameplayEffectToTarget(BaseDmgEffect, DestComponent, 1.f);
 
 		float ExpectedValue = StartHealth - (DamageValue / DamageProtectionDivisor) - DamageValue;
-		float ActualValue = DestComponent->GetSet<USkillSystemTestAttributeSet>()->Health;
+		float ActualValue = DestComponent->GetSet<UAbilitySystemTestAttributeSet>()->Health;
 
 		Test->TestTrue(SKILL_TEST_TEXT("Buff Instant Damage Applied"), (ActualValue == ExpectedValue));
-		SKILL_LOG(Log, TEXT("Final Health: %.2f"), DestComponent->GetSet<USkillSystemTestAttributeSet>()->Health);
+		ABILITY_LOG(Log, TEXT("Final Health: %.2f"), DestComponent->GetSet<UAbilitySystemTestAttributeSet>()->Health);
 	}
 
 	World->EditorDestroyActor(SourceActor, false);
@@ -5160,19 +5160,19 @@ bool GameplayEffectsTest_ChanceToApplyToTarget(UWorld *World, FAutomationTestBas
 	const float StartHealth = 100.f;
 	const float DamageValue = 5.f;
 
-	ASkillSystemTestPawn *SourceActor = World->SpawnActor<ASkillSystemTestPawn>();
-	ASkillSystemTestPawn *DestActor = World->SpawnActor<ASkillSystemTestPawn>();
+	AAbilitySystemTestPawn *SourceActor = World->SpawnActor<AAbilitySystemTestPawn>();
+	AAbilitySystemTestPawn *DestActor = World->SpawnActor<AAbilitySystemTestPawn>();
 
-	UProperty *DamageProperty = FindFieldChecked<UProperty>(USkillSystemTestAttributeSet::StaticClass(), GET_MEMBER_NAME_CHECKED(USkillSystemTestAttributeSet, Damage));
+	UProperty *DamageProperty = FindFieldChecked<UProperty>(UAbilitySystemTestAttributeSet::StaticClass(), GET_MEMBER_NAME_CHECKED(UAbilitySystemTestAttributeSet, Damage));
 
 	UAttributeComponent * SourceComponent = SourceActor->AttributeComponent;
 	UAttributeComponent * DestComponent = DestActor->AttributeComponent;
 
-	SourceComponent->GetSet<USkillSystemTestAttributeSet>()->Health = StartHealth;
-	DestComponent->GetSet<USkillSystemTestAttributeSet>()->Health = StartHealth;
+	SourceComponent->GetSet<UAbilitySystemTestAttributeSet>()->Health = StartHealth;
+	DestComponent->GetSet<UAbilitySystemTestAttributeSet>()->Health = StartHealth;
 
 	{
-		SKILL_LOG_SCOPE(TEXT("Apply Damage mod"))
+		ABILITY_LOG_SCOPE(TEXT("Apply Damage mod"))
 
 		UGameplayEffect * BaseDmgEffect = Cast<UGameplayEffect>(StaticConstructObject(UGameplayEffect::StaticClass(), GetTransientPackage(), FName(TEXT("Damage"))));
 		BaseDmgEffect->Modifiers.SetNum(1);
@@ -5188,13 +5188,13 @@ bool GameplayEffectsTest_ChanceToApplyToTarget(UWorld *World, FAutomationTestBas
 		SourceComponent->ApplyGameplayEffectToTarget(BaseDmgEffect, DestComponent, 1.f);
 
 		float ExpectedValue = (StartHealth - DamageValue);
-		float ActualValue = DestComponent->GetSet<USkillSystemTestAttributeSet>()->Health;
+		float ActualValue = DestComponent->GetSet<UAbilitySystemTestAttributeSet>()->Health;
 
 		Test->TestTrue(SKILL_TEST_TEXT("Damage Applied. %.2f == %.2f", ActualValue, ExpectedValue), ActualValue == ExpectedValue);
 	}
 
 	{
-		SKILL_LOG_SCOPE(TEXT("Apply Damage mod"))
+		ABILITY_LOG_SCOPE(TEXT("Apply Damage mod"))
 
 		UGameplayEffect * BaseDmgEffect = Cast<UGameplayEffect>(StaticConstructObject(UGameplayEffect::StaticClass(), GetTransientPackage(), FName(TEXT("Damage"))));
 		BaseDmgEffect->Modifiers.SetNum(1);
@@ -5210,7 +5210,7 @@ bool GameplayEffectsTest_ChanceToApplyToTarget(UWorld *World, FAutomationTestBas
 		FActiveGameplayEffectHandle Handle = SourceComponent->ApplyGameplayEffectToTarget(BaseDmgEffect, DestComponent, 1.f);
 
 		float ExpectedValue = (StartHealth - DamageValue);
-		float ActualValue = DestComponent->GetSet<USkillSystemTestAttributeSet>()->Health;
+		float ActualValue = DestComponent->GetSet<UAbilitySystemTestAttributeSet>()->Health;
 
 		Test->TestTrue(SKILL_TEST_TEXT("Damage Applied. %.2f == %.2f", ActualValue, ExpectedValue), ActualValue == ExpectedValue);
 		Test->TestFalse(SKILL_TEST_TEXT("Effect applied to target when chance was 0.f"), Handle.IsValid());
@@ -5229,20 +5229,20 @@ bool GameplayEffectsTest_ChanceToExecuteOnActiveGEMod(UWorld *World, FAutomation
 	const float BonusDamageMultiplier = 2.f;
 	const float ExtraDamageMultiplier = 1.f;
 
-	ASkillSystemTestPawn *SourceActor = World->SpawnActor<ASkillSystemTestPawn>();
-	ASkillSystemTestPawn *DestActor = World->SpawnActor<ASkillSystemTestPawn>();
+	AAbilitySystemTestPawn *SourceActor = World->SpawnActor<AAbilitySystemTestPawn>();
+	AAbilitySystemTestPawn *DestActor = World->SpawnActor<AAbilitySystemTestPawn>();
 
-	UProperty *DamageProperty = FindFieldChecked<UProperty>(USkillSystemTestAttributeSet::StaticClass(), GET_MEMBER_NAME_CHECKED(USkillSystemTestAttributeSet, Damage));
+	UProperty *DamageProperty = FindFieldChecked<UProperty>(UAbilitySystemTestAttributeSet::StaticClass(), GET_MEMBER_NAME_CHECKED(UAbilitySystemTestAttributeSet, Damage));
 
 	UAttributeComponent * SourceComponent = SourceActor->AttributeComponent;
 	UAttributeComponent * DestComponent = DestActor->AttributeComponent;
 
-	SourceComponent->GetSet<USkillSystemTestAttributeSet>()->Health = StartHealth;
-	DestComponent->GetSet<USkillSystemTestAttributeSet>()->Health = StartHealth;
+	SourceComponent->GetSet<UAbilitySystemTestAttributeSet>()->Health = StartHealth;
+	DestComponent->GetSet<UAbilitySystemTestAttributeSet>()->Health = StartHealth;
 
 	// Setup a GE to modify OutgoingGEs
 	{
-		SKILL_LOG_SCOPE(TEXT("Apply DamageBuff"))
+		ABILITY_LOG_SCOPE(TEXT("Apply DamageBuff"))
 
 		UGameplayEffect * BuffDmgEffect = Cast<UGameplayEffect>(StaticConstructObject(UGameplayEffect::StaticClass(), GetTransientPackage(), FName(TEXT("DamageBuff"))));
 		BuffDmgEffect->Modifiers.SetNum(1);
@@ -5260,7 +5260,7 @@ bool GameplayEffectsTest_ChanceToExecuteOnActiveGEMod(UWorld *World, FAutomation
 
 	// attempt to modify the buff but fail because of a 0.f chance to apply
 	{
-		SKILL_LOG_SCOPE(TEXT("Fail to modify DamageBuff"))
+		ABILITY_LOG_SCOPE(TEXT("Fail to modify DamageBuff"))
 			
 		UGameplayEffect * ModBuffEffect = Cast<UGameplayEffect>(StaticConstructObject(UGameplayEffect::StaticClass(), GetTransientPackage(), FName(TEXT("ModDamageBuff"))));
 		ModBuffEffect->Modifiers.SetNum(1);
@@ -5276,7 +5276,7 @@ bool GameplayEffectsTest_ChanceToExecuteOnActiveGEMod(UWorld *World, FAutomation
 	}
 
 	{
-		SKILL_LOG_SCOPE(TEXT("Apply Damage mod"))
+		ABILITY_LOG_SCOPE(TEXT("Apply Damage mod"))
 
 		UGameplayEffect * BaseDmgEffect = Cast<UGameplayEffect>(StaticConstructObject(UGameplayEffect::StaticClass(), GetTransientPackage(), FName(TEXT("Damage"))));
 		BaseDmgEffect->Modifiers.SetNum(1);
@@ -5291,14 +5291,14 @@ bool GameplayEffectsTest_ChanceToExecuteOnActiveGEMod(UWorld *World, FAutomation
 		SourceComponent->ApplyGameplayEffectToTarget(BaseDmgEffect, DestComponent, 1.f);
 
 		float ExpectedValue = StartHealth - (DamageValue * BonusDamageMultiplier);
-		float ActualValue = DestComponent->GetSet<USkillSystemTestAttributeSet>()->Health;
+		float ActualValue = DestComponent->GetSet<UAbilitySystemTestAttributeSet>()->Health;
 
 		Test->TestTrue(SKILL_TEST_TEXT("Damage Applied. %.2f == %.2f", ActualValue, ExpectedValue), ActualValue == ExpectedValue);
 	}
 
 	// modify the buff
 	{
-		SKILL_LOG_SCOPE(TEXT("Modify DamageBuff"))
+		ABILITY_LOG_SCOPE(TEXT("Modify DamageBuff"))
 
 		UGameplayEffect * ModBuffEffect = Cast<UGameplayEffect>(StaticConstructObject(UGameplayEffect::StaticClass(), GetTransientPackage(), FName(TEXT("ModDamageBuff"))));
 		ModBuffEffect->Modifiers.SetNum(1);
@@ -5314,7 +5314,7 @@ bool GameplayEffectsTest_ChanceToExecuteOnActiveGEMod(UWorld *World, FAutomation
 	}
 
 	{
-		SKILL_LOG_SCOPE(TEXT("Apply Damage mod"))
+		ABILITY_LOG_SCOPE(TEXT("Apply Damage mod"))
 
 		UGameplayEffect * BaseDmgEffect = Cast<UGameplayEffect>(StaticConstructObject(UGameplayEffect::StaticClass(), GetTransientPackage(), FName(TEXT("Damage"))));
 		BaseDmgEffect->Modifiers.SetNum(1);
@@ -5329,7 +5329,7 @@ bool GameplayEffectsTest_ChanceToExecuteOnActiveGEMod(UWorld *World, FAutomation
 		SourceComponent->ApplyGameplayEffectToTarget(BaseDmgEffect, DestComponent, 1.f);
 
 		float ExpectedValue = StartHealth - (DamageValue * BonusDamageMultiplier) - (DamageValue * (BonusDamageMultiplier + ExtraDamageMultiplier));
-		float ActualValue = DestComponent->GetSet<USkillSystemTestAttributeSet>()->Health;
+		float ActualValue = DestComponent->GetSet<UAbilitySystemTestAttributeSet>()->Health;
 
 		Test->TestTrue(SKILL_TEST_TEXT("Damage Applied. %.2f == %.2f", ActualValue, ExpectedValue), ActualValue == ExpectedValue);
 	}
@@ -5346,20 +5346,20 @@ bool GameplayEffectsTest_ChanceToExecuteOnActiveGEImmunity(UWorld *World, FAutom
 	const float DamageValue = 5.f;
 	const float BonusDamageMultiplier = 2.f;
 
-	ASkillSystemTestPawn *SourceActor = World->SpawnActor<ASkillSystemTestPawn>();
-	ASkillSystemTestPawn *DestActor = World->SpawnActor<ASkillSystemTestPawn>();
+	AAbilitySystemTestPawn *SourceActor = World->SpawnActor<AAbilitySystemTestPawn>();
+	AAbilitySystemTestPawn *DestActor = World->SpawnActor<AAbilitySystemTestPawn>();
 
-	UProperty *DamageProperty = FindFieldChecked<UProperty>(USkillSystemTestAttributeSet::StaticClass(), GET_MEMBER_NAME_CHECKED(USkillSystemTestAttributeSet, Damage));
+	UProperty *DamageProperty = FindFieldChecked<UProperty>(UAbilitySystemTestAttributeSet::StaticClass(), GET_MEMBER_NAME_CHECKED(UAbilitySystemTestAttributeSet, Damage));
 
 	UAttributeComponent * SourceComponent = SourceActor->AttributeComponent;
 	UAttributeComponent * DestComponent = DestActor->AttributeComponent;
 
-	SourceComponent->GetSet<USkillSystemTestAttributeSet>()->Health = StartHealth;
-	DestComponent->GetSet<USkillSystemTestAttributeSet>()->Health = StartHealth;
+	SourceComponent->GetSet<UAbilitySystemTestAttributeSet>()->Health = StartHealth;
+	DestComponent->GetSet<UAbilitySystemTestAttributeSet>()->Health = StartHealth;
 
 	// Setup a GE to modify OutgoingGEs
 	{
-		SKILL_LOG_SCOPE(TEXT("Apply DamageBuff"))
+		ABILITY_LOG_SCOPE(TEXT("Apply DamageBuff"))
 
 		UGameplayEffect * BuffDmgEffect = Cast<UGameplayEffect>(StaticConstructObject(UGameplayEffect::StaticClass(), GetTransientPackage(), FName(TEXT("DamageBuff"))));
 		BuffDmgEffect->Modifiers.SetNum(1);
@@ -5377,7 +5377,7 @@ bool GameplayEffectsTest_ChanceToExecuteOnActiveGEImmunity(UWorld *World, FAutom
 
 	// attempt to remove the buff but fail because of a 0.f chance to apply
 	{
-		SKILL_LOG_SCOPE(TEXT("Fail to remove DamageBuff"))
+		ABILITY_LOG_SCOPE(TEXT("Fail to remove DamageBuff"))
 
 		UGameplayEffect * ModBuffEffect = Cast<UGameplayEffect>(StaticConstructObject(UGameplayEffect::StaticClass(), GetTransientPackage(), FName(TEXT("ModDamageBuff"))));
 		ModBuffEffect->GameplayEffectRequiredTags.AddTag(IGameplayTagsModule::RequestGameplayTag(FName(TEXT("Damage.Buff"))));
@@ -5390,7 +5390,7 @@ bool GameplayEffectsTest_ChanceToExecuteOnActiveGEImmunity(UWorld *World, FAutom
 	}
 
 	{
-		SKILL_LOG_SCOPE(TEXT("Apply Damage mod"))
+		ABILITY_LOG_SCOPE(TEXT("Apply Damage mod"))
 
 		UGameplayEffect * BaseDmgEffect = Cast<UGameplayEffect>(StaticConstructObject(UGameplayEffect::StaticClass(), GetTransientPackage(), FName(TEXT("Damage"))));
 		BaseDmgEffect->Modifiers.SetNum(1);
@@ -5405,14 +5405,14 @@ bool GameplayEffectsTest_ChanceToExecuteOnActiveGEImmunity(UWorld *World, FAutom
 		SourceComponent->ApplyGameplayEffectToTarget(BaseDmgEffect, DestComponent, 1.f);
 
 		float ExpectedValue = StartHealth - (DamageValue * BonusDamageMultiplier);
-		float ActualValue = DestComponent->GetSet<USkillSystemTestAttributeSet>()->Health;
+		float ActualValue = DestComponent->GetSet<UAbilitySystemTestAttributeSet>()->Health;
 
 		Test->TestTrue(SKILL_TEST_TEXT("Damage Applied. %.2f == %.2f", ActualValue, ExpectedValue), ActualValue == ExpectedValue);
 	}
 
 	// modify the buff
 	{
-		SKILL_LOG_SCOPE(TEXT("Remove DamageBuff"))
+		ABILITY_LOG_SCOPE(TEXT("Remove DamageBuff"))
 
 		UGameplayEffect * ModBuffEffect = Cast<UGameplayEffect>(StaticConstructObject(UGameplayEffect::StaticClass(), GetTransientPackage(), FName(TEXT("ModDamageBuff"))));
 		ModBuffEffect->GameplayEffectRequiredTags.AddTag(IGameplayTagsModule::RequestGameplayTag(FName(TEXT("Damage.Buff"))));
@@ -5425,7 +5425,7 @@ bool GameplayEffectsTest_ChanceToExecuteOnActiveGEImmunity(UWorld *World, FAutom
 	}
 
 	{
-		SKILL_LOG_SCOPE(TEXT("Apply Damage mod"))
+		ABILITY_LOG_SCOPE(TEXT("Apply Damage mod"))
 
 		UGameplayEffect * BaseDmgEffect = Cast<UGameplayEffect>(StaticConstructObject(UGameplayEffect::StaticClass(), GetTransientPackage(), FName(TEXT("Damage"))));
 		BaseDmgEffect->Modifiers.SetNum(1);
@@ -5440,7 +5440,7 @@ bool GameplayEffectsTest_ChanceToExecuteOnActiveGEImmunity(UWorld *World, FAutom
 		SourceComponent->ApplyGameplayEffectToTarget(BaseDmgEffect, DestComponent, 1.f);
 
 		float ExpectedValue = StartHealth - (DamageValue * BonusDamageMultiplier) - (DamageValue);
-		float ActualValue = DestComponent->GetSet<USkillSystemTestAttributeSet>()->Health;
+		float ActualValue = DestComponent->GetSet<UAbilitySystemTestAttributeSet>()->Health;
 
 		Test->TestTrue(SKILL_TEST_TEXT("Damage Applied. %.2f == %.2f", ActualValue, ExpectedValue), ActualValue == ExpectedValue);
 	}
@@ -5457,20 +5457,20 @@ bool GameplayEffectsTest_ChanceToExecuteOnOutgoingGEMod(UWorld *World, FAutomati
 	const float DamageValue = 5.f;
 	const float BonusDamageMultiplier = 2.f;
 
-	ASkillSystemTestPawn *SourceActor = World->SpawnActor<ASkillSystemTestPawn>();
-	ASkillSystemTestPawn *DestActor = World->SpawnActor<ASkillSystemTestPawn>();
+	AAbilitySystemTestPawn *SourceActor = World->SpawnActor<AAbilitySystemTestPawn>();
+	AAbilitySystemTestPawn *DestActor = World->SpawnActor<AAbilitySystemTestPawn>();
 
-	UProperty *DamageProperty = FindFieldChecked<UProperty>(USkillSystemTestAttributeSet::StaticClass(), GET_MEMBER_NAME_CHECKED(USkillSystemTestAttributeSet, Damage));
+	UProperty *DamageProperty = FindFieldChecked<UProperty>(UAbilitySystemTestAttributeSet::StaticClass(), GET_MEMBER_NAME_CHECKED(UAbilitySystemTestAttributeSet, Damage));
 	
 	UAttributeComponent * SourceComponent = SourceActor->AttributeComponent;
 	UAttributeComponent * DestComponent = DestActor->AttributeComponent;
-	SourceComponent->GetSet<USkillSystemTestAttributeSet>()->Health = StartHealth;
-	DestComponent->GetSet<USkillSystemTestAttributeSet>()->Health = StartHealth;
+	SourceComponent->GetSet<UAbilitySystemTestAttributeSet>()->Health = StartHealth;
+	DestComponent->GetSet<UAbilitySystemTestAttributeSet>()->Health = StartHealth;
 
 	// Setup a GE to modify OutgoingGEs
 	// this GE won't do anything because the chance to execute is zero
 	{
-		SKILL_LOG_SCOPE(TEXT("Apply useless DamageBuff"))
+		ABILITY_LOG_SCOPE(TEXT("Apply useless DamageBuff"))
 
 		UGameplayEffect * BuffDmgEffect = Cast<UGameplayEffect>(StaticConstructObject(UGameplayEffect::StaticClass(), GetTransientPackage(), FName(TEXT("DamageBuff"))));
 		BuffDmgEffect->Modifiers.SetNum(1);
@@ -5488,7 +5488,7 @@ bool GameplayEffectsTest_ChanceToExecuteOnOutgoingGEMod(UWorld *World, FAutomati
 	}
 
 	{
-		SKILL_LOG_SCOPE(TEXT("Apply Damage mod"))
+		ABILITY_LOG_SCOPE(TEXT("Apply Damage mod"))
 
 		UGameplayEffect * BaseDmgEffect = Cast<UGameplayEffect>(StaticConstructObject(UGameplayEffect::StaticClass(), GetTransientPackage(), FName(TEXT("Damage"))));
 		BaseDmgEffect->Modifiers.SetNum(1);
@@ -5503,7 +5503,7 @@ bool GameplayEffectsTest_ChanceToExecuteOnOutgoingGEMod(UWorld *World, FAutomati
 		SourceComponent->ApplyGameplayEffectToTarget(BaseDmgEffect, DestComponent, 1.f);
 
 		float ExpectedValue = StartHealth - DamageValue;
-		float ActualValue = DestComponent->GetSet<USkillSystemTestAttributeSet>()->Health;
+		float ActualValue = DestComponent->GetSet<UAbilitySystemTestAttributeSet>()->Health;
 
 		Test->TestTrue(SKILL_TEST_TEXT("Damage Applied. %.2f == %.2f", ActualValue, ExpectedValue), ActualValue == ExpectedValue);
 	}
@@ -5511,7 +5511,7 @@ bool GameplayEffectsTest_ChanceToExecuteOnOutgoingGEMod(UWorld *World, FAutomati
 	// Setup a GE to modify OutgoingGEs
 	// this GE will always execute work because the chance to execute is one
 	{
-		SKILL_LOG_SCOPE(TEXT("Apply DamageBuff"))
+		ABILITY_LOG_SCOPE(TEXT("Apply DamageBuff"))
 
 		UGameplayEffect * BuffDmgEffect = Cast<UGameplayEffect>(StaticConstructObject(UGameplayEffect::StaticClass(), GetTransientPackage(), FName(TEXT("DamageBuff"))));
 		BuffDmgEffect->Modifiers.SetNum(1);
@@ -5529,7 +5529,7 @@ bool GameplayEffectsTest_ChanceToExecuteOnOutgoingGEMod(UWorld *World, FAutomati
 	}
 
 	{
-		SKILL_LOG_SCOPE(TEXT("Apply Damage mod"))
+		ABILITY_LOG_SCOPE(TEXT("Apply Damage mod"))
 
 		UGameplayEffect * BaseDmgEffect = Cast<UGameplayEffect>(StaticConstructObject(UGameplayEffect::StaticClass(), GetTransientPackage(), FName(TEXT("Damage"))));
 		BaseDmgEffect->Modifiers.SetNum(1);
@@ -5544,7 +5544,7 @@ bool GameplayEffectsTest_ChanceToExecuteOnOutgoingGEMod(UWorld *World, FAutomati
 		SourceComponent->ApplyGameplayEffectToTarget(BaseDmgEffect, DestComponent, 1.f);
 
 		float ExpectedValue = StartHealth - DamageValue - (DamageValue * BonusDamageMultiplier);
-		float ActualValue = DestComponent->GetSet<USkillSystemTestAttributeSet>()->Health;
+		float ActualValue = DestComponent->GetSet<UAbilitySystemTestAttributeSet>()->Health;
 
 		Test->TestTrue(SKILL_TEST_TEXT("Damage Applied. %.2f == %.2f", ActualValue, ExpectedValue), ActualValue == ExpectedValue);
 	}
@@ -5560,20 +5560,20 @@ bool GameplayEffectsTest_ChanceToExecuteOnOutgoingGEImmunity(UWorld *World, FAut
 	const float StartHealth = 100.f;
 	const float DamageValue = 5.f;
 
-	ASkillSystemTestPawn *SourceActor = World->SpawnActor<ASkillSystemTestPawn>();
-	ASkillSystemTestPawn *DestActor = World->SpawnActor<ASkillSystemTestPawn>();
+	AAbilitySystemTestPawn *SourceActor = World->SpawnActor<AAbilitySystemTestPawn>();
+	AAbilitySystemTestPawn *DestActor = World->SpawnActor<AAbilitySystemTestPawn>();
 
-	UProperty *DamageProperty = FindFieldChecked<UProperty>(USkillSystemTestAttributeSet::StaticClass(), GET_MEMBER_NAME_CHECKED(USkillSystemTestAttributeSet, Damage));
+	UProperty *DamageProperty = FindFieldChecked<UProperty>(UAbilitySystemTestAttributeSet::StaticClass(), GET_MEMBER_NAME_CHECKED(UAbilitySystemTestAttributeSet, Damage));
 
 	UAttributeComponent * SourceComponent = SourceActor->AttributeComponent;
 	UAttributeComponent * DestComponent = DestActor->AttributeComponent;
-	SourceComponent->GetSet<USkillSystemTestAttributeSet>()->Health = StartHealth;
-	DestComponent->GetSet<USkillSystemTestAttributeSet>()->Health = StartHealth;
+	SourceComponent->GetSet<UAbilitySystemTestAttributeSet>()->Health = StartHealth;
+	DestComponent->GetSet<UAbilitySystemTestAttributeSet>()->Health = StartHealth;
 
 
 	// attempt to prevent outgoing damage but fail because of a 0.f chance to apply
 	{
-		SKILL_LOG_SCOPE(TEXT("Fail to remove DamageBuff"))
+		ABILITY_LOG_SCOPE(TEXT("Fail to remove DamageBuff"))
 
 		UGameplayEffect * ModBuffEffect = Cast<UGameplayEffect>(StaticConstructObject(UGameplayEffect::StaticClass(), GetTransientPackage(), FName(TEXT("ModDamageBuff"))));
 		ModBuffEffect->GameplayEffectRequiredTags.AddTag(IGameplayTagsModule::RequestGameplayTag(FName(TEXT("Damage.Basic"))));
@@ -5586,7 +5586,7 @@ bool GameplayEffectsTest_ChanceToExecuteOnOutgoingGEImmunity(UWorld *World, FAut
 	}
 
 	{
-		SKILL_LOG_SCOPE(TEXT("Apply Damage mod"))
+		ABILITY_LOG_SCOPE(TEXT("Apply Damage mod"))
 
 		UGameplayEffect * BaseDmgEffect = Cast<UGameplayEffect>(StaticConstructObject(UGameplayEffect::StaticClass(), GetTransientPackage(), FName(TEXT("Damage"))));
 		BaseDmgEffect->Modifiers.SetNum(1);
@@ -5601,14 +5601,14 @@ bool GameplayEffectsTest_ChanceToExecuteOnOutgoingGEImmunity(UWorld *World, FAut
 		SourceComponent->ApplyGameplayEffectToTarget(BaseDmgEffect, DestComponent, 1.f);
 
 		float ExpectedValue = StartHealth - DamageValue;
-		float ActualValue = DestComponent->GetSet<USkillSystemTestAttributeSet>()->Health;
+		float ActualValue = DestComponent->GetSet<UAbilitySystemTestAttributeSet>()->Health;
 
 		Test->TestTrue(SKILL_TEST_TEXT("Damage Applied. %.2f == %.2f", ActualValue, ExpectedValue), ActualValue == ExpectedValue);
 	}
 
 	// prevent outgoing damage buff
 	{
-		SKILL_LOG_SCOPE(TEXT("Fail to remove DamageBuff"))
+		ABILITY_LOG_SCOPE(TEXT("Fail to remove DamageBuff"))
 
 		UGameplayEffect * ModBuffEffect = Cast<UGameplayEffect>(StaticConstructObject(UGameplayEffect::StaticClass(), GetTransientPackage(), FName(TEXT("ModDamageBuff"))));
 		ModBuffEffect->GameplayEffectRequiredTags.AddTag(IGameplayTagsModule::RequestGameplayTag(FName(TEXT("Damage.Basic"))));
@@ -5621,7 +5621,7 @@ bool GameplayEffectsTest_ChanceToExecuteOnOutgoingGEImmunity(UWorld *World, FAut
 	}
 
 	{
-		SKILL_LOG_SCOPE(TEXT("Apply Damage mod"))
+		ABILITY_LOG_SCOPE(TEXT("Apply Damage mod"))
 
 		UGameplayEffect * BaseDmgEffect = Cast<UGameplayEffect>(StaticConstructObject(UGameplayEffect::StaticClass(), GetTransientPackage(), FName(TEXT("Damage"))));
 		BaseDmgEffect->Modifiers.SetNum(1);
@@ -5636,7 +5636,7 @@ bool GameplayEffectsTest_ChanceToExecuteOnOutgoingGEImmunity(UWorld *World, FAut
 		SourceComponent->ApplyGameplayEffectToTarget(BaseDmgEffect, DestComponent, 1.f);
 
 		float ExpectedValue = StartHealth - DamageValue;
-		float ActualValue = DestComponent->GetSet<USkillSystemTestAttributeSet>()->Health;
+		float ActualValue = DestComponent->GetSet<UAbilitySystemTestAttributeSet>()->Health;
 
 		Test->TestTrue(SKILL_TEST_TEXT("Damage Applied. %.2f == %.2f", ActualValue, ExpectedValue), ActualValue == ExpectedValue);
 	}
@@ -5653,20 +5653,20 @@ bool GameplayEffectsTest_ChanceToExecuteOnIncomingGEMod(UWorld *World, FAutomati
 	const float DamageValue = 5.f;
 	const float BonusDamageMultiplier = 2.f;
 
-	ASkillSystemTestPawn *SourceActor = World->SpawnActor<ASkillSystemTestPawn>();
-	ASkillSystemTestPawn *DestActor = World->SpawnActor<ASkillSystemTestPawn>();
+	AAbilitySystemTestPawn *SourceActor = World->SpawnActor<AAbilitySystemTestPawn>();
+	AAbilitySystemTestPawn *DestActor = World->SpawnActor<AAbilitySystemTestPawn>();
 
-	UProperty *DamageProperty = FindFieldChecked<UProperty>(USkillSystemTestAttributeSet::StaticClass(), GET_MEMBER_NAME_CHECKED(USkillSystemTestAttributeSet, Damage));
+	UProperty *DamageProperty = FindFieldChecked<UProperty>(UAbilitySystemTestAttributeSet::StaticClass(), GET_MEMBER_NAME_CHECKED(UAbilitySystemTestAttributeSet, Damage));
 
 	UAttributeComponent * SourceComponent = SourceActor->AttributeComponent;
 	UAttributeComponent * DestComponent = DestActor->AttributeComponent;
-	SourceComponent->GetSet<USkillSystemTestAttributeSet>()->Health = StartHealth;
-	DestComponent->GetSet<USkillSystemTestAttributeSet>()->Health = StartHealth;
+	SourceComponent->GetSet<UAbilitySystemTestAttributeSet>()->Health = StartHealth;
+	DestComponent->GetSet<UAbilitySystemTestAttributeSet>()->Health = StartHealth;
 
 	// Setup a GE to modify OutgoingGEs
 	// this GE won't do anything because the chance to execute is zero
 	{
-		SKILL_LOG_SCOPE(TEXT("Apply useless DamageBuff"))
+		ABILITY_LOG_SCOPE(TEXT("Apply useless DamageBuff"))
 
 		UGameplayEffect * BuffDmgEffect = Cast<UGameplayEffect>(StaticConstructObject(UGameplayEffect::StaticClass(), GetTransientPackage(), FName(TEXT("DamageBuff"))));
 		BuffDmgEffect->Modifiers.SetNum(1);
@@ -5684,7 +5684,7 @@ bool GameplayEffectsTest_ChanceToExecuteOnIncomingGEMod(UWorld *World, FAutomati
 	}
 
 	{
-		SKILL_LOG_SCOPE(TEXT("Apply Damage mod"))
+		ABILITY_LOG_SCOPE(TEXT("Apply Damage mod"))
 
 		UGameplayEffect * BaseDmgEffect = Cast<UGameplayEffect>(StaticConstructObject(UGameplayEffect::StaticClass(), GetTransientPackage(), FName(TEXT("Damage"))));
 		BaseDmgEffect->Modifiers.SetNum(1);
@@ -5699,7 +5699,7 @@ bool GameplayEffectsTest_ChanceToExecuteOnIncomingGEMod(UWorld *World, FAutomati
 		SourceComponent->ApplyGameplayEffectToTarget(BaseDmgEffect, DestComponent, 1.f);
 
 		float ExpectedValue = StartHealth - DamageValue;
-		float ActualValue = DestComponent->GetSet<USkillSystemTestAttributeSet>()->Health;
+		float ActualValue = DestComponent->GetSet<UAbilitySystemTestAttributeSet>()->Health;
 
 		Test->TestTrue(SKILL_TEST_TEXT("Damage Applied. %.2f == %.2f", ActualValue, ExpectedValue), ActualValue == ExpectedValue);
 	}
@@ -5707,7 +5707,7 @@ bool GameplayEffectsTest_ChanceToExecuteOnIncomingGEMod(UWorld *World, FAutomati
 	// Setup a GE to modify OutgoingGEs
 	// this GE will always execute work because the chance to execute is one
 	{
-		SKILL_LOG_SCOPE(TEXT("Apply DamageBuff"))
+		ABILITY_LOG_SCOPE(TEXT("Apply DamageBuff"))
 
 		UGameplayEffect * BuffDmgEffect = Cast<UGameplayEffect>(StaticConstructObject(UGameplayEffect::StaticClass(), GetTransientPackage(), FName(TEXT("DamageBuff"))));
 		BuffDmgEffect->Modifiers.SetNum(1);
@@ -5725,7 +5725,7 @@ bool GameplayEffectsTest_ChanceToExecuteOnIncomingGEMod(UWorld *World, FAutomati
 	}
 
 	{
-		SKILL_LOG_SCOPE(TEXT("Apply Damage mod"))
+		ABILITY_LOG_SCOPE(TEXT("Apply Damage mod"))
 
 		UGameplayEffect * BaseDmgEffect = Cast<UGameplayEffect>(StaticConstructObject(UGameplayEffect::StaticClass(), GetTransientPackage(), FName(TEXT("Damage"))));
 		BaseDmgEffect->Modifiers.SetNum(1);
@@ -5740,7 +5740,7 @@ bool GameplayEffectsTest_ChanceToExecuteOnIncomingGEMod(UWorld *World, FAutomati
 		SourceComponent->ApplyGameplayEffectToTarget(BaseDmgEffect, DestComponent, 1.f);
 
 		float ExpectedValue = StartHealth - DamageValue - (DamageValue * BonusDamageMultiplier);
-		float ActualValue = DestComponent->GetSet<USkillSystemTestAttributeSet>()->Health;
+		float ActualValue = DestComponent->GetSet<UAbilitySystemTestAttributeSet>()->Health;
 
 		Test->TestTrue(SKILL_TEST_TEXT("Damage Applied. %.2f == %.2f", ActualValue, ExpectedValue), ActualValue == ExpectedValue);
 	}
@@ -5756,20 +5756,20 @@ bool GameplayEffectsTest_ChanceToExecuteOnIncomingGEImmunity(UWorld *World, FAut
 	const float StartHealth = 100.f;
 	const float DamageValue = 5.f;
 
-	ASkillSystemTestPawn *SourceActor = World->SpawnActor<ASkillSystemTestPawn>();
-	ASkillSystemTestPawn *DestActor = World->SpawnActor<ASkillSystemTestPawn>();
+	AAbilitySystemTestPawn *SourceActor = World->SpawnActor<AAbilitySystemTestPawn>();
+	AAbilitySystemTestPawn *DestActor = World->SpawnActor<AAbilitySystemTestPawn>();
 
-	UProperty *DamageProperty = FindFieldChecked<UProperty>(USkillSystemTestAttributeSet::StaticClass(), GET_MEMBER_NAME_CHECKED(USkillSystemTestAttributeSet, Damage));
+	UProperty *DamageProperty = FindFieldChecked<UProperty>(UAbilitySystemTestAttributeSet::StaticClass(), GET_MEMBER_NAME_CHECKED(UAbilitySystemTestAttributeSet, Damage));
 
 	UAttributeComponent * SourceComponent = SourceActor->AttributeComponent;
 	UAttributeComponent * DestComponent = DestActor->AttributeComponent;
-	SourceComponent->GetSet<USkillSystemTestAttributeSet>()->Health = StartHealth;
-	DestComponent->GetSet<USkillSystemTestAttributeSet>()->Health = StartHealth;
+	SourceComponent->GetSet<UAbilitySystemTestAttributeSet>()->Health = StartHealth;
+	DestComponent->GetSet<UAbilitySystemTestAttributeSet>()->Health = StartHealth;
 
 
 	// attempt to prevent outgoing damage but fail because of a 0.f chance to apply
 	{
-		SKILL_LOG_SCOPE(TEXT("Fail to remove DamageBuff"))
+		ABILITY_LOG_SCOPE(TEXT("Fail to remove DamageBuff"))
 
 		UGameplayEffect * ModBuffEffect = Cast<UGameplayEffect>(StaticConstructObject(UGameplayEffect::StaticClass(), GetTransientPackage(), FName(TEXT("ModDamageBuff"))));
 		ModBuffEffect->GameplayEffectRequiredTags.AddTag(IGameplayTagsModule::RequestGameplayTag(FName(TEXT("Damage.Basic"))));
@@ -5782,7 +5782,7 @@ bool GameplayEffectsTest_ChanceToExecuteOnIncomingGEImmunity(UWorld *World, FAut
 	}
 
 	{
-		SKILL_LOG_SCOPE(TEXT("Apply Damage mod"))
+		ABILITY_LOG_SCOPE(TEXT("Apply Damage mod"))
 
 		UGameplayEffect * BaseDmgEffect = Cast<UGameplayEffect>(StaticConstructObject(UGameplayEffect::StaticClass(), GetTransientPackage(), FName(TEXT("Damage"))));
 		BaseDmgEffect->Modifiers.SetNum(1);
@@ -5797,14 +5797,14 @@ bool GameplayEffectsTest_ChanceToExecuteOnIncomingGEImmunity(UWorld *World, FAut
 		SourceComponent->ApplyGameplayEffectToTarget(BaseDmgEffect, DestComponent, 1.f);
 
 		float ExpectedValue = StartHealth - DamageValue;
-		float ActualValue = DestComponent->GetSet<USkillSystemTestAttributeSet>()->Health;
+		float ActualValue = DestComponent->GetSet<UAbilitySystemTestAttributeSet>()->Health;
 
 		Test->TestTrue(SKILL_TEST_TEXT("Damage Applied. %.2f == %.2f", ActualValue, ExpectedValue), ActualValue == ExpectedValue);
 	}
 
 	// prevent outgoing damage buff
 	{
-		SKILL_LOG_SCOPE(TEXT("Fail to remove DamageBuff"))
+		ABILITY_LOG_SCOPE(TEXT("Fail to remove DamageBuff"))
 
 		UGameplayEffect * ModBuffEffect = Cast<UGameplayEffect>(StaticConstructObject(UGameplayEffect::StaticClass(), GetTransientPackage(), FName(TEXT("ModDamageBuff"))));
 		ModBuffEffect->GameplayEffectRequiredTags.AddTag(IGameplayTagsModule::RequestGameplayTag(FName(TEXT("Damage.Basic"))));
@@ -5817,7 +5817,7 @@ bool GameplayEffectsTest_ChanceToExecuteOnIncomingGEImmunity(UWorld *World, FAut
 	}
 
 	{
-		SKILL_LOG_SCOPE(TEXT("Apply Damage mod"))
+		ABILITY_LOG_SCOPE(TEXT("Apply Damage mod"))
 
 		UGameplayEffect * BaseDmgEffect = Cast<UGameplayEffect>(StaticConstructObject(UGameplayEffect::StaticClass(), GetTransientPackage(), FName(TEXT("Damage"))));
 		BaseDmgEffect->Modifiers.SetNum(1);
@@ -5832,7 +5832,7 @@ bool GameplayEffectsTest_ChanceToExecuteOnIncomingGEImmunity(UWorld *World, FAut
 		SourceComponent->ApplyGameplayEffectToTarget(BaseDmgEffect, DestComponent, 1.f);
 
 		float ExpectedValue = StartHealth - DamageValue;
-		float ActualValue = DestComponent->GetSet<USkillSystemTestAttributeSet>()->Health;
+		float ActualValue = DestComponent->GetSet<UAbilitySystemTestAttributeSet>()->Health;
 
 		Test->TestTrue(SKILL_TEST_TEXT("Damage Applied. %.2f == %.2f", ActualValue, ExpectedValue), ActualValue == ExpectedValue);
 	}
@@ -5848,18 +5848,18 @@ bool GameplayEffectsTest_ModifyChanceToApplyToTarget(UWorld *World, FAutomationT
 	const float StartHealth = 100.f;
 	const float DamageValue = 5.f;
 
-	ASkillSystemTestPawn *SourceActor = World->SpawnActor<ASkillSystemTestPawn>();
-	ASkillSystemTestPawn *DestActor = World->SpawnActor<ASkillSystemTestPawn>();
+	AAbilitySystemTestPawn *SourceActor = World->SpawnActor<AAbilitySystemTestPawn>();
+	AAbilitySystemTestPawn *DestActor = World->SpawnActor<AAbilitySystemTestPawn>();
 
-	UProperty *DamageProperty = FindFieldChecked<UProperty>(USkillSystemTestAttributeSet::StaticClass(), GET_MEMBER_NAME_CHECKED(USkillSystemTestAttributeSet, Damage));
+	UProperty *DamageProperty = FindFieldChecked<UProperty>(UAbilitySystemTestAttributeSet::StaticClass(), GET_MEMBER_NAME_CHECKED(UAbilitySystemTestAttributeSet, Damage));
 
 	UAttributeComponent * SourceComponent = SourceActor->AttributeComponent;
 	UAttributeComponent * DestComponent = DestActor->AttributeComponent;
-	SourceComponent->GetSet<USkillSystemTestAttributeSet>()->Health = StartHealth;
-	DestComponent->GetSet<USkillSystemTestAttributeSet>()->Health = StartHealth;
+	SourceComponent->GetSet<UAbilitySystemTestAttributeSet>()->Health = StartHealth;
+	DestComponent->GetSet<UAbilitySystemTestAttributeSet>()->Health = StartHealth;
 
 	{
-		SKILL_LOG_SCOPE(TEXT("Apply boost to chance to apply"))
+		ABILITY_LOG_SCOPE(TEXT("Apply boost to chance to apply"))
 
 		UGameplayEffect * BaseEffect = Cast<UGameplayEffect>(StaticConstructObject(UGameplayEffect::StaticClass(), GetTransientPackage(), FName(TEXT("ChanceToApply"))));
 		BaseEffect->Modifiers.SetNum(1);
@@ -5874,7 +5874,7 @@ bool GameplayEffectsTest_ModifyChanceToApplyToTarget(UWorld *World, FAutomationT
 	}
 
 	{
-		SKILL_LOG_SCOPE(TEXT("Apply Damage mod"))
+		ABILITY_LOG_SCOPE(TEXT("Apply Damage mod"))
 
 		UGameplayEffect * BaseDmgEffect = Cast<UGameplayEffect>(StaticConstructObject(UGameplayEffect::StaticClass(), GetTransientPackage(), FName(TEXT("Damage"))));
 		BaseDmgEffect->Modifiers.SetNum(1);
@@ -5890,7 +5890,7 @@ bool GameplayEffectsTest_ModifyChanceToApplyToTarget(UWorld *World, FAutomationT
 		SourceComponent->ApplyGameplayEffectToTarget(BaseDmgEffect, DestComponent, 0.f);
 
 		float ExpectedValue = (StartHealth - DamageValue);
-		float ActualValue = DestComponent->GetSet<USkillSystemTestAttributeSet>()->Health;
+		float ActualValue = DestComponent->GetSet<UAbilitySystemTestAttributeSet>()->Health;
 
 		Test->TestTrue(SKILL_TEST_TEXT("Damage Applied. %.2f == %.2f", ActualValue, ExpectedValue), ActualValue == ExpectedValue);
 	}
@@ -5906,18 +5906,18 @@ bool GameplayEffectsTest_ModifyChanceToExecuteOnGE(UWorld *World, FAutomationTes
 	const float StartHealth = 100.f;
 	const float DamageValue = 5.f;
 
-	ASkillSystemTestPawn *SourceActor = World->SpawnActor<ASkillSystemTestPawn>();
-	ASkillSystemTestPawn *DestActor = World->SpawnActor<ASkillSystemTestPawn>();
+	AAbilitySystemTestPawn *SourceActor = World->SpawnActor<AAbilitySystemTestPawn>();
+	AAbilitySystemTestPawn *DestActor = World->SpawnActor<AAbilitySystemTestPawn>();
 
-	UProperty *DamageProperty = FindFieldChecked<UProperty>(USkillSystemTestAttributeSet::StaticClass(), GET_MEMBER_NAME_CHECKED(USkillSystemTestAttributeSet, Damage));
+	UProperty *DamageProperty = FindFieldChecked<UProperty>(UAbilitySystemTestAttributeSet::StaticClass(), GET_MEMBER_NAME_CHECKED(UAbilitySystemTestAttributeSet, Damage));
 
 	UAttributeComponent * SourceComponent = SourceActor->AttributeComponent;
 	UAttributeComponent * DestComponent = DestActor->AttributeComponent;
-	SourceComponent->GetSet<USkillSystemTestAttributeSet>()->Health = StartHealth;
-	DestComponent->GetSet<USkillSystemTestAttributeSet>()->Health = StartHealth;
+	SourceComponent->GetSet<UAbilitySystemTestAttributeSet>()->Health = StartHealth;
+	DestComponent->GetSet<UAbilitySystemTestAttributeSet>()->Health = StartHealth;
 
 	{
-		SKILL_LOG_SCOPE(TEXT("Apply Damage buff"))
+		ABILITY_LOG_SCOPE(TEXT("Apply Damage buff"))
 
 		UGameplayEffect * BuffEffect = Cast<UGameplayEffect>(StaticConstructObject(UGameplayEffect::StaticClass(), GetTransientPackage(), FName(TEXT("DamageBuff"))));
 		BuffEffect->Modifiers.SetNum(1);
@@ -5937,7 +5937,7 @@ bool GameplayEffectsTest_ModifyChanceToExecuteOnGE(UWorld *World, FAutomationTes
 
 	// verify that outgoing damage is unbuffed
 	{
-		SKILL_LOG_SCOPE(TEXT("Apply Damage mod"))
+		ABILITY_LOG_SCOPE(TEXT("Apply Damage mod"))
 
 		UGameplayEffect * BaseDmgEffect = Cast<UGameplayEffect>(StaticConstructObject(UGameplayEffect::StaticClass(), GetTransientPackage(), FName(TEXT("Damage"))));
 		BaseDmgEffect->Modifiers.SetNum(1);
@@ -5952,13 +5952,13 @@ bool GameplayEffectsTest_ModifyChanceToExecuteOnGE(UWorld *World, FAutomationTes
 		SourceComponent->ApplyGameplayEffectToTarget(BaseDmgEffect, DestComponent, 0.f);
 
 		float ExpectedValue = (StartHealth - DamageValue);
-		float ActualValue = DestComponent->GetSet<USkillSystemTestAttributeSet>()->Health;
+		float ActualValue = DestComponent->GetSet<UAbilitySystemTestAttributeSet>()->Health;
 
 		Test->TestTrue(SKILL_TEST_TEXT("Damage Applied. %.2f == %.2f", ActualValue, ExpectedValue), ActualValue == ExpectedValue);
 	}
 
 	{
-		SKILL_LOG_SCOPE(TEXT("Apply boost to chance to execute"))
+		ABILITY_LOG_SCOPE(TEXT("Apply boost to chance to execute"))
 
 		UGameplayEffect * BaseEffect = Cast<UGameplayEffect>(StaticConstructObject(UGameplayEffect::StaticClass(), GetTransientPackage(), FName(TEXT("ChanceToExecute"))));
 		BaseEffect->Modifiers.SetNum(1);
@@ -5974,7 +5974,7 @@ bool GameplayEffectsTest_ModifyChanceToExecuteOnGE(UWorld *World, FAutomationTes
 
 	// verify that the buff applies now
 	{
-		SKILL_LOG_SCOPE(TEXT("Apply Damage mod"))
+		ABILITY_LOG_SCOPE(TEXT("Apply Damage mod"))
 
 		UGameplayEffect * BaseDmgEffect = Cast<UGameplayEffect>(StaticConstructObject(UGameplayEffect::StaticClass(), GetTransientPackage(), FName(TEXT("Damage"))));
 		BaseDmgEffect->Modifiers.SetNum(1);
@@ -5989,7 +5989,7 @@ bool GameplayEffectsTest_ModifyChanceToExecuteOnGE(UWorld *World, FAutomationTes
 		SourceComponent->ApplyGameplayEffectToTarget(BaseDmgEffect, DestComponent, 0.f);
 
 		float ExpectedValue = (StartHealth - DamageValue * 3.f);
-		float ActualValue = DestComponent->GetSet<USkillSystemTestAttributeSet>()->Health;
+		float ActualValue = DestComponent->GetSet<UAbilitySystemTestAttributeSet>()->Health;
 
 		Test->TestTrue(SKILL_TEST_TEXT("Damage Applied. %.2f == %.2f", ActualValue, ExpectedValue), ActualValue == ExpectedValue);
 	}
@@ -6004,8 +6004,8 @@ bool GameplayEffectsTest_ModifyChanceToExecuteOnGE(UWorld *World, FAutomationTes
 
 bool FGameplayEffectsTest::RunTest( const FString& Parameters )
 {
-	UCurveTable *CurveTable = IGameplayAbilitiesModule::Get().GetSkillSystemGlobals()->GetGlobalCurveTable();
-	UDataTable *DataTable = IGameplayAbilitiesModule::Get().GetSkillSystemGlobals()->GetGlobalAttributeDataTable();
+	UCurveTable *CurveTable = IGameplayAbilitiesModule::Get().GetAbilitySystemGlobals()->GetGlobalCurveTable();
+	UDataTable *DataTable = IGameplayAbilitiesModule::Get().GetAbilitySystemGlobals()->GetGlobalAttributeDataTable();
 
 	UWorld *World = UWorld::CreateWorld(EWorldType::Game, false);
 	FWorldContext &WorldContext = GEngine->CreateNewWorldContext(EWorldType::Game);
@@ -6132,8 +6132,8 @@ bool FGameplayEffectsTest::RunTest( const FString& Parameters )
 	GEngine->DestroyWorldContext(World);
 	World->DestroyWorld(false);
 
-	IGameplayAbilitiesModule::Get().GetSkillSystemGlobals()->AutomationTestOnly_SetGlobalCurveTable(CurveTable);
-	IGameplayAbilitiesModule::Get().GetSkillSystemGlobals()->AutomationTestOnly_SetGlobalAttributeDataTable(DataTable);
+	IGameplayAbilitiesModule::Get().GetAbilitySystemGlobals()->AutomationTestOnly_SetGlobalCurveTable(CurveTable);
+	IGameplayAbilitiesModule::Get().GetAbilitySystemGlobals()->AutomationTestOnly_SetGlobalAttributeDataTable(DataTable);
 	return true;
 }
 
