Mass rename SkillSystem -> AbilitySystem

diff --git a/Engine/Source/Runtime/GameplayAbilities/Public/AbilitySystemLog.h b/Engine/Source/Runtime/GameplayAbilities/Public/AbilitySystemLog.h
new file mode 100644
--- /dev/null
+++ b/Engine/Source/Runtime/GameplayAbilities/Public/AbilitySystemLog.h
@@ -0,0 +1,66 @@
+// Copyright 1998-2014 Epic Games, Inc. All Rights Reserved.
+
+#pragma once
+
+// Intended categories:
+//	Log - This happened. What gameplay programmers may care about to debug
+//	Verbose - This is why this happened. What you may turn on to debug the skill system code.
+//  VeryVerbose - This is what didn't happen, and why. Extreme printf debugging
+//
+
+
+GAMEPLAYABILITIES_API DECLARE_LOG_CATEGORY_EXTERN(LogAbilitySystem, Warning, All);
+
+#define ABILITY_LOG(Verbosity, Format, ...) \
+{ \
+	FString Str = AbilitySystemLog::Log(ELogVerbosity::Verbosity, FString::Printf(Format, ##__VA_ARGS__)); \
+	UE_LOG(LogAbilitySystem, Verbosity, TEXT("%s"), *Str); \
+}
+
+#define ABILITY_LOG_TOKENPASTE_INNER(x,y) x##y
+#define ABILITY_LOG_TOKENPASTE(x,y) ABILITY_LOG_TOKENPASTE_INNER(x,y)
+#define ABILITY_LOG_SCOPE( Format, ... ) AbilitySystemLogScope ABILITY_LOG_TOKENPASTE(LogScope,__LINE__)( FString::Printf(Format, ##__VA_ARGS__));
+
+struct AbilitySystemLogScope
+{
+	AbilitySystemLogScope()
+	{
+		Init();
+	}
+
+	AbilitySystemLogScope(FString InScopeName) :
+		ScopeName(InScopeName)
+	{
+		Init();
+	}
+
+	~AbilitySystemLogScope();
+	
+	FString ScopeName;
+	bool PrintedInThisScope;
+
+private:
+
+	void Init();
+};
+
+class AbilitySystemLog
+{
+public:
+
+	AbilitySystemLog()
+	{
+		Indent = 0;
+		NeedNewLine = false;
+	}
+
+	static FString Log(ELogVerbosity::Type, FString Str);
+
+	static void PushScope(AbilitySystemLogScope* Scope);
+
+	static void PopScope();	
+
+	int32	Indent;
+	TArray<AbilitySystemLogScope*>	ScopeStack;
+	bool NeedNewLine;
+};
\ No newline at end of file
diff --git a/Engine/Source/Runtime/GameplayAbilities/Public/AbilitySystemStats.h b/Engine/Source/Runtime/GameplayAbilities/Public/AbilitySystemStats.h
new file mode 100644
--- /dev/null
+++ b/Engine/Source/Runtime/GameplayAbilities/Public/AbilitySystemStats.h
@@ -0,0 +1,17 @@
+// Copyright 1998-2013 Epic Games, Inc. All Rights Reserved.
+
+#pragma once
+
+//#include "GameplayTagAssetInterface.h"
+//#include "GameplayTagsManager.h"
+//#include "GameplayTagContainer.h"
+#include "AbilitySystemLog.h"
+
+DECLARE_STATS_GROUP(TEXT("AbilitySystem"), STATGROUP_AbilitySystem, STATCAT_Advanced);
+
+DECLARE_CYCLE_STAT_EXTERN(TEXT("GameplayEffectsHasAllTagsTime"), STAT_GameplayEffectsHasAllTags, STATGROUP_AbilitySystem, );
+DECLARE_CYCLE_STAT_EXTERN(TEXT("GameplayEffectsHasAnyTagTime"), STAT_GameplayEffectsHasAnyTag, STATGROUP_AbilitySystem, );
+DECLARE_CYCLE_STAT_EXTERN(TEXT("GameplayEffectsTick"), STAT_GameplayEffectsTick, STATGROUP_AbilitySystem, );
+DECLARE_CYCLE_STAT_EXTERN(TEXT("CanApplyAttributeModifiers"), STAT_GameplayEffectsCanApplyAttributeModifiers, STATGROUP_AbilitySystem, );
+DECLARE_CYCLE_STAT_EXTERN(TEXT("GetActiveEffectsData"), STAT_GameplayEffectsGetActiveEffectsData, STATGROUP_AbilitySystem, );
+DECLARE_CYCLE_STAT_EXTERN(TEXT("GetCooldownTimeRemaining"), STAT_GameplayAbilityGetCooldownTimeRemaining, STATGROUP_AbilitySystem, );
diff --git a/Engine/Source/Runtime/GameplayAbilities/Public/AbilitySystemTypes.h b/Engine/Source/Runtime/GameplayAbilities/Public/AbilitySystemTypes.h
new file mode 100644
--- /dev/null
+++ b/Engine/Source/Runtime/GameplayAbilities/Public/AbilitySystemTypes.h
@@ -0,0 +1,218 @@
+﻿// Copyright 1998-2014 Epic Games, Inc. All Rights Reserved.
+
+#pragma once
+
+#include "AbilitySystemTypes.generated.h"
+
+class UAttributeComponent;
+class UGameplayEffect;
+
+UENUM(BlueprintType)
+namespace EGameplayAbilityActivationMode
+{
+	enum Type
+	{
+		Authority,			// We are the authority activating this ability
+		NonAuthority,		// We are not the authority but aren't predicting yet. This is a mostly invalid state to be in.
+
+		Predicting,			// We are predicting the activation of this ability
+		Confirmed,			// We are not the authority, but the authority has confirmed this activation
+	};
+}
+
+/**
+ *	FGameplayAbilityActorInfo 
+ *	
+ *	Cached data associated with an Actor using an Ability.
+ *		-Initialized from an AActor* in InitFromActor
+ *		-Abilities use this to know what to actor upon. E.g., instead of being coupled to a specific actor class.
+ *		-These are generally passed around as pointers to support polymorphism.
+ *		-Projects can override UAbilitySystemGlobals::AllocAbilityActorInfo to override the default struct type that is created. 
+ *			
+ */
+USTRUCT(BlueprintType)
+struct GAMEPLAYABILITIES_API FGameplayAbilityActorInfo
+{
+	GENERATED_USTRUCT_BODY()
+
+	virtual ~FGameplayAbilityActorInfo() {}
+
+	UPROPERTY(BlueprintReadOnly, Category = "ActorInfo")
+	TWeakObjectPtr<AActor>	Actor;
+
+	UPROPERTY(BlueprintReadOnly, Category = "ActorInfo")
+	TWeakObjectPtr<UAnimInstance>	AnimInstance;
+
+	UPROPERTY(BlueprintReadOnly, Category = "ActorInfo")
+	TWeakObjectPtr<UAttributeComponent>	AttributeComponent;
+
+	UPROPERTY(BlueprintReadOnly, Category = "ActorInfo")
+	TWeakObjectPtr<UMovementComponent>	MovementComponent;
+
+	virtual void InitFromActor(AActor *Actor);
+};
+
+/**
+ *	FGameplayAbilityActivationInfo
+ *	
+ *	Data tied to a specific activation of an ability. 
+ *		-Tell us whether we are the authority, if we are predicting, confirmed, etc.
+ *		-Holds PredictionKey
+ *		-Generally not meant to be subclassed in projects.
+ *		-Passed around by value since the struct is small.
+ */
+
+USTRUCT(BlueprintType)
+struct GAMEPLAYABILITIES_API FGameplayAbilityActivationInfo
+{
+	GENERATED_USTRUCT_BODY()
+
+	FGameplayAbilityActivationInfo()
+		: ActivationMode(EGameplayAbilityActivationMode::Authority),
+		PredictionKey(0)
+	{
+
+	}
+
+	FGameplayAbilityActivationInfo(AActor * InActor, int32 InPredictionKey)
+		: PredictionKey(InPredictionKey)
+	{
+		// On Init, we are either Authority or NonAuthority. We haven't been given a PredictionKey and we haven't been confirmed.
+		// NonAuthority essentially means 'I'm not sure what how I'm going to do this yet'.
+		ActivationMode = (InActor->Role == ROLE_Authority ? EGameplayAbilityActivationMode::Authority : EGameplayAbilityActivationMode::NonAuthority);
+	}
+
+	FGameplayAbilityActivationInfo(EGameplayAbilityActivationMode::Type InType, int32 InPredictionKey)
+		: ActivationMode(InType)
+		, PredictionKey(InPredictionKey)
+	{
+	}
+
+
+	void GeneratePredictionKey(UAttributeComponent * Component) const;
+
+	void SetPredictionKey(int32 InPredictionKey);
+
+	void SetActivationConfirmed();
+
+	UPROPERTY(BlueprintReadOnly, Category = "ActorInfo")
+	mutable TEnumAsByte<EGameplayAbilityActivationMode::Type>	ActivationMode;
+
+	UPROPERTY()
+	mutable int32 PredictionKey;
+};
+
+
+// ----------------------------------------------------
+
+/**
+ *	A generic structure for targeting data. We want generic functions to produce this data and other generic
+ *	functions to consume this data.
+ *	
+ *	We expect this to be able to hold specific actors/object reference and also generic location/direction/origin 
+ *	information.
+ *	
+ *	Some example producers:
+ *		-Overlap/Hit collision event generates TargetData about who was hit in a melee attack
+ *		-A mouse input causes a hit trace and the actor infront of the crosshair is turned into TargetData
+ *		-A mouse input causes TargetData to be generated from the owner's crosshair view origin/direction
+ *		-An AOE/aura pulses and all actors in a radius around the instigator are added to TargetData
+ *		-Panzer Dragoon style 'painting' targeting mode
+ *		-MMORPG style ground AOE targeting style (potentially both a location on the ground and actors that were targeted)
+ *		
+ *	Some example consumers:
+ *		-Apply a GameplayEffect to all actors in TargetData
+ *		-Find closest actor from all in TargetData
+ *		-Call some function on all actors in TargetData
+ *		-Filter or merge TargetDatas
+ *		-Spawn a new actor at a TargetData location
+ *		
+ *		
+ *		
+ *	Maybe it is better to distinguish between actor list targeting vs positional targeting data?
+ *		-AOE/aura type of targeting data blurs the line
+ *		
+ *		
+ */
+
+struct GAMEPLAYABILITIES_API FGameplayAbilityTargetData
+{
+	virtual ~FGameplayAbilityTargetData() { }
+
+	void ApplyGameplayEffect(UGameplayEffect *GameplayEffect, const FGameplayAbilityActorInfo InstigatorInfo);
+
+	virtual TArray<AActor*>	GetActors() const = 0;
+	
+	virtual bool HasHitResult() const 
+	{	
+		return false;
+	}
+
+	virtual const FHitResult * GetHitResult() const
+	{
+		return NULL;
+	}
+
+	virtual FString ToString() const;
+};
+
+struct GAMEPLAYABILITIES_API FGameplayAbilityTargetData_SingleTargetHit : public FGameplayAbilityTargetData
+{
+	FGameplayAbilityTargetData_SingleTargetHit(const FHitResult InHitResult)
+	: HitResult(InHitResult)
+	{ }
+
+	virtual TArray<AActor*>	GetActors() const
+	{
+		TArray<AActor*>	Actors;
+		if (HitResult.Actor.IsValid())
+		{
+			Actors.Push(HitResult.Actor.Get());
+		}
+		return Actors;
+	}
+
+	virtual bool HasHitResult() const
+	{
+		return true;
+	}
+
+	virtual const FHitResult * GetHitResult() const
+	{
+		return &HitResult;
+	}
+
+	FHitResult	HitResult;
+};
+
+/**
+ *	Handle for Targeting Data. This servers two main purposes:
+ *		-Avoid us having to copy around the full targeting data structure in Blueprints
+ *		-Allows us to leverage polymorphism in the target data structure
+ *		-Allows us to implement NetSerialize and replicate by value between clients/server	
+ *		
+ *		-Avoid using UObjects could be used to give us polymorphism and by reference passing in blueprints.
+ *		-However we would still be screwed when it came to replication
+ *		
+ *		-Replication by value
+ *		-Pass by reference in blueprints
+ *		-Polymophism in TargetData structure
+ *		
+ */
+
+USTRUCT(BlueprintType)
+struct GAMEPLAYABILITIES_API FGameplayAbilityTargetDataHandle
+{
+	GENERATED_USTRUCT_BODY()
+	
+	TSharedPtr<FGameplayAbilityTargetData>	Data;
+};
+
+template<>
+struct TStructOpsTypeTraits<FGameplayAbilityTargetDataHandle> : public TStructOpsTypeTraitsBase
+{
+	enum
+	{
+		WithCopy = true		// Necessary so that TSharedPtr<FGameplayAbilityTargetData> Data is copied around
+	};
+};
diff --git a/Engine/Source/Runtime/GameplayAbilities/Public/GameplayAbilitiesModule.h b/Engine/Source/Runtime/GameplayAbilities/Public/GameplayAbilitiesModule.h
--- a/Engine/Source/Runtime/GameplayAbilities/Public/GameplayAbilitiesModule.h
+++ b/Engine/Source/Runtime/GameplayAbilities/Public/GameplayAbilitiesModule.h
@@ -3,7 +3,7 @@
 #pragma once
 
 #include "ModuleManager.h"
-#include "SkillSystemGlobals.h"
+#include "AbilitySystemGlobals.h"
 
 
 /**
@@ -35,6 +35,6 @@ public:
 		return FModuleManager::Get().IsModuleLoaded( "GameplayAbilities" );
 	}
 
-	virtual USkillSystemGlobals* GetSkillSystemGlobals() = 0;
+	virtual UAbilitySystemGlobals* GetAbilitySystemGlobals() = 0;
 };
 
diff --git a/Engine/Source/Runtime/Launch/Private/LaunchEngineLoop.cpp b/Engine/Source/Runtime/Launch/Private/LaunchEngineLoop.cpp
--- a/Engine/Source/Runtime/Launch/Private/LaunchEngineLoop.cpp
+++ b/Engine/Source/Runtime/Launch/Private/LaunchEngineLoop.cpp
@@ -1704,13 +1704,13 @@ bool FEngineLoop::LoadStartupCoreModules()
 
 	// -----------------------------------------------------
 
-	// HACK: load SkillSystem editor as early as possible for statically initialized assets (non cooked BT assets needs it)
+	// HACK: load AbilitySystem editor as early as possible for statically initialized assets (non cooked BT assets needs it)
 	// cooking needs this module too
-	bool bSkillSystemEditorEnabled = false;
-	GConfig->GetBool(TEXT("SkillSystemEd"), TEXT("SKillSystemEditorEnabled"), bSkillSystemEditorEnabled, GEngineIni);
-	if (bSkillSystemEditorEnabled)
+	bool bGameplayAbilitiesEnabled = false;
+	GConfig->GetBool(TEXT("GameplayAbilities"), TEXT("GameplayAbilitiedEditorEnabled"), bGameplayAbilitiesEnabled, GEngineIni);
+	if (bGameplayAbilitiesEnabled)
 	{
-		FModuleManager::Get().LoadModule(TEXT("SkillSystemEditor"));
+		FModuleManager::Get().LoadModule(TEXT("GameplayAbilitiesEditor"));
 	}
 
 	// -----------------------------------------------------
