Rename AttributeComponent -> AbilitySystemComponent

diff --git a/Engine/Source/Runtime/GameplayAbilities/Classes/Abilities/GameplayAbility.h b/Engine/Source/Runtime/GameplayAbilities/Classes/Abilities/GameplayAbility.h
--- a/Engine/Source/Runtime/GameplayAbilities/Classes/Abilities/GameplayAbility.h
+++ b/Engine/Source/Runtime/GameplayAbilities/Classes/Abilities/GameplayAbility.h
@@ -9,7 +9,7 @@
 #include "GameplayAbility.generated.h"
 
 class UAnimInstance;
-class UAttributeComponent;
+class UAbilitySystemComponent;
 
 UENUM(BlueprintType)
 namespace EGameplayAbilityInstancingPolicy
@@ -316,5 +316,5 @@ protected:
 
 	void	ApplyCost(const FGameplayAbilityActorInfo* ActorInfo, const FGameplayAbilityActivationInfo ActivationInfo);
 
-	friend class UAttributeComponent;
+	friend class UAbilitySystemComponent;
 };
\ No newline at end of file
diff --git a/Engine/Source/Runtime/GameplayAbilities/Classes/Abilities/GameplayAbility_Instanced.h b/Engine/Source/Runtime/GameplayAbilities/Classes/Abilities/GameplayAbility_Instanced.h
--- a/Engine/Source/Runtime/GameplayAbilities/Classes/Abilities/GameplayAbility_Instanced.h
+++ b/Engine/Source/Runtime/GameplayAbilities/Classes/Abilities/GameplayAbility_Instanced.h
@@ -5,7 +5,7 @@
 #include "GameplayAbility_Instanced.generated.h"
 
 class UAnimInstance;
-class UAttributeComponent;
+class UAbilitySystemComponent;
 
 /**
 * UGameplayEffect
diff --git a/Engine/Source/Runtime/GameplayAbilities/Classes/AbilitySystemBlueprintLibrary.h b/Engine/Source/Runtime/GameplayAbilities/Classes/AbilitySystemBlueprintLibrary.h
--- a/Engine/Source/Runtime/GameplayAbilities/Classes/AbilitySystemBlueprintLibrary.h
+++ b/Engine/Source/Runtime/GameplayAbilities/Classes/AbilitySystemBlueprintLibrary.h
@@ -17,7 +17,7 @@ class UAbilitySystemBlueprintLibrary : public UBlueprintFunctionLibrary
 	GENERATED_UCLASS_BODY()
 
 	UFUNCTION(BlueprintPure, Category = Ability)
-	static UAttributeComponent* GetAttributeComponent(AActor *Actor);
+	static UAbilitySystemComponent* GetAbilitySystemComponent(AActor *Actor);
 
 	UFUNCTION(BlueprintCallable, Category=Ability)
 	static void ApplyGameplayEffectToTargetData(FGameplayAbilityTargetDataHandle Target, UGameplayEffect *GameplayEffect, const FGameplayAbilityActorInfo InstigatorInfo);
diff --git a/Engine/Source/Runtime/GameplayAbilities/Classes/AbilitySystemComponent.h b/Engine/Source/Runtime/GameplayAbilities/Classes/AbilitySystemComponent.h
new file mode 100644
--- /dev/null
+++ b/Engine/Source/Runtime/GameplayAbilities/Classes/AbilitySystemComponent.h
@@ -0,0 +1,368 @@
+// Copyright 1998-2013 Epic Games, Inc. All Rights Reserved.
+#pragma once
+
+#include "GameplayEffect.h"
+#include "AbilitySystemTypes.h"
+#include "AbilitySystemComponent.generated.h"
+
+GAMEPLAYABILITIES_API DECLARE_LOG_CATEGORY_EXTERN(LogAbilitySystemComponent, Log, All);
+
+DECLARE_MULTICAST_DELEGATE(FAbilitySystemComponentPredictionKeyClear);
+
+USTRUCT()
+struct GAMEPLAYABILITIES_API FAttributeDefaults
+{
+	GENERATED_USTRUCT_BODY()
+
+	UPROPERTY(EditAnywhere, Category="AttributeTest")
+	TSubclassOf<class UAttributeSet> Attributes;
+
+	UPROPERTY(EditAnywhere, BlueprintReadWrite, Category="AttributeTest")
+	class UDataTable*	DefaultStartingTable;
+};
+
+/** 
+ * 
+ */
+UCLASS(ClassGroup=AbilitySystem, hidecategories=(Object,LOD,Lighting,Transform,Sockets,TextureStreaming), editinlinenew, meta=(BlueprintSpawnableComponent))
+class GAMEPLAYABILITIES_API UAbilitySystemComponent : public UActorComponent
+{
+	GENERATED_UCLASS_BODY()
+
+	friend FGameplayEffectSpec;
+	friend class AAbilitySystemDebugHUD;
+
+	virtual ~UAbilitySystemComponent();
+
+	virtual void InitializeComponent() OVERRIDE;
+
+	/** Finds existing AttributeSet */
+	template <class T >
+	T*	GetSet()
+	{
+		return (T*)GetAttributeSubobject(T::StaticClass());
+	}
+
+	template <class T >
+	T*	GetSetChecked()
+	{
+		return (T*)GetAttributeSubobjectChecked(T::StaticClass());
+	}
+
+	/** Adds a new AttributeSet (initialized to default values) */
+	template <class T >
+	T*  AddSet()
+	{
+		return (T*)GetOrCreateAttributeSubobject(T::StaticClass());
+	}
+
+	/** Adds a new AttributeSet that is a DSO (created by called in their CStor) */
+	template <class T>
+	T*	AddDefaultSubobjectSet(TSubobjectPtr<T> Subobject)
+	{
+		T* Set = Subobject.Get();
+		SpawnedAttributes.Add(Set);
+		return Set;
+	}
+
+	UFUNCTION(BlueprintCallable, Category="Skills")
+	UAttributeSet * InitStats(TSubclassOf<class UAttributeSet> Attributes, const UDataTable* DataTable);
+
+	UFUNCTION(BlueprintCallable, Category="Skills")
+	void ModifyStats(TSubclassOf<class UAttributeSet> Attributes, FDataTableRowHandle ModifierHandle);
+
+	UPROPERTY(EditAnywhere, Category="AttributeTest")
+	TArray<FAttributeDefaults>	DefaultStartingData;
+
+	UPROPERTY(Replicated)
+	TArray<UAttributeSet*>	SpawnedAttributes;
+
+	void SetNumericAttribute(const FGameplayAttribute &Attribute, float NewFloatValue);
+	float GetNumericAttribute(const FGameplayAttribute &Attribute);
+
+	// -- Replication -------------------------------------------------------------------------------------------------
+
+	virtual bool ReplicateSubobjects(class UActorChannel *Channel, class FOutBunch *Bunch, FReplicationFlags *RepFlags) OVERRIDE;
+	
+	virtual void GetSubobjectsWithStableNamesForNetworking(TArray<UObject*>& Objs) OVERRIDE;
+
+	virtual void PostNetReceive() OVERRIDE;
+
+	/**
+	 *	Prediction Keys
+	 *		PredictionKey is a simple system to allow clients to locally predict interactions within the AbilityComponent while having
+	 *		a mechanism for merging results from the server.
+	 *		
+	 *		There are two main cases:
+	 *			-Client predicts something and is wrong. (Mispredictions). He must cleanup stuff that he created or changed.
+	 *			-Client predicts something and is right. 
+	 *			
+	 *			The Client predicts something and asks the server. "Can I do X? PredictionKey=Y".
+	 *			Server can reply immediately and directly with yes or no, and include the PredictionKey with that answer.
+	 *			
+	 *			If No, the client can just immediately cleanup whatever he predicted. This is straightforward.
+	 *			
+	 *			If Yes, there is a small window where the Yes RPC is received but the actor/component properties have not been updated.
+	 *			The client's cannot immediately cleanup his prediction work: he needs to wait for property replication.
+	 *			
+	 *			This is where ReplicatedPredictionKey comes in. 
+	 *			
+	 *			When the server allows a predicted action, he sets ReplicatedPredictionKey to this value. The client then looks
+	 *			at this replicated value. In OnRep_PredictionKey, the client goes through and cleans up all stuff that was predicted
+	 *			and is no longer needed.
+	 *			
+	 *			-Client calls GetNextPredictionKey() at the start of a predicted action. This gives him a key to use for this 'run'.
+	 *			-Client can predict whatever he wants as long as he registers an 'undo' function with GetPredictionKeyDelegate().
+	 *			-This delegate is invoked if the client is denied his prediction, or if he is right and property replication catches up.
+	 *			
+	 *	 
+	 */
+
+	UPROPERTY(ReplicatedUsing=OnRep_PredictionKey)
+	int32	ReplicatedPredictionKey;
+
+	UPROPERTY(transient)
+	int32	LocalPredictionKey;
+
+	int32	GetNextPredictionKey();
+
+	UFUNCTION()
+	void OnRep_PredictionKey();
+
+	TArray<TPair<int32, FAbilitySystemComponentPredictionKeyClear> >	PredictionDelegates;
+	
+
+	FAbilitySystemComponentPredictionKeyClear &	GetPredictionKeyDelegate(int32 PredictionKey);
+
+	/*
+	virtual void OnSubobjectCreatedFromReplication(UObject *NewSubobject) OVERRIDE;
+	
+	virtual void OnSubobjectDestroyFromReplication(UObject *NewSubobject) OVERRIDE;
+	*/
+
+	// ----------------------------------------------------------------------------------------------------------------
+	//
+	//	GameplayEffects
+	//	(maybe should go in a different component?)
+	// ----------------------------------------------------------------------------------------------------------------
+
+	// --------------------------------------------
+	// Primary outward facing API for other systems:
+	// --------------------------------------------
+	FActiveGameplayEffectHandle ApplyGameplayEffectSpecToTarget(OUT FGameplayEffectSpec &GameplayEffect, UAbilitySystemComponent *Target, FModifierQualifier BaseQualifier = FModifierQualifier());
+	FActiveGameplayEffectHandle ApplyGameplayEffectSpecToSelf(const FGameplayEffectSpec &GameplayEffect, FModifierQualifier BaseQualifier = FModifierQualifier());
+
+	UFUNCTION(BlueprintCallable, BlueprintAuthorityOnly, Category = GameplayEffects)
+	bool RemoveActiveGameplayEffect(FActiveGameplayEffectHandle Handle);
+
+	/** This only exists so it can be hooked up to a multicast delegate */
+	void RemoveActiveGameplayEffect_NoReturn(FActiveGameplayEffectHandle Handle)
+	{
+		RemoveActiveGameplayEffect(Handle);
+	}
+
+	float GetGameplayEffectDuration(FActiveGameplayEffectHandle Handle) const;
+
+	float GetGameplayEffectDuration() const;
+
+	// Not happy with this interface but don't see a better way yet. How should outside code (UI, etc) ask things like 'how much is this gameplay effect modifying my damage by'
+	// (most likely we want to catch this on the backend - when damage is applied we can get a full dump/history of how the number got to where it is. But still we may need polling methods like below (how much would my damage be)
+	UFUNCTION(BlueprintCallable, Category = GameplayEffects)
+	float GetGameplayEffectMagnitude(FActiveGameplayEffectHandle Handle, FGameplayAttribute Attribute) const;
+
+	UFUNCTION(BlueprintCallable, Category = GameplayEffects)
+	float GetGameplayEffectMagnitudeByTag(FActiveGameplayEffectHandle InHandle, const FGameplayTag& InTag) const;
+
+	UFUNCTION(BlueprintCallable, Category = GameplayEffects)
+	bool IsGameplayEffectActive(FActiveGameplayEffectHandle InHandle) const;
+
+	// Tags
+	UFUNCTION(BlueprintCallable, Category=GameplayEffects)
+	bool HasAnyTags(FGameplayTagContainer &Tags);
+
+	UFUNCTION(BlueprintCallable, Category = GameplayEffects)
+	bool HasAllTags(FGameplayTagContainer &Tags);
+	
+
+	// --------------------------------------------
+	// Possibly useful but not primary API functions:
+	// --------------------------------------------
+
+	FActiveGameplayEffectHandle ApplyGameplayEffectToTarget(UGameplayEffect *GameplayEffect, UAbilitySystemComponent *Target, float Level = FGameplayEffectLevelSpec::INVALID_LEVEL, FModifierQualifier BaseQualifier = FModifierQualifier());
+
+	UFUNCTION(BlueprintCallable, BlueprintAuthorityOnly, Category = GameplayEffects, meta=(FriendlyName = "ApplyGameplayEffectToSelf"))
+	FActiveGameplayEffectHandle K2_ApplyGameplayEffectToSelf(const UGameplayEffect *GameplayEffect, float Level, AActor *Instigator);
+	
+	FActiveGameplayEffectHandle ApplyGameplayEffectToSelf(const UGameplayEffect *GameplayEffect, float Level, AActor *Instigator, FModifierQualifier BaseQualifier = FModifierQualifier() );
+
+	int32 GetNumActiveGameplayEffect() const;
+
+	void AddDependancyToAttribute(FGameplayAttribute Attribute, const TWeakPtr<FAggregator> InDependant);
+
+	/** Tests if all modifiers in this GameplayEffect will leave the attribute > 0.f */
+	bool CanApplyAttributeModifiers(const UGameplayEffect *GameplayEffect, float Level, AActor *Instigator);
+
+	// Generic 'Get expected magnitude (list) if I was to apply this outgoing or incoming'
+
+	// Get duration or magnitude (list) of active effects
+	//		-Get duration of CD
+	//		-Get magnitude + duration of a movespeed buff
+
+	TArray<float> GetActiveEffectsTimeRemaining(const FActiveGameplayEffectQuery Query) const;
+
+	void OnRestackGameplayEffects();
+
+	// --------------------------------------------
+	// Temp / Debug
+	// --------------------------------------------
+
+	void TEMP_ApplyActiveGameplayEffects();
+	
+	void PrintAllGameplayEffects() const;
+
+	void PushGlobalCurveOveride(UCurveTable *OverrideTable)
+	{
+		if (OverrideTable)
+		{
+			GlobalCurveDataOverride.Overrides.Push(OverrideTable);
+		}
+	}
+
+	// ----------------------------------------------------------------------------------------------------------------
+
+	/**
+	 *	GameplayAbilities
+	 *	
+	 *	The role of the AbilitySystemComponent wrt Abilities is to provide:
+	 *		-Management of abiltiy instances (whether per actor or per execution instance).
+	 *			-Someone *has* to keep track of these instances.
+	 *			-Non instanced abilities *could* be executed without any ability stuff in AbilitySystemComponent.
+	 *				They should be able to operate on an GameplayAbilityActorInfo + GameplayAbility.
+	 *		
+	 *	As convenience it may provide some other features:
+	 *		-Some basic input binding (whether instanced or non instanced abilities).
+	 *		-Concepts like "this component has these abilities
+	 *	
+	 */
+
+	class UGameplayAbility * CreateNewInstanceOfAbility(UGameplayAbility* Ability);
+
+	void CancelAbilitiesWithTags(const FGameplayTagContainer Tags, const FGameplayAbilityActorInfo* ActorInfo, const FGameplayAbilityActivationInfo ActivationInfo, UGameplayAbility* Ignore);
+	
+	/** References to non-replicating abilities that we instanced. We need to keep these refernces to avoid GC */
+	UPROPERTY()
+	TArray<class UGameplayAbility *>	NonReplicatedInstancedAbilities;
+
+	/** References to replicating abilities that we instanced. We need to keep these refernces to avoid GC */
+	UPROPERTY(Replicated)
+	TArray<class UGameplayAbility *>	ReplicatedInstancedAbilities;
+
+	/**
+	 *	The abilities we can activate. 
+	 *		-This will include CDOs for non instanced abilities and per-execution instanced abilities. 
+	 *		-Actor-instanced abilities will be the actual instance (not CDO)
+	 *		
+	 *	This array is not vital for things to work. It is a convenience thing for 'giving abilties to the actor'. But abilities could also work on things
+	 *	without an AbilitySystemComponent. For example an ability could be written to execute on a StaticMeshActor. As long as the ability doesn't require 
+	 *	instancing or anything else that the AbilitySystemComponent would provide, then it doesn't need the component to function.
+	 */ 
+	UPROPERTY(Replicated, BlueprintReadOnly, Category = "Abilities")
+	TArray<class UGameplayAbility *>	ActivatableAbilities;
+	
+
+	UFUNCTION(Server, WithValidation, Reliable)
+	void	ServerTryActivateAbility(class UGameplayAbility* AbilityToActivate, int32 PredictionKey);
+
+	UFUNCTION(Client, Reliable)
+	void	ClientActivateAbilityFailed(class UGameplayAbility* AbilityToActivate, int32 PredictionKey);
+
+	UFUNCTION(Client, Reliable)
+	void	ClientActivateAbilitySucceed(class UGameplayAbility* AbilityToActivate, int32 PredictionKey);
+
+	// ----------------------------------------------------------------------------------------------------------------
+
+	/** This is probably temp for testing. We want to think a bit more on the API for outside stuff activating abilities */
+	UFUNCTION(BlueprintCallable, Category = "Abilities")
+	bool	ActivateAbility(int32 AbilityIdx);
+
+	/** There needs to be a concept of an animating ability. Only one may exist at a time. New requestrs can be queued up, overridden, or ignored. */
+	UPROPERTY()
+	class UGameplayAbility *	AnimatingAbility;
+
+	void	MontageBranchPoint_AbilityDecisionStop();
+
+	void	MontageBranchPoint_AbilityDecisionStart();
+
+	
+
+	// -----------------------------------------------------------------------------
+
+	/** Cached off data about the owning actor that abilities will need to frequently access (movement component, mesh component, anim instance, etc) */
+	TSharedPtr<FGameplayAbilityActorInfo>	AbilityActorInfo;
+
+	void InitAbilityActorInfo();
+
+private:
+
+	// --------------------------------------------
+	// Important internal functions
+	// --------------------------------------------
+
+	void ExecutePeriodicEffect(FActiveGameplayEffectHandle	Handle);
+
+	void ExecuteGameplayEffect(const FGameplayEffectSpec &Spec, const FModifierQualifier &QualifierContext);
+
+	void CheckDurationExpired(FActiveGameplayEffectHandle Handle);
+
+	// --------------------------------------------
+	// Internal Utility / helper functions
+	// --------------------------------------------
+
+	bool AreGameplayEffectApplicationRequirementsSatisfied(const class UGameplayEffect* EffectToAdd, FGameplayEffectInstigatorContext& Instigator) const;
+
+	bool IsOwnerActorAuthoritative() const;
+
+	void OnAttributeGameplayEffectSpecExected(const FGameplayAttribute &Attribute, const struct FGameplayEffectSpec &Spec, struct FGameplayModifierEvaluatedData &Data);
+
+	const FGlobalCurveDataOverride * GetCurveDataOverride() const
+	{
+		// only return data if we have overrides. NULL if we dont.
+		return (GlobalCurveDataOverride.Overrides.Num() > 0 ? &GlobalCurveDataOverride : NULL);
+	}
+
+	FGlobalCurveDataOverride	GlobalCurveDataOverride;
+	
+
+	// --------------------------------------------
+	
+	UPROPERTY(ReplicatedUsing=OnRep_GameplayEffects)
+	FActiveGameplayEffectsContainer	ActiveGameplayEffects;
+
+	UFUNCTION()
+	void OnRep_GameplayEffects();
+
+	// ---------------------------------------------
+
+	UFUNCTION(NetMulticast, unreliable)
+	void NetMulticast_InvokeGameplayCueExecuted(const FGameplayEffectSpec Spec);
+
+	void InvokeGameplayCueExecute(const FGameplayEffectSpec &Spec);
+	void InvokeGameplayCueActivated(const FGameplayEffectSpec &Spec);
+	void InvokeGameplayCueAdded(const FGameplayEffectSpec &Spec);
+	void InvokeGameplayCueRemoved(const FGameplayEffectSpec &Spec);
+
+	// ---------------------------------------------
+
+	
+
+protected:
+
+	virtual void OnRegister() OVERRIDE;
+
+	UAttributeSet*	GetAttributeSubobject(const TSubclassOf<UAttributeSet> AttributeClass) const;
+	UAttributeSet*	GetAttributeSubobjectChecked(const TSubclassOf<UAttributeSet> AttributeClass) const;
+	UAttributeSet*	GetOrCreateAttributeSubobject(const TSubclassOf<UAttributeSet> AttributeClass);
+
+	friend struct FActiveGameplayEffectsContainer;
+	friend struct FActiveGameplayEffect;
+};
diff --git a/Engine/Source/Runtime/GameplayAbilities/Classes/AbilitySystemDebugHUD.h b/Engine/Source/Runtime/GameplayAbilities/Classes/AbilitySystemDebugHUD.h
--- a/Engine/Source/Runtime/GameplayAbilities/Classes/AbilitySystemDebugHUD.h
+++ b/Engine/Source/Runtime/GameplayAbilities/Classes/AbilitySystemDebugHUD.h
@@ -38,6 +38,6 @@ class AAbilitySystemDebugHUD : public AHUD
 
 private:
 
-	void DrawDebugAttributeComponent(UAttributeComponent *Component);
+	void DrawDebugAbilitySystemComponent(UAbilitySystemComponent *Component);
 
 };
\ No newline at end of file
diff --git a/Engine/Source/Runtime/GameplayAbilities/Classes/AbilitySystemGlobals.h b/Engine/Source/Runtime/GameplayAbilities/Classes/AbilitySystemGlobals.h
--- a/Engine/Source/Runtime/GameplayAbilities/Classes/AbilitySystemGlobals.h
+++ b/Engine/Source/Runtime/GameplayAbilities/Classes/AbilitySystemGlobals.h
@@ -5,7 +5,7 @@
 #include "AbilitySystemGlobals.generated.h"
 
 class AActor;
-class UAttributeComponent;
+class UAbilitySystemComponent;
 struct FGameplayAbilityActorInfo;
 
 /** Holds global data for the skill system. Can be configured per project via config file */
@@ -40,9 +40,9 @@ class GAMEPLAYABILITIES_API UAbilitySystemGlobals : public UObject
 
 	/**
 	 *	Games may want to override this in a UMyProjectAbilitySystemsGlobals class and provide
-	 *	a faster lookup for attribute compnents (E.g. Cast<MyProjPawn>(Actor)->AttributeComponent; etc)
+	 *	a faster lookup for attribute compnents (E.g. Cast<MyProjPawn>(Actor)->AbilitySystemComponent; etc)
 	 */
-	virtual UAttributeComponent * GetAttributeComponentFromActor(AActor *Actor) const;
+	virtual UAbilitySystemComponent * GetAbilitySystemComponentFromActor(AActor *Actor) const;
 
 	virtual FGameplayAbilityActorInfo * AllocAbilityActorInfo(AActor *Actor) const;
 
diff --git a/Engine/Source/Runtime/GameplayAbilities/Classes/AbilitySystemTestPawn.h b/Engine/Source/Runtime/GameplayAbilities/Classes/AbilitySystemTestPawn.h
--- a/Engine/Source/Runtime/GameplayAbilities/Classes/AbilitySystemTestPawn.h
+++ b/Engine/Source/Runtime/GameplayAbilities/Classes/AbilitySystemTestPawn.h
@@ -27,7 +27,7 @@ class GAMEPLAYABILITIES_API AAbilitySystemTestPawn : public ADefaultPawn, public
 
 	/** DefaultPawn collision component */
 	UPROPERTY(Category = AbilitySystem, VisibleAnywhere, BlueprintReadOnly)
-	TSubobjectPtr<class UAttributeComponent> AttributeComponent;
+	TSubobjectPtr<class UAbilitySystemComponent> AbilitySystemComponent;
 
 	UPROPERTY(EditDefaultsOnly, Category=GameplayEffects)
 	FGameplayCueHandler	GameplayCueHandler;
@@ -35,5 +35,5 @@ class GAMEPLAYABILITIES_API AAbilitySystemTestPawn : public ADefaultPawn, public
 	//UPROPERTY(EditDefaultsOnly, Category=GameplayEffects)
 	//UGameplayAbilitySet * DefaultAbilitySet;
 
-	static FName AttributeComponentName;
+	static FName AbilitySystemComponentName;
 };
\ No newline at end of file
diff --git a/Engine/Source/Runtime/GameplayAbilities/Classes/AttributeComponent.h b/Engine/Source/Runtime/GameplayAbilities/Classes/AttributeComponent.h
deleted file mode 100644
diff --git a/Engine/Source/Runtime/GameplayAbilities/Classes/GameplayAbilitySet.h b/Engine/Source/Runtime/GameplayAbilities/Classes/GameplayAbilitySet.h
--- a/Engine/Source/Runtime/GameplayAbilities/Classes/GameplayAbilitySet.h
+++ b/Engine/Source/Runtime/GameplayAbilities/Classes/GameplayAbilitySet.h
@@ -35,8 +35,8 @@ struct FGameplayAbilityBindInfo
 /**
 * UGameplayAbilitySet
 *	This contains a list of abilities along with key bindings. This will be very game specific, so it is expected for games to override this class.
-*	Functions like InitializeAbilities and BindInputComponentToAbilities are setup to act on the attribute component (rather than attribute component processing the ability set).
-*	Ideally this avoids games having to implement their own UAttributeComponent and UGameplayAbilitySet, and they can just implement thei own UGameplayAbilitySet.
+*	Functions like InitializeAbilities and BindInputComponentToAbilities are setup to act on the AbilitySystemComponent (rather than AbilitySystemComponent processing the ability set).
+*	Ideally this avoids games having to implement their own UAbilitySystemComponent and UGameplayAbilitySet, and they can just implement thei own UGameplayAbilitySet.
 */
 UCLASS(Blueprintable)
 class GAMEPLAYABILITIES_API UGameplayAbilitySet : public UDataAsset
@@ -48,9 +48,9 @@ public:
 	UPROPERTY(EditAnywhere, Category=AbilitySet)
 	TArray<FGameplayAbilityBindInfo>	Abilities;
 
-	/** Give all of the abilities in this set to given AttributeComponent */
-	virtual void InitializeAbilities(UAttributeComponent *AttributeComponent) const;
+	/** Give all of the abilities in this set to given AbilitySystemComponent */
+	virtual void InitializeAbilities(UAbilitySystemComponent *AbilitySystemComponent) const;
 
-	/** Binds the ability set directly to this inputcomponent and attributecomponent. This is a convenience method, not all games will want to bind input directly to abilities. */
-	virtual void BindInputComponentToAbilities(UInputComponent *InputComponent, UAttributeComponent *AttributeComponent) const;
+	/** Binds the ability set directly to this inputcomponent and AbilitySystemComponent. This is a convenience method, not all games will want to bind input directly to abilities. */
+	virtual void BindInputComponentToAbilities(UInputComponent *InputComponent, UAbilitySystemComponent *AbilitySystemComponent) const;
 };
\ No newline at end of file
diff --git a/Engine/Source/Runtime/GameplayAbilities/Classes/GameplayEffect.h b/Engine/Source/Runtime/GameplayAbilities/Classes/GameplayEffect.h
--- a/Engine/Source/Runtime/GameplayAbilities/Classes/GameplayEffect.h
+++ b/Engine/Source/Runtime/GameplayAbilities/Classes/GameplayEffect.h
@@ -23,7 +23,7 @@ struct FActiveGameplayEffect;
 #endif
 
 class UGameplayEffect;
-class UAttributeComponent;
+class UAbilitySystemComponent;
 
 UENUM(BlueprintType)
 namespace EGameplayModOp
@@ -607,7 +607,7 @@ struct FGameplayEffectInstigatorContext
 
 	FGameplayEffectInstigatorContext()
 		: Instigator(NULL)
-		, InstigatorAttributeComponent(NULL)
+		, InstigatorAbilitySystemComponent(NULL)
 	{
 	}
 
@@ -635,9 +635,9 @@ struct FGameplayEffectInstigatorContext
 		return Instigator;
 	}
 
-	UAttributeComponent * GetOriginalInstigatorAttributeComponent() const
+	UAbilitySystemComponent * GetOriginalInstigatorAbilitySystemComponent() const
 	{
-		return InstigatorAttributeComponent;
+		return InstigatorAbilitySystemComponent;
 	}
 	
 	/** Instigator controller */
@@ -645,7 +645,7 @@ struct FGameplayEffectInstigatorContext
 	class AActor* Instigator;
 
 	UPROPERTY(NotReplicated)
-	UAttributeComponent *InstigatorAttributeComponent;
+	UAbilitySystemComponent *InstigatorAbilitySystemComponent;
 
 	/** Trace information - may be NULL in many cases */
 	TSharedPtr<FHitResult>	HitResult;
@@ -1322,7 +1322,7 @@ struct FActiveGameplayEffectQuery
  *	-Bucket of ActiveGameplayEffects
  *	-Needed for FFastArraySerialization
  *  
- * This should only be used by UAttributeComponent. All of this could just live in UAttributeComponent except that we need a distinct USTRUCT to implement FFastArraySerializer.
+ * This should only be used by UAbilitySystemComponent. All of this could just live in UAbilitySystemComponent except that we need a distinct USTRUCT to implement FFastArraySerializer.
  *
  */
 USTRUCT()
@@ -1345,7 +1345,7 @@ struct FActiveGameplayEffectsContainer : public FFastArraySerializer
 		
 	void ExecuteActiveEffectsFrom(const FGameplayEffectSpec &Spec, const FModifierQualifier &QualifierContext);
 	
-	void ExecutePeriodicGameplayEffect(FActiveGameplayEffectHandle Handle);	// This should not be outward facing to the skill system API, should only be called by the owning attribute component
+	void ExecutePeriodicGameplayEffect(FActiveGameplayEffectHandle Handle);	// This should not be outward facing to the skill system API, should only be called by the owning AbilitySystemComponent
 
 	void AddDependancyToAttribute(FGameplayAttribute Attribute, const TWeakPtr<FAggregator> InDependant);
 
@@ -1370,7 +1370,7 @@ struct FActiveGameplayEffectsContainer : public FFastArraySerializer
 
 	void OnPropertyAggregatorDirty(FAggregator* Aggregator, FGameplayAttribute Attribute);
 
-	UAttributeComponent *Owner;
+	UAbilitySystemComponent *Owner;
 
 	void CheckDuration(FActiveGameplayEffectHandle Handle);
 
@@ -1399,7 +1399,7 @@ struct FActiveGameplayEffectsContainer : public FFastArraySerializer
 	int32 GetGameStateTime() const;
 	float GetWorldTime() const;
 
-	void OnDurationAggregatorDirty(FAggregator* Aggregator, UAttributeComponent* Owner, FActiveGameplayEffectHandle Handle);
+	void OnDurationAggregatorDirty(FAggregator* Aggregator, UAbilitySystemComponent* Owner, FActiveGameplayEffectHandle Handle);
 
 private:
 
diff --git a/Engine/Source/Runtime/GameplayAbilities/Classes/GameplayEffectExtension.h b/Engine/Source/Runtime/GameplayAbilities/Classes/GameplayEffectExtension.h
--- a/Engine/Source/Runtime/GameplayAbilities/Classes/GameplayEffectExtension.h
+++ b/Engine/Source/Runtime/GameplayAbilities/Classes/GameplayEffectExtension.h
@@ -10,7 +10,7 @@
 
 struct FGameplayEffectModCallbackData
 {
-	FGameplayEffectModCallbackData(const FGameplayEffectSpec & InEffectSpec, const FModifierSpec & InModifierSpec, FGameplayModifierEvaluatedData & InEvaluatedData, UAttributeComponent & InTarget)
+	FGameplayEffectModCallbackData(const FGameplayEffectSpec & InEffectSpec, const FModifierSpec & InModifierSpec, FGameplayModifierEvaluatedData & InEvaluatedData, UAbilitySystemComponent & InTarget)
 		: EffectSpec(InEffectSpec)
 		, ModifierSpec(InModifierSpec)
 		, EvaluatedData(InEvaluatedData)
@@ -23,7 +23,7 @@ struct FGameplayEffectModCallbackData
 	const struct FModifierSpec &			ModifierSpec;	// The mod we are going to apply
 	struct FGameplayModifierEvaluatedData & EvaluatedData;	// The 'flat'/computed data to be applied to the target
 
-	class UAttributeComponent &Target;		// Target we intend to apply to
+	class UAbilitySystemComponent &Target;		// Target we intend to apply to
 };
 
 UCLASS(BlueprintType)
diff --git a/Engine/Source/Runtime/GameplayAbilities/Classes/GameplayEffectStackingExtension.h b/Engine/Source/Runtime/GameplayAbilities/Classes/GameplayEffectStackingExtension.h
--- a/Engine/Source/Runtime/GameplayAbilities/Classes/GameplayEffectStackingExtension.h
+++ b/Engine/Source/Runtime/GameplayAbilities/Classes/GameplayEffectStackingExtension.h
@@ -10,7 +10,7 @@
 
 struct FGameplayEffectStackingCallbackData
 {
-	FGameplayEffectStackingCallbackData(const FGameplayEffectSpec & InEffectSpec, const FModifierSpec & InModifierSpec, FGameplayModifierEvaluatedData & InEvaluatedData, UAttributeComponent & InTarget)
+	FGameplayEffectStackingCallbackData(const FGameplayEffectSpec & InEffectSpec, const FModifierSpec & InModifierSpec, FGameplayModifierEvaluatedData & InEvaluatedData, UAbilitySystemComponent & InTarget)
 		: EffectSpec(InEffectSpec)
 		, ModifierSpec(InModifierSpec)
 		, EvaluatedData(InEvaluatedData)
@@ -23,7 +23,7 @@ struct FGameplayEffectStackingCallbackData
 	const struct FModifierSpec &			ModifierSpec;	// The mod we are going to apply
 	struct FGameplayModifierEvaluatedData & EvaluatedData;	// The 'flat'/computed data to be applied to the target
 
-	class UAttributeComponent &Target;		// Target we intend to apply to
+	class UAbilitySystemComponent &Target;		// Target we intend to apply to
 };
 
 
diff --git a/Engine/Source/Runtime/GameplayAbilities/Private/Abilities/GameplayAbility.cpp b/Engine/Source/Runtime/GameplayAbilities/Private/Abilities/GameplayAbility.cpp
--- a/Engine/Source/Runtime/GameplayAbilities/Private/Abilities/GameplayAbility.cpp
+++ b/Engine/Source/Runtime/GameplayAbilities/Private/Abilities/GameplayAbility.cpp
@@ -3,7 +3,7 @@
 #include "AbilitySystemPrivatePCH.h"
 #include "Abilities/GameplayAbility.h"
 #include "GameplayEffect.h"
-#include "AttributeComponent.h"
+#include "AbilitySystemComponent.h"
 #include "AbilitySystemBlueprintLibrary.h"
 
 // --------------------------------------------------------------------------------------------------------------------------------------------------------
@@ -72,7 +72,7 @@ bool UGameplayAbility::CommitCheck(const FGameplayAbilityActorInfo* ActorInfo, c
 
 void UGameplayAbility::CommitExecute(const FGameplayAbilityActorInfo* ActorInfo, const FGameplayAbilityActivationInfo ActivationInfo)
 {
-	if (!ActorInfo->AttributeComponent->IsNetSimulating() || ActivationInfo.PredictionKey != 0)
+	if (!ActorInfo->AbilitySystemComponent->IsNetSimulating() || ActivationInfo.PredictionKey != 0)
 	{
 		ApplyCooldown(ActorInfo, ActivationInfo);
 
@@ -101,7 +101,7 @@ bool UGameplayAbility::TryActivateAbility(const FGameplayAbilityActorInfo* Actor
 		
 		if (GetInstancingPolicy() == EGameplayAbilityInstancingPolicy::InstancedPerExecution)
 		{
-			UGameplayAbility * Ability = ActorInfo->AttributeComponent->CreateNewInstanceOfAbility(this);
+			UGameplayAbility * Ability = ActorInfo->AbilitySystemComponent->CreateNewInstanceOfAbility(this);
 			Ability->CallActivateAbility(ActorInfo, ActivationInfo);
 			if (OutInstancedAbility)
 			{
@@ -120,7 +120,7 @@ bool UGameplayAbility::TryActivateAbility(const FGameplayAbilityActorInfo* Actor
 	else if (NetExecutionPolicy == EGameplayAbilityNetExecutionPolicy::Predictive)
 	{
 		// This execution is now officially EGameplayAbilityActivationMode:Predicting and has a PredictionKey
-		ActivationInfo.GeneratePredictionKey(ActorInfo->AttributeComponent.Get());
+		ActivationInfo.GeneratePredictionKey(ActorInfo->AbilitySystemComponent.Get());
 
 		if (GetInstancingPolicy() == EGameplayAbilityInstancingPolicy::InstancedPerExecution)
 		{
@@ -130,7 +130,7 @@ bool UGameplayAbility::TryActivateAbility(const FGameplayAbilityActorInfo* Actor
 			
 			if(GetReplicationPolicy() == EGameplayAbilityReplicationPolicy::ReplicateNone)
 			{
-				UGameplayAbility * Ability = ActorInfo->AttributeComponent->CreateNewInstanceOfAbility(this);
+				UGameplayAbility * Ability = ActorInfo->AbilitySystemComponent->CreateNewInstanceOfAbility(this);
 				Ability->CallPredictiveActivateAbility(ActorInfo, ActivationInfo);
 				if (OutInstancedAbility)
 				{
@@ -170,7 +170,7 @@ void UGameplayAbility::ActivateAbility(const FGameplayAbilityActorInfo* ActorInf
 
 void UGameplayAbility::PreActivate(const FGameplayAbilityActorInfo* ActorInfo, const FGameplayAbilityActivationInfo ActivationInfo)
 {
-	ActorInfo->AttributeComponent->CancelAbilitiesWithTags(CancelAbilitiesWithTag, ActorInfo, ActivationInfo, this);
+	ActorInfo->AbilitySystemComponent->CancelAbilitiesWithTags(CancelAbilitiesWithTag, ActorInfo, ActivationInfo, this);
 }
 
 void UGameplayAbility::CallActivateAbility(const FGameplayAbilityActorInfo* ActorInfo, const FGameplayAbilityActivationInfo ActivationInfo)
@@ -192,7 +192,7 @@ void UGameplayAbility::PredictiveActivateAbility(const FGameplayAbilityActorInfo
 
 void UGameplayAbility::ServerTryActivateAbility(const FGameplayAbilityActorInfo* ActorInfo, const FGameplayAbilityActivationInfo ActivationInfo)
 {
-	ActorInfo->AttributeComponent->ServerTryActivateAbility(this, ActivationInfo.PredictionKey);
+	ActorInfo->AbilitySystemComponent->ServerTryActivateAbility(this, ActivationInfo.PredictionKey);
 }
 
 void UGameplayAbility::ClientActivateAbility(const FGameplayAbilityActorInfo* ActorInfo, const FGameplayAbilityActivationInfo ActivationInfo)
@@ -214,8 +214,8 @@ bool UGameplayAbility::CheckCooldown(const FGameplayAbilityActorInfo* ActorInfo)
 {
 	if (CooldownGameplayEffect)
 	{
-		check(ActorInfo->AttributeComponent.IsValid());
-		if (ActorInfo->AttributeComponent->HasAnyTags(CooldownGameplayEffect->OwnedTagsContainer))
+		check(ActorInfo->AbilitySystemComponent.IsValid());
+		if (ActorInfo->AbilitySystemComponent->HasAnyTags(CooldownGameplayEffect->OwnedTagsContainer))
 		{
 			return false;
 		}
@@ -227,8 +227,8 @@ void UGameplayAbility::ApplyCooldown(const FGameplayAbilityActorInfo* ActorInfo,
 {
 	if (CooldownGameplayEffect)
 	{
-		check(ActorInfo->AttributeComponent.IsValid());
-		ActorInfo->AttributeComponent->ApplyGameplayEffectToSelf(CooldownGameplayEffect, 1.f, ActorInfo->Actor.Get(), FModifierQualifier().PredictionKey(ActivationInfo.PredictionKey));
+		check(ActorInfo->AbilitySystemComponent.IsValid());
+		ActorInfo->AbilitySystemComponent->ApplyGameplayEffectToSelf(CooldownGameplayEffect, 1.f, ActorInfo->Actor.Get(), FModifierQualifier().PredictionKey(ActivationInfo.PredictionKey));
 	}
 }
 
@@ -236,18 +236,18 @@ bool UGameplayAbility::CheckCost(const FGameplayAbilityActorInfo* ActorInfo) con
 {
 	if (CostGameplayEffect)
 	{
-		check(ActorInfo->AttributeComponent.IsValid());
-		return ActorInfo->AttributeComponent->CanApplyAttributeModifiers(CostGameplayEffect, 1.f, ActorInfo->Actor.Get());
+		check(ActorInfo->AbilitySystemComponent.IsValid());
+		return ActorInfo->AbilitySystemComponent->CanApplyAttributeModifiers(CostGameplayEffect, 1.f, ActorInfo->Actor.Get());
 	}
 	return true;
 }
 
 void UGameplayAbility::ApplyCost(const FGameplayAbilityActorInfo* ActorInfo, const FGameplayAbilityActivationInfo ActivationInfo)
 {
-	check(ActorInfo->AttributeComponent.IsValid());
+	check(ActorInfo->AbilitySystemComponent.IsValid());
 	if (CostGameplayEffect)
 	{
-		ActorInfo->AttributeComponent->ApplyGameplayEffectToSelf(CostGameplayEffect, 1.f, ActorInfo->Actor.Get(), FModifierQualifier().PredictionKey(ActivationInfo.PredictionKey));
+		ActorInfo->AbilitySystemComponent->ApplyGameplayEffectToSelf(CostGameplayEffect, 1.f, ActorInfo->Actor.Get(), FModifierQualifier().PredictionKey(ActivationInfo.PredictionKey));
 	}
 }
 
@@ -255,10 +255,10 @@ float UGameplayAbility::GetCooldownTimeRemaining(const FGameplayAbilityActorInfo
 {
 	SCOPE_CYCLE_COUNTER(STAT_GameplayAbilityGetCooldownTimeRemaining);
 
-	check(ActorInfo->AttributeComponent.IsValid());
+	check(ActorInfo->AbilitySystemComponent.IsValid());
 	if (CooldownGameplayEffect)
 	{
-		TArray< float > Durations = ActorInfo->AttributeComponent->GetActiveEffectsTimeRemaining(FActiveGameplayEffectQuery(&CooldownGameplayEffect->OwnedTagsContainer));
+		TArray< float > Durations = ActorInfo->AbilitySystemComponent->GetActiveEffectsTimeRemaining(FActiveGameplayEffectQuery(&CooldownGameplayEffect->OwnedTagsContainer));
 		if (Durations.Num() > 0)
 		{
 			Durations.Sort();
@@ -305,12 +305,12 @@ void FGameplayAbilityActorInfo::InitFromActor(AActor *InActor)
 		this->AnimInstance = SkelMeshComponent->GetAnimInstance();
 	}
 
-	AttributeComponent = InActor->FindComponentByClass<UAttributeComponent>();
+	AbilitySystemComponent = InActor->FindComponentByClass<UAbilitySystemComponent>();
 
 	MovementComponent = InActor->FindComponentByClass<UMovementComponent>();
 }
 
-void FGameplayAbilityActivationInfo::GeneratePredictionKey(UAttributeComponent * Component) const
+void FGameplayAbilityActivationInfo::GeneratePredictionKey(UAbilitySystemComponent * Component) const
 {
 	check(Component);
 	check(PredictionKey == 0);
@@ -353,10 +353,10 @@ void FGameplayAbilityTargetData::ApplyGameplayEffect(UGameplayEffect* GameplayEf
 	{
 		check(TargetActor);
 		
-		UAttributeComponent * TargetComponent = UAbilitySystemBlueprintLibrary::GetAttributeComponent(TargetActor);
+		UAbilitySystemComponent * TargetComponent = UAbilitySystemBlueprintLibrary::GetAbilitySystemComponent(TargetActor);
 		if (TargetComponent)
 		{
-			InstigatorInfo.AttributeComponent->ApplyGameplayEffectSpecToTarget(SpecToApply, TargetComponent);
+			InstigatorInfo.AbilitySystemComponent->ApplyGameplayEffectSpecToTarget(SpecToApply, TargetComponent);
 		}
 	}
 }
diff --git a/Engine/Source/Runtime/GameplayAbilities/Private/Abilities/GameplayAbility_Instanced.cpp b/Engine/Source/Runtime/GameplayAbilities/Private/Abilities/GameplayAbility_Instanced.cpp
--- a/Engine/Source/Runtime/GameplayAbilities/Private/Abilities/GameplayAbility_Instanced.cpp
+++ b/Engine/Source/Runtime/GameplayAbilities/Private/Abilities/GameplayAbility_Instanced.cpp
@@ -109,7 +109,7 @@ void UGameplayAbility_Instanced::EndAbility(const FGameplayAbilityActorInfo* Act
 		MarkPendingKill();
 	}
 
-	// Remove from owning attributecomponent?
+	// Remove from owning AbilitySystemComponent?
 	// generic way of releasing all callbacks!
 }
 
diff --git a/Engine/Source/Runtime/GameplayAbilities/Private/Abilities/GameplayAbility_Montage.cpp b/Engine/Source/Runtime/GameplayAbilities/Private/Abilities/GameplayAbility_Montage.cpp
--- a/Engine/Source/Runtime/GameplayAbilities/Private/Abilities/GameplayAbility_Montage.cpp
+++ b/Engine/Source/Runtime/GameplayAbilities/Private/Abilities/GameplayAbility_Montage.cpp
@@ -3,7 +3,7 @@
 #include "AbilitySystemPrivatePCH.h"
 #include "Abilities/GameplayAbility_Montage.h"
 #include "GameplayEffect.h"
-#include "AttributeComponent.h"
+#include "AbilitySystemComponent.h"
 
 // --------------------------------------------------------------------------------------------------------------------------------------------------------
 //
@@ -32,7 +32,7 @@ void UGameplayAbility_Montage::ActivateAbility(const FGameplayAbilityActorInfo*
 		// Apply GameplayEffects
 		for (const UGameplayEffect * Effect : GameplayEffectsWhileAnimating)
 		{
-			FActiveGameplayEffectHandle Handle = ActorInfo->AttributeComponent->ApplyGameplayEffectToSelf(Effect, 1.f, ActorInfo->Actor.Get());
+			FActiveGameplayEffectHandle Handle = ActorInfo->AbilitySystemComponent->ApplyGameplayEffectToSelf(Effect, 1.f, ActorInfo->Actor.Get());
 			if (Handle.IsValid())
 			{
 				AppliedEffects.Add(Handle);
@@ -55,7 +55,7 @@ void UGameplayAbility_Montage::ActivateAbility(const FGameplayAbilityActorInfo*
 void UGameplayAbility_Montage::OnMontageEnded(UAnimMontage* Montage, bool bInterrupted, const FGameplayAbilityActorInfo* ActorInfo, TArray<FActiveGameplayEffectHandle> AppliedEffects)
 {
 	// Remove any GameplayEffects that we applied
-	UAttributeComponent *OwnerComponent = ActorInfo->AttributeComponent.Get();
+	UAbilitySystemComponent *OwnerComponent = ActorInfo->AbilitySystemComponent.Get();
 	if (OwnerComponent)
 	{
 		for (FActiveGameplayEffectHandle Handle : AppliedEffects)
diff --git a/Engine/Source/Runtime/GameplayAbilities/Private/AbilitySystemBlueprintLibrary.cpp b/Engine/Source/Runtime/GameplayAbilities/Private/AbilitySystemBlueprintLibrary.cpp
--- a/Engine/Source/Runtime/GameplayAbilities/Private/AbilitySystemBlueprintLibrary.cpp
+++ b/Engine/Source/Runtime/GameplayAbilities/Private/AbilitySystemBlueprintLibrary.cpp
@@ -107,9 +107,9 @@ UAbilityTask_WaitOverlap* UAbilitySystemBlueprintLibrary::CreateWaitOverlap(clas
  *		-think of a better way to have a global, overridable per project function. Thats it!
  */
 
-UAttributeComponent* UAbilitySystemBlueprintLibrary::GetAttributeComponent(AActor *Actor)
+UAbilitySystemComponent* UAbilitySystemBlueprintLibrary::GetAbilitySystemComponent(AActor *Actor)
 {
-	return UAbilitySystemGlobals::Get().GetAttributeComponentFromActor(Actor);
+	return UAbilitySystemGlobals::Get().GetAbilitySystemComponentFromActor(Actor);
 }
 
 void UAbilitySystemBlueprintLibrary::ApplyGameplayEffectToTargetData(FGameplayAbilityTargetDataHandle Target, UGameplayEffect *GameplayEffect, const FGameplayAbilityActorInfo InstigatorInfo)
diff --git a/Engine/Source/Runtime/GameplayAbilities/Private/AbilitySystemComponent.cpp b/Engine/Source/Runtime/GameplayAbilities/Private/AbilitySystemComponent.cpp
new file mode 100644
--- /dev/null
+++ b/Engine/Source/Runtime/GameplayAbilities/Private/AbilitySystemComponent.cpp
@@ -0,0 +1,738 @@
+// Copyright 1998-2013 Epic Games, Inc. All Rights Reserved.
+// ActorComponent.cpp: Actor component implementation.
+
+#include "AbilitySystemPrivatePCH.h"
+#include "AbilitySystemComponent.h"
+#include "GameplayCueInterface.h"
+#include "Abilities/GameplayAbility.h"
+
+
+#include "Net/UnrealNetwork.h"
+#include "Engine/ActorChannel.h"
+#include "MessageLog.h"
+#include "UObjectToken.h"
+#include "MapErrors.h"
+
+DEFINE_LOG_CATEGORY(LogAbilitySystemComponent);
+
+#define LOCTEXT_NAMESPACE "AbilitySystemComponent"
+
+
+int32 DebugGameplayCues = 0;
+static FAutoConsoleVariableRef CVarDebugGameplayCues(
+	TEXT("AbilitySystem.DebugGameplayCues"),
+	DebugGameplayCues,
+	TEXT("Enables Debugging for GameplayCue events"),
+	ECVF_Default
+	);
+
+/** Enable to log out all render state create, destroy and updatetransform events */
+#define LOG_RENDER_STATE 0
+
+UAbilitySystemComponent::UAbilitySystemComponent(const class FPostConstructInitializeProperties& PCIP)
+	: Super(PCIP)
+{
+	bWantsInitializeComponent = true;
+
+	PrimaryComponentTick.TickGroup = TG_DuringPhysics;
+	PrimaryComponentTick.bStartWithTickEnabled = true; // FIXME! Just temp until timer manager figured out
+	PrimaryComponentTick.bCanEverTick = true;
+	
+	ActiveGameplayEffects.Owner = this;
+
+	SetIsReplicated(true);
+	ReplicatedPredictionKey = 0;
+	LocalPredictionKey = 0;
+}
+
+UAbilitySystemComponent::~UAbilitySystemComponent()
+{
+	ActiveGameplayEffects.PreDestroy();
+}
+
+UAttributeSet* UAbilitySystemComponent::InitStats(TSubclassOf<class UAttributeSet> Attributes, const UDataTable* DataTable)
+{
+	UAttributeSet* AttributeObj = NULL;
+	if (Attributes)
+	{
+		AttributeObj = GetOrCreateAttributeSubobject(Attributes);
+		if (AttributeObj && DataTable)
+		{
+			AttributeObj->InitFromMetaDataTable(DataTable);
+		}
+	}
+	return AttributeObj;
+}
+
+void UAbilitySystemComponent::ModifyStats(TSubclassOf<class UAttributeSet> AttributeClass, FDataTableRowHandle ModifierHandle)
+{
+	UAttributeSet *Attributes = Cast<UAttributeSet>(GetAttributeSubobject(AttributeClass));
+	if (Attributes)
+	{
+		FAttributeModifierTest *ModifierData = ModifierHandle.GetRow<FAttributeModifierTest>();
+		if (ModifierData)
+		{
+			ModifierData->ApplyModifier(Attributes);
+		}
+	}
+}
+
+UAttributeSet*	UAbilitySystemComponent::GetOrCreateAttributeSubobject(const TSubclassOf<UAttributeSet> AttributeClass)
+{
+	AActor *OwningActor = GetOwner();
+	UAttributeSet *MyAttributes  = NULL;
+	if (OwningActor && AttributeClass)
+	{
+		MyAttributes = GetAttributeSubobject(AttributeClass);
+		if (!MyAttributes)
+		{
+			MyAttributes = ConstructObject<UAttributeSet>(AttributeClass, OwningActor);
+			SpawnedAttributes.AddUnique(MyAttributes);
+		}
+	}
+
+	return MyAttributes;
+}
+
+UAttributeSet* UAbilitySystemComponent::GetAttributeSubobjectChecked(const TSubclassOf<UAttributeSet> AttributeClass) const
+{
+	UAttributeSet *Set = GetAttributeSubobject(AttributeClass);
+	check(Set);
+	return Set;
+}
+
+UAttributeSet* UAbilitySystemComponent::GetAttributeSubobject(const TSubclassOf<UAttributeSet> AttributeClass) const
+{
+	for (UAttributeSet * Set : SpawnedAttributes)
+	{
+		if (Set && Set->IsA(AttributeClass))
+		{
+			return Set;
+		}
+	}
+	return NULL;
+}
+
+void UAbilitySystemComponent::OnRegister()
+{
+	Super::OnRegister();
+
+	// Init starting data
+	for (int32 i=0; i < DefaultStartingData.Num(); ++i)
+	{
+		if (DefaultStartingData[i].Attributes && DefaultStartingData[i].DefaultStartingTable)
+		{
+			UAttributeSet* Attributes = GetOrCreateAttributeSubobject(DefaultStartingData[i].Attributes);
+			Attributes->InitFromMetaDataTable(DefaultStartingData[i].DefaultStartingTable);
+		}
+	}
+}
+
+// ---------------------------------------------------------
+
+bool UAbilitySystemComponent::AreGameplayEffectApplicationRequirementsSatisfied(const class UGameplayEffect* EffectToAdd, FGameplayEffectInstigatorContext& Instigator) const
+{
+	bool bReqsSatisfied = false;
+	if (EffectToAdd)
+	{
+		// Collect gameplay tags from instigator and target to see if requirements are satisfied
+		FGameplayTagContainer InstigatorTags;
+		Instigator.GetOwnedGameplayTags(InstigatorTags);
+
+		FGameplayTagContainer TargetTags;
+		IGameplayTagAssetInterface* OwnerGTA = InterfaceCast<IGameplayTagAssetInterface>(GetOwner());
+		if (OwnerGTA)
+		{
+			OwnerGTA->GetOwnedGameplayTags(TargetTags);
+		}
+
+		bReqsSatisfied = EffectToAdd->AreApplicationTagRequirementsSatisfied(InstigatorTags, TargetTags);
+	}
+
+	return bReqsSatisfied;
+}
+
+// ---------------------------------------------------------
+
+bool UAbilitySystemComponent::IsOwnerActorAuthoritative() const
+{
+	return !IsNetSimulating();
+}
+
+void UAbilitySystemComponent::SetNumericAttribute(const FGameplayAttribute &Attribute, float NewFloatValue)
+{
+	UAttributeSet * AttributeSet = GetAttributeSubobjectChecked(Attribute.GetAttributeSetClass());
+	Attribute.SetNumericValueChecked(NewFloatValue, AttributeSet);
+}
+
+float UAbilitySystemComponent::GetNumericAttribute(const FGameplayAttribute &Attribute)
+{
+	UAttributeSet * AttributeSet = GetAttributeSubobjectChecked(Attribute.GetAttributeSetClass());
+	return Attribute.GetNumericValueChecked(AttributeSet);
+}
+
+/** This is a helper function used in automated testing, I'm not sure how useful it will be to gamecode or blueprints */
+FActiveGameplayEffectHandle UAbilitySystemComponent::ApplyGameplayEffectToTarget(UGameplayEffect *GameplayEffect, UAbilitySystemComponent *Target, float Level, FModifierQualifier BaseQualifier)
+{
+	check(GameplayEffect);
+
+	FGameplayEffectSpec	Spec(GameplayEffect, GetOwner(),  Level, GetCurveDataOverride());
+	
+	return ApplyGameplayEffectSpecToTarget(Spec, Target, BaseQualifier);
+}
+
+/** Helper function since we can't have default/optional values for FModifierQualifier in K2 function */
+FActiveGameplayEffectHandle UAbilitySystemComponent::K2_ApplyGameplayEffectToSelf(const UGameplayEffect *GameplayEffect, float Level, AActor *Instigator)
+{
+	return ApplyGameplayEffectToSelf(GameplayEffect, Level, Instigator);
+}
+
+/** This is a helper function - it seems like this will be useful as a blueprint interface at the least, but Level parameter may need to be expanded */
+FActiveGameplayEffectHandle UAbilitySystemComponent::ApplyGameplayEffectToSelf(const UGameplayEffect *GameplayEffect, float Level, AActor *Instigator, FModifierQualifier BaseQualifier)
+{
+	check(GameplayEffect);
+
+	FGameplayEffectSpec	Spec(GameplayEffect, Instigator, Level, GetCurveDataOverride());
+
+	return ApplyGameplayEffectSpecToSelf(Spec, BaseQualifier);
+}
+
+float UAbilitySystemComponent::GetGameplayEffectMagnitudeByTag(FActiveGameplayEffectHandle InHandle, const FGameplayTag& InTag) const
+{
+	return ActiveGameplayEffects.GetGameplayEffectMagnitudeByTag(InHandle, InTag);
+}
+
+int32 UAbilitySystemComponent::GetNumActiveGameplayEffect() const
+{
+	return ActiveGameplayEffects.GetNumGameplayEffects();
+}
+
+bool UAbilitySystemComponent::IsGameplayEffectActive(FActiveGameplayEffectHandle InHandle) const
+{
+	return ActiveGameplayEffects.IsGameplayEffectActive(InHandle);
+}
+
+bool UAbilitySystemComponent::HasAnyTags(FGameplayTagContainer &Tags)
+{
+	// Fixme: we could aggregate our current tags into a map to avoid this type of iteration
+	return ActiveGameplayEffects.HasAnyTags(Tags);
+}
+
+bool UAbilitySystemComponent::HasAllTags(FGameplayTagContainer &Tags)
+{
+	return ActiveGameplayEffects.HasAllTags(Tags);
+}
+
+void UAbilitySystemComponent::TEMP_ApplyActiveGameplayEffects()
+{
+	for (int32 idx=0; idx < ActiveGameplayEffects.GameplayEffects.Num(); ++idx)
+	{
+		FActiveGameplayEffect& ActiveEffect = ActiveGameplayEffects.GameplayEffects[idx];
+
+		ExecuteGameplayEffect(ActiveEffect.Spec, FModifierQualifier().IgnoreHandle(ActiveEffect.Handle));
+
+		ABILITY_LOG(Log, TEXT("ActiveEffect[%d] %s - Duration: %.2f]"), idx, *ActiveEffect.Spec.ToSimpleString(), ActiveEffect.Spec.GetDuration());
+	}
+}
+
+FActiveGameplayEffectHandle UAbilitySystemComponent::ApplyGameplayEffectSpecToTarget(OUT FGameplayEffectSpec &Spec, UAbilitySystemComponent *Target, FModifierQualifier BaseQualifier)
+{
+	// Apply outgoing Effects to the Spec.
+	// Outgoing immunity may stop the outgoing effect from being applied to the target
+	if (ActiveGameplayEffects.ApplyActiveEffectsTo(Spec, FModifierQualifier(BaseQualifier).Type(EGameplayMod::OutgoingGE)))
+	{
+		return Target->ApplyGameplayEffectSpecToSelf(Spec, BaseQualifier);
+	}
+
+	return FActiveGameplayEffectHandle();
+}
+
+FActiveGameplayEffectHandle UAbilitySystemComponent::ApplyGameplayEffectSpecToSelf(const FGameplayEffectSpec &Spec, FModifierQualifier BaseQualifier)
+{
+	// Temp, only non instant, non periodic GEs can be predictive
+	// Effects with other effects may be a mix so go with non-predictive
+	check((BaseQualifier.PredictionKey() == 0) || (Spec.GetDuration() != UGameplayEffect::INSTANT_APPLICATION && Spec.GetPeriod() == UGameplayEffect::NO_PERIOD));
+
+	// check if the effect being applied actually succeeds
+	float ChanceToApply = Spec.GetChanceToApplyToTarget();
+	if ((ChanceToApply < 1.f - SMALL_NUMBER) && (FMath::FRand() > ChanceToApply))
+	{
+		return FActiveGameplayEffectHandle();
+	}
+
+	// Make sure we create our copy of the spec in the right place first...
+	FActiveGameplayEffectHandle	MyHandle;
+	bool bInvokeGameplayCueApplied = false;
+
+	FGameplayEffectSpec* OurCopyOfSpec = NULL;
+	TSharedPtr<FGameplayEffectSpec> StackSpec;
+	{
+	
+		float Duration = Spec.GetDuration();
+
+		if (Duration != UGameplayEffect::INSTANT_APPLICATION)
+		{
+			// recalculating stacking needs to come before creating the new effect
+			if (Spec.GetStackingType() != EGameplayEffectStackingPolicy::Unlimited)
+			{
+				ActiveGameplayEffects.StacksNeedToRecalculate();
+			}
+			FActiveGameplayEffect &NewActiveEffect = ActiveGameplayEffects.CreateNewActiveGameplayEffect(Spec, BaseQualifier.PredictionKey());
+			MyHandle = NewActiveEffect.Handle;
+			OurCopyOfSpec = &NewActiveEffect.Spec;
+
+			bInvokeGameplayCueApplied = true;
+		}
+		
+		if (!OurCopyOfSpec)
+		{
+			StackSpec = TSharedPtr<FGameplayEffectSpec>(new FGameplayEffectSpec(Spec));
+			OurCopyOfSpec = StackSpec.Get();
+		}
+
+		// Do a 1st order copy of the spec so that we can modify it
+		// (the one passed in is owned by the caller, we can't apply our incoming GEs to it)
+		// Note that at this point the spec has a bunch of modifiers. Those modifiers may
+		// have other modifiers. THOSE modifiers may or may not be copies of whatever.
+		//
+		// In theory, we don't modify 2nd order modifiers after they are 'attached'
+		// Long complex chains can be created but we never say 'Modify a GE that is modding another GE'
+		OurCopyOfSpec->MakeUnique();
+	}
+
+	// Now that we have our own copy, apply our GEs that modify IncomingGEs
+	if (!ActiveGameplayEffects.ApplyActiveEffectsTo(*OurCopyOfSpec, FModifierQualifier(BaseQualifier).Type(EGameplayMod::IncomingGE).IgnoreHandle(MyHandle)))
+	{
+		// We're immune to this effect
+		return FActiveGameplayEffectHandle();
+	}	
+	
+	// Now that we have the final version of this effect, actually apply it if its going to be hanging around
+	if (Spec.GetDuration() != UGameplayEffect::INSTANT_APPLICATION )
+	{
+		if (Spec.GetPeriod() == UGameplayEffect::NO_PERIOD)
+		{
+			ActiveGameplayEffects.ApplySpecToActiveEffectsAndAttributes(*OurCopyOfSpec, FModifierQualifier(BaseQualifier).IgnoreHandle(MyHandle));
+		}
+	}
+	
+	// We still probably want to apply tags and stuff even if instant?
+	if (bInvokeGameplayCueApplied)
+	{
+		// We both added and activated the GameplayCue here.
+		// On the client, who will invoke the gameplay cue from an OnRep, he will need to look at the StartTime to determine
+		// if the Cue was actually added+activated or just added (due to relevancy)
+
+		// Fixme: what if we wanted to scale Cue magnitude based on damage? E.g, scale an cue effect when the GE is buffed?
+		InvokeGameplayCueAdded(*OurCopyOfSpec);
+		InvokeGameplayCueActivated(*OurCopyOfSpec);
+	}
+	
+	// Execute the GE at least once (if instant, this will execute once and be done. If persistent, it was added to ActiveGameplayEffects above)
+	
+	// Execute if this is an instant application effect
+	if (Spec.GetDuration() == UGameplayEffect::INSTANT_APPLICATION)
+	{
+		ExecuteGameplayEffect(*OurCopyOfSpec, FModifierQualifier(BaseQualifier).IgnoreHandle(MyHandle));
+	}
+
+	if (Spec.GetPeriod() != UGameplayEffect::NO_PERIOD && Spec.TargetEffectSpecs.Num() > 0)
+	{
+		ABILITY_LOG(Warning, TEXT("%s is periodic but also applies GameplayEffects to its target. GameplayEffects will only be applied once, not every period."), *Spec.Def->GetPathName());
+	}
+	// todo: this is ignoring the returned handles, should we put them into a TArray and return all of the handles?
+	for (const TSharedRef<FGameplayEffectSpec> TargetSpec : Spec.TargetEffectSpecs)
+	{
+		ApplyGameplayEffectSpecToSelf(TargetSpec.Get(), BaseQualifier);
+	}
+
+	return MyHandle;
+}
+
+void UAbilitySystemComponent::ExecutePeriodicEffect(FActiveGameplayEffectHandle	Handle)
+{
+	ActiveGameplayEffects.ExecutePeriodicGameplayEffect(Handle);
+}
+
+void UAbilitySystemComponent::ExecuteGameplayEffect(const FGameplayEffectSpec &Spec, const FModifierQualifier &QualifierContext)
+{
+	// Should only ever execute effects that are instant application or periodic application
+	// Effects with no period and that aren't instant application should never be executed
+	check( (Spec.GetDuration() == UGameplayEffect::INSTANT_APPLICATION || Spec.GetPeriod() != UGameplayEffect::NO_PERIOD) );
+	
+	ActiveGameplayEffects.ExecuteActiveEffectsFrom(Spec, QualifierContext);
+}
+
+void UAbilitySystemComponent::CheckDurationExpired(FActiveGameplayEffectHandle Handle)
+{
+	ActiveGameplayEffects.CheckDuration(Handle);
+}
+
+bool UAbilitySystemComponent::RemoveActiveGameplayEffect(FActiveGameplayEffectHandle Handle)
+{
+	return ActiveGameplayEffects.RemoveActiveGameplayEffect(Handle);
+}
+
+float UAbilitySystemComponent::GetGameplayEffectDuration(FActiveGameplayEffectHandle Handle) const
+{
+	return ActiveGameplayEffects.GetGameplayEffectDuration(Handle);
+}
+
+float UAbilitySystemComponent::GetGameplayEffectMagnitude(FActiveGameplayEffectHandle Handle, FGameplayAttribute Attribute) const
+{
+	return ActiveGameplayEffects.GetGameplayEffectMagnitude(Handle, Attribute);
+}
+
+void UAbilitySystemComponent::InvokeGameplayCueExecute(const FGameplayEffectSpec &Spec)
+{
+	if (DebugGameplayCues)
+	{
+		ABILITY_LOG(Warning, TEXT("InvokeGameplayCueExecute: %s"), *Spec.ToSimpleString());
+	}
+
+
+	AActor *ActorOwner = GetOwner();
+	IGameplayCueInterface * GameplayCueInterface = InterfaceCast<IGameplayCueInterface>(ActorOwner);
+	if (!GameplayCueInterface)
+	{
+		return;
+	}
+	
+	// FIXME: Replication of level not finished
+	float ExecuteLevel = Spec.ModifierLevel.Get()->IsValid() ? Spec.ModifierLevel.Get()->GetLevel() : 1.f;
+	for (FGameplayEffectCue CueInfo : Spec.Def->GameplayCues)
+	{
+		float NormalizedMagnitude = CueInfo.NormalizeLevel(ExecuteLevel);
+		GameplayCueInterface->GameplayCueExecuted(CueInfo.GameplayCueTags, NormalizedMagnitude, Spec.InstigatorContext);
+
+		if (DebugGameplayCues)
+		{
+			DrawDebugSphere(GetWorld(), Spec.InstigatorContext.HitResult->Location, 30.f, 32, FColor(255, 0, 0), true, 30.f);
+			ABILITY_LOG(Warning, TEXT("   %s"), *CueInfo.GameplayCueTags.ToString());
+		}
+	}
+}
+
+void UAbilitySystemComponent::InvokeGameplayCueActivated(const FGameplayEffectSpec &Spec)
+{
+	AActor *ActorOwner = GetOwner();
+	IGameplayCueInterface * GameplayCueInterface = InterfaceCast<IGameplayCueInterface>(ActorOwner);
+	if (!GameplayCueInterface)
+	{
+		return;
+	}
+
+	// FIXME: Replication of level not finished
+	float ExecuteLevel = Spec.ModifierLevel.Get()->IsValid() ? Spec.ModifierLevel.Get()->GetLevel() : 1.f; 
+	for (FGameplayEffectCue CueInfo : Spec.Def->GameplayCues)
+	{
+		float NormalizedMagnitude = CueInfo.NormalizeLevel(ExecuteLevel);
+		GameplayCueInterface->GameplayCueActivated(CueInfo.GameplayCueTags, NormalizedMagnitude, Spec.InstigatorContext);
+	}
+}
+
+void UAbilitySystemComponent::NetMulticast_InvokeGameplayCueExecuted_Implementation(const FGameplayEffectSpec Spec)
+{
+	InvokeGameplayCueExecute(Spec);
+}
+
+void UAbilitySystemComponent::InvokeGameplayCueAdded(const FGameplayEffectSpec &Spec)
+{
+	AActor *ActorOwner = GetOwner();
+	IGameplayCueInterface * GameplayCueInterface = InterfaceCast<IGameplayCueInterface>(ActorOwner);
+	if (!GameplayCueInterface)
+	{
+		return;
+	}
+
+	// FIXME: Replication of level not finished
+	float ExecuteLevel = Spec.ModifierLevel.Get()->IsValid() ? Spec.ModifierLevel.Get()->GetLevel() : 1.f;
+	for (FGameplayEffectCue CueInfo : Spec.Def->GameplayCues)
+	{
+		float NormalizedMagnitude = CueInfo.NormalizeLevel(ExecuteLevel);
+		GameplayCueInterface->GameplayCueAdded(CueInfo.GameplayCueTags, NormalizedMagnitude, Spec.InstigatorContext);
+	}
+}
+
+void UAbilitySystemComponent::InvokeGameplayCueRemoved(const FGameplayEffectSpec &Spec)
+{
+	AActor *ActorOwner = GetOwner();
+	IGameplayCueInterface * GameplayCueInterface = InterfaceCast<IGameplayCueInterface>(ActorOwner);
+	if (!GameplayCueInterface)
+	{
+		return;
+	}
+
+	// FIXME: Replication of level not finished
+	float ExecuteLevel = Spec.ModifierLevel.Get()->IsValid() ? Spec.ModifierLevel.Get()->GetLevel() : 1.f;
+	for (FGameplayEffectCue CueInfo : Spec.Def->GameplayCues)
+	{
+		float NormalizedMagnitude = CueInfo.NormalizeLevel(ExecuteLevel);
+		GameplayCueInterface->GameplayCueRemoved(CueInfo.GameplayCueTags, NormalizedMagnitude, Spec.InstigatorContext);
+	}
+}
+
+void UAbilitySystemComponent::AddDependancyToAttribute(FGameplayAttribute Attribute, const TWeakPtr<FAggregator> InDependant)
+{
+	ActiveGameplayEffects.AddDependancyToAttribute(Attribute, InDependant);
+}
+
+bool UAbilitySystemComponent::CanApplyAttributeModifiers(const UGameplayEffect *GameplayEffect, float Level, AActor *Instigator)
+{
+	return ActiveGameplayEffects.CanApplyAttributeModifiers(GameplayEffect, Level, Instigator);
+}
+
+TArray<float> UAbilitySystemComponent::GetActiveEffectsTimeRemaining(const FActiveGameplayEffectQuery Query) const
+{
+	return ActiveGameplayEffects.GetActiveEffectsTimeRemaining(Query);
+}
+
+void UAbilitySystemComponent::OnRestackGameplayEffects()
+{
+	ActiveGameplayEffects.RecalculateStacking();
+}
+
+// ---------------------------------------------------------------------------------------
+
+void UAbilitySystemComponent::GetLifetimeReplicatedProps(TArray< FLifetimeProperty > & OutLifetimeProps) const
+{
+	Super::GetLifetimeReplicatedProps(OutLifetimeProps);
+
+	DOREPLIFETIME(UAbilitySystemComponent, SpawnedAttributes);
+	DOREPLIFETIME(UAbilitySystemComponent, ActiveGameplayEffects);
+
+	DOREPLIFETIME(UAbilitySystemComponent, ReplicatedInstancedAbilities);
+	DOREPLIFETIME(UAbilitySystemComponent, ActivatableAbilities);
+
+	DOREPLIFETIME(UAbilitySystemComponent, ReplicatedPredictionKey);
+}
+
+bool UAbilitySystemComponent::ReplicateSubobjects(class UActorChannel *Channel, class FOutBunch *Bunch, FReplicationFlags *RepFlags)
+{
+	bool WroteSomething = Super::ReplicateSubobjects(Channel, Bunch, RepFlags);
+
+	for (UAttributeSet* Set : SpawnedAttributes)
+	{
+		if (Set)
+		{
+			WroteSomething |= Channel->ReplicateSubobject(Set, *Bunch, *RepFlags);
+		}
+	}
+
+	for (UGameplayAbility* Ability : ReplicatedInstancedAbilities)
+	{
+		if (Ability)
+		{
+			WroteSomething |= Channel->ReplicateSubobject(Ability, *Bunch, *RepFlags);
+		}
+	}
+
+	return WroteSomething;
+}
+
+void UAbilitySystemComponent::GetSubobjectsWithStableNamesForNetworking(TArray<UObject*>& Objs)
+{
+	for (UAttributeSet* Set : SpawnedAttributes)
+	{
+		if (Set && Set->IsNameStableForNetworking())
+		{
+			Objs.Add(Set);
+		}
+	}
+}
+
+void UAbilitySystemComponent::PostNetReceive()
+{
+	Super::PostNetReceive();
+}
+
+void UAbilitySystemComponent::OnRep_GameplayEffects()
+{
+
+}
+
+void UAbilitySystemComponent::OnRep_PredictionKey()
+{
+	// Every predictive action we've done up to and including the current value of ReplicatedPredictionKey needs to be wiped
+	int32 idx = 0;
+	for ( ; idx < PredictionDelegates.Num(); ++idx)
+	{
+		if (PredictionDelegates[idx].Key <= ReplicatedPredictionKey)
+		{
+			PredictionDelegates[idx].Value.Broadcast();
+		}
+		else
+		{
+			break;
+		}
+	}
+	if (idx > 0)
+	{
+		// Remove everyone we called Broadcast on in one call
+		PredictionDelegates.RemoveAt(0, idx);
+	}
+}
+
+FAbilitySystemComponentPredictionKeyClear& UAbilitySystemComponent::GetPredictionKeyDelegate(int32 PredictionKey)
+{
+	// See if we already have one for this key
+	for (TPair<int32, FAbilitySystemComponentPredictionKeyClear> & Item : PredictionDelegates)
+	{
+		if (Item.Key == PredictionKey)
+		{
+			return Item.Value;
+		}
+	}
+
+	// Create a new one
+	TPair<int32, FAbilitySystemComponentPredictionKeyClear> NewPair;
+	NewPair.Key = PredictionKey;
+	
+	PredictionDelegates.Add(NewPair);
+	return PredictionDelegates.Last().Value;
+}
+
+int32 UAbilitySystemComponent::GetNextPredictionKey()
+{
+	return ++LocalPredictionKey;
+}
+
+// ---------------------------------------------------------------------------------------
+
+void UAbilitySystemComponent::PrintAllGameplayEffects() const
+{
+	ABILITY_LOG_SCOPE(TEXT("PrintAllGameplayEffects %s"), *GetName());
+	ABILITY_LOG(Log, TEXT("Owner: %s"), *GetOwner()->GetName());
+	ActiveGameplayEffects.PrintAllGameplayEffects();
+}
+
+void FActiveGameplayEffectsContainer::PrintAllGameplayEffects() const
+{
+	ABILITY_LOG_SCOPE(TEXT("ActiveGameplayEffects. Num: %d"), GameplayEffects.Num());
+	for (const FActiveGameplayEffect& Effect : GameplayEffects)
+	{
+		Effect.PrintAll();
+	}
+}
+
+void FActiveGameplayEffect::PrintAll() const
+{
+	ABILITY_LOG(Log, TEXT("Handle: %s"), *Handle.ToString());
+	ABILITY_LOG(Log, TEXT("StartWorldTime: %.2f"), StartWorldTime);
+	Spec.PrintAll();
+}
+
+void FGameplayEffectSpec::PrintAll() const
+{
+	ABILITY_LOG_SCOPE(TEXT("GameplayEffectSpec"));
+	ABILITY_LOG(Log, TEXT("Def: %s"), *Def->GetName());
+	
+	ABILITY_LOG(Log, TEXT("Duration: "));
+	Duration.PrintAll();
+
+	ABILITY_LOG(Log, TEXT("Period:"));
+	Period.PrintAll();
+
+	ABILITY_LOG(Log, TEXT("Modifiers:"));
+	for (const FModifierSpec &Mod : Modifiers)
+	{
+		Mod.PrintAll();
+	}
+}
+
+void FModifierSpec::PrintAll() const
+{
+	ABILITY_LOG_SCOPE(TEXT("ModifierSpec"));
+	ABILITY_LOG(Log, TEXT("Attribute: %s"), *Info.Attribute.GetName());
+	ABILITY_LOG(Log, TEXT("ModifierType: %s"), *EGameplayModToString(Info.ModifierType));
+	ABILITY_LOG(Log, TEXT("ModifierOp: %s"), *EGameplayModOpToString(Info.ModifierOp));
+	ABILITY_LOG(Log, TEXT("EffectType: %s"), *EGameplayModEffectToString(Info.EffectType));
+	ABILITY_LOG(Log, TEXT("RequiredTags: %s"), *Info.RequiredTags.ToString());
+	ABILITY_LOG(Log, TEXT("OwnedTags: %s"), *Info.OwnedTags.ToString());
+	ABILITY_LOG(Log, TEXT("(Base) Magnitude: %s"), *Info.Magnitude.ToSimpleString());
+
+	Aggregator.PrintAll();
+}
+
+void FAggregatorRef::PrintAll() const
+{
+	if (!WeakPtr.IsValid())
+	{
+		ABILITY_LOG(Log, TEXT("Invalid AggregatorRef"));
+		return;
+	}
+
+	if (SharedPtr.IsValid())
+	{
+		ABILITY_LOG(Log, TEXT("HardRef AggregatorRef"));
+	}
+	else
+	{
+		ABILITY_LOG(Log, TEXT("SoftRef AggregatorRef"));
+
+	}
+	
+	Get()->PrintAll();
+}
+
+void FAggregator::PrintAll() const
+{
+	ABILITY_LOG_SCOPE(TEXT("FAggregator 0x%X"), this);
+
+#if SKILL_SYSTEM_AGGREGATOR_DEBUG
+	ABILITY_LOG(Log, TEXT("DebugStr: %s"), *DebugString);
+	ABILITY_LOG(Log, TEXT("Copies (of me): %d"), CopiesMade);
+#endif
+
+	if (Level.IsValid())
+	{
+		ABILITY_LOG_SCOPE(TEXT("LevelInfo"));
+		Level->PrintAll();
+	}
+	else
+	{
+		ABILITY_LOG(Log, TEXT("No Level Data"));
+	}
+
+	{
+		ABILITY_LOG_SCOPE(TEXT("BaseData"));
+		BaseData.PrintAll();
+	}
+
+	{
+		ABILITY_LOG_SCOPE(TEXT("CachedData"));
+		CachedData.PrintAll();
+	}
+	
+	for (int32 i=0; i < EGameplayModOp::Max; ++i)
+	{
+		if (Mods[i].Num() > 0)
+		{
+			ABILITY_LOG_SCOPE(TEXT("%s Mods"), *EGameplayModOpToString(i));
+			for (const FAggregatorRef &Ref : Mods[i])
+			{
+				Ref.PrintAll();
+			}
+		}
+	}
+}
+
+void FGameplayModifierData::PrintAll() const
+{
+	ABILITY_LOG(Log, TEXT("Magnitude: %s"), *Magnitude.ToSimpleString());
+	ABILITY_LOG(Log, TEXT("Tags: %s"), *Tags.ToString());
+}
+
+void FGameplayModifierEvaluatedData::PrintAll() const
+{
+	ABILITY_LOG(Log, TEXT("IsValid: %d"), IsValid); 
+	ABILITY_LOG(Log, TEXT("Magnitude: %.2f"), Magnitude);
+	ABILITY_LOG(Log, TEXT("Tags: %s"), *Tags.ToString());
+}
+
+void FGameplayEffectLevelSpec::PrintAll() const
+{
+	ABILITY_LOG(Log, TEXT("ConstantLevel: %.2f"), ConstantLevel);
+}
+
+#undef LOCTEXT_NAMESPACE
diff --git a/Engine/Source/Runtime/GameplayAbilities/Private/AbilitySystemComponent_Abilities.cpp b/Engine/Source/Runtime/GameplayAbilities/Private/AbilitySystemComponent_Abilities.cpp
new file mode 100644
--- /dev/null
+++ b/Engine/Source/Runtime/GameplayAbilities/Private/AbilitySystemComponent_Abilities.cpp
@@ -0,0 +1,256 @@
+// Copyright 1998-2014 Epic Games, Inc. All Rights Reserved.
+// ActorComponent.cpp: Actor component implementation.
+
+#include "AbilitySystemPrivatePCH.h"
+#include "AbilitySystemComponent.h"
+#include "Abilities/GameplayAbility.h"
+
+#include "Net/UnrealNetwork.h"
+#include "MessageLog.h"
+#include "UObjectToken.h"
+#include "MapErrors.h"
+
+#define LOCTEXT_NAMESPACE "AbilitySystemComponent"
+
+/** Enable to log out all render state create, destroy and updatetransform events */
+#define LOG_RENDER_STATE 0
+
+void UAbilitySystemComponent::InitializeComponent()
+{
+	Super::InitializeComponent();
+
+	InitAbilityActorInfo();
+	
+	// Look for DSO AttributeSets
+	AActor *Owner = GetOwner();
+
+	TArray<UObject*> ChildObjects;
+	GetObjectsWithOuter(Owner, ChildObjects, false, RF_PendingKill);
+	for (UObject * Obj : ChildObjects)
+	{
+		UAttributeSet * Set = Cast<UAttributeSet>(Obj);
+		if (Set)  
+		{
+			UObject * AT = Set->GetArchetype();
+
+			//ensure(SpawnedAttributes.Contains(Set) == false);
+			//ensure(Set->IsDefaultSubobject());
+			SpawnedAttributes.Add(Set);
+		}
+	}
+}
+
+void UAbilitySystemComponent::InitAbilityActorInfo()
+{
+	if (!AbilityActorInfo.IsValid())
+	{
+		// Alloc (and init) a new actor info
+		AbilityActorInfo = TSharedPtr<FGameplayAbilityActorInfo>(UAbilitySystemGlobals::Get().AllocAbilityActorInfo(GetOwner()));
+	}
+	else
+	{
+		// We already have a valid actor info, just reinit it.
+		AbilityActorInfo->InitFromActor(GetOwner());
+	}
+}
+
+UGameplayAbility * UAbilitySystemComponent::CreateNewInstanceOfAbility(UGameplayAbility *Ability)
+{
+	check(Ability);
+	check(Ability->HasAllFlags(RF_ClassDefaultObject));
+
+	AActor * OwnerActor = GetOwner();
+	check(OwnerActor);
+
+	UGameplayAbility * AbilityInstance = ConstructObject<UGameplayAbility>(Ability->GetClass(), OwnerActor);
+	check(AbilityInstance);
+
+	// Add it to one of our instance lists so that it doesn't GC.
+	if (AbilityInstance->GetReplicationPolicy() != EGameplayAbilityReplicationPolicy::ReplicateNone)
+	{
+		ReplicatedInstancedAbilities.Add(AbilityInstance);
+	}
+	else
+	{
+		NonReplicatedInstancedAbilities.Add(AbilityInstance);
+	}
+	
+	return AbilityInstance;
+}
+
+bool UAbilitySystemComponent::ActivateAbility(int32 AbilityIdx)
+{
+	check(AbilityActorInfo.IsValid());
+
+	AActor * OwnerActor = GetOwner();
+	check(OwnerActor);
+
+	if (ActivatableAbilities.IsValidIndex(AbilityIdx))
+	{
+		ActivatableAbilities[AbilityIdx]->InputPressed(0, AbilityActorInfo.Get());
+	}
+
+	return false;
+}
+
+void UAbilitySystemComponent::CancelAbilitiesWithTags(const FGameplayTagContainer Tags, const FGameplayAbilityActorInfo* ActorInfo, const FGameplayAbilityActivationInfo ActivationInfo, UGameplayAbility* Ignore)
+{
+	/**
+	 *	FIXME
+	 *
+	 *	Right now we are cancelling all activatable abilites that match Tags. This includes abilities that might not have been activated in the first place!
+	 *	For instanced-per-actor abilities this is fine. They could check if they were activated/still activating.
+	 *	For non-instanced abilities it is ambigous. We have no way to know 'how many' non-instanced abilities are in flight.
+	 *	Likewise for instanced-per-execution abilities. Though they are present in Replicated/NonReplicatedInstancedAbilities list
+	 */
+
+	struct local
+	{
+		static void CancelAbilitiesWithTags(const FGameplayTagContainer Tags, TArray<UGameplayAbility*> Abilities, const FGameplayAbilityActorInfo* ActorInfo, const FGameplayAbilityActivationInfo ActivationInfo, UGameplayAbility * Ignore)
+		{
+			for (int32 idx=0; idx < Abilities.Num(); ++idx)
+			{
+				UGameplayAbility *Ability = Abilities[idx];
+				if (Ability && (Ability != Ignore) && Ability->AbilityTags.MatchesAny(Tags, false))
+				{
+					Ability->CancelAbility(ActorInfo, ActivationInfo);
+					if (!Ability->HasAnyFlags(RF_ClassDefaultObject) && Ability->GetInstancingPolicy() == EGameplayAbilityInstancingPolicy::InstancedPerExecution)
+					{
+						Ability->MarkPendingKill();
+						Abilities.RemoveAtSwap(idx);
+						idx--;
+					}
+				}
+			}
+		}
+	};
+
+	//local::CancelAbilitiesWithTags(Tags, ReplicatedInstancedAbilities, ActorInfo, Ignore);
+	//local::CancelAbilitiesWithTags(Tags, NonReplicatedInstancedAbilities, ActorInfo, Ignore);
+
+	local::CancelAbilitiesWithTags(Tags, ActivatableAbilities, ActorInfo, ActivationInfo, Ignore);
+}
+
+#if !(UE_BUILD_SHIPPING || UE_BUILD_TEST)
+int32 DenyClientActivation = 0;
+static FAutoConsoleVariableRef CVarDenyClientActivation(
+	TEXT("AbilitySystem.DenyClientActivations"),
+	DenyClientActivation,
+	TEXT("Make server deny the next X ability activations from clients. For testing misprediction."),
+	ECVF_Default
+	);
+#endif
+
+void UAbilitySystemComponent::ServerTryActivateAbility_Implementation(UGameplayAbility * AbilityToActivate, int32 PredictionKey)
+{
+#if !(UE_BUILD_SHIPPING || UE_BUILD_TEST)
+	if (DenyClientActivation > 0)
+	{
+		DenyClientActivation--;
+		ClientActivateAbilityFailed(AbilityToActivate, PredictionKey);
+		return;
+	}
+#endif
+
+	ensure(AbilityToActivate);
+	ensure(AbilityActorInfo.IsValid());
+
+	AActor * OwnerActor = GetOwner();
+	check(OwnerActor);
+
+	UGameplayAbility *InstancedAbility = NULL;
+
+	if (AbilityToActivate->TryActivateAbility(AbilityActorInfo.Get(), PredictionKey, &InstancedAbility))
+	{
+		if (InstancedAbility && InstancedAbility->GetReplicationPolicy() != EGameplayAbilityReplicationPolicy::ReplicateNone)
+		{
+			InstancedAbility->ClientActivateAbilitySucceed(PredictionKey);
+		}
+		else
+		{
+			ClientActivateAbilitySucceed(AbilityToActivate, PredictionKey);
+		}
+
+		// Update our ReplicatedPredictionKey. When the client gets value, he will know his state (actor+all components/subobjects) are up to do date and he can
+		// remove any necessary predictive work.
+
+		if (PredictionKey > 0)
+		{
+			ensure(PredictionKey > ReplicatedPredictionKey);
+			ReplicatedPredictionKey = PredictionKey;
+		}
+	}
+	else
+	{
+		ClientActivateAbilityFailed(AbilityToActivate, PredictionKey);
+	}
+}
+
+bool UAbilitySystemComponent::ServerTryActivateAbility_Validate(UGameplayAbility * AbilityToActivate, int32 PredictionKey)
+{
+	return true;
+}
+
+void UAbilitySystemComponent::ClientActivateAbilityFailed_Implementation(UGameplayAbility * AbilityToActivate, int32 PredictionKey)
+{
+	if (PredictionKey > 0)
+	{
+		for (int32 idx = 0; idx < PredictionDelegates.Num(); ++idx)
+		{
+			int32 Key = PredictionDelegates[idx].Key;
+			if (Key == PredictionKey)
+			{
+				ABILITY_LOG(Warning, TEXT("Failed ActivateAbility, clearing prediction data %d"), PredictionKey);
+
+				PredictionDelegates[idx].Value.Broadcast();
+				PredictionDelegates.RemoveAt(idx);
+				break;
+			}
+			else if(Key >= PredictionKey)
+			{
+				break;
+			}
+		}
+	}
+}
+
+void UAbilitySystemComponent::ClientActivateAbilitySucceed_Implementation(UGameplayAbility* AbilityToActivate, int32 PredictionKey)
+{
+	ensure(AbilityToActivate);
+	ensure(AbilityActorInfo.IsValid());
+
+	AActor * OwnerActor = GetOwner();
+	check(OwnerActor);
+
+	FGameplayAbilityActivationInfo	ActivationInfo(EGameplayAbilityActivationMode::Confirmed, PredictionKey);
+
+	AbilityToActivate->CallActivateAbility(AbilityActorInfo.Get(), ActivationInfo);
+}
+
+void UAbilitySystemComponent::MontageBranchPoint_AbilityDecisionStop()
+{
+	if (AnimatingAbility)
+	{
+		AActor * OwnerActor = GetOwner();
+		check(OwnerActor);
+
+		FGameplayAbilityActorInfo	ActorInfo;
+		ActorInfo.InitFromActor(OwnerActor);
+
+		AnimatingAbility->MontageBranchPoint_AbilityDecisionStop(AbilityActorInfo.Get());
+	}
+}
+
+void UAbilitySystemComponent::MontageBranchPoint_AbilityDecisionStart()
+{
+	if (AnimatingAbility)
+	{
+		AActor * OwnerActor = GetOwner();
+		check(OwnerActor);
+
+		AnimatingAbility->MontageBranchPoint_AbilityDecisionStart(AbilityActorInfo.Get());
+	}
+}
+
+#undef LOCTEXT_NAMESPACE
+
diff --git a/Engine/Source/Runtime/GameplayAbilities/Private/AbilitySystemDebugHUD.cpp b/Engine/Source/Runtime/GameplayAbilities/Private/AbilitySystemDebugHUD.cpp
--- a/Engine/Source/Runtime/GameplayAbilities/Private/AbilitySystemDebugHUD.cpp
+++ b/Engine/Source/Runtime/GameplayAbilities/Private/AbilitySystemDebugHUD.cpp
@@ -1,7 +1,7 @@
 // Copyright 1998-2014 Epic Games, Inc. All Rights Reserved.
 
 #include "AbilitySystemPrivatePCH.h"
-#include "AttributeComponent.h"
+#include "AbilitySystemComponent.h"
 #include "AbilitySystemDebugHUD.h"
 #include "DebugRenderSceneProxy.h"
 
@@ -62,15 +62,15 @@ void AAbilitySystemDebugHUD::DrawDebugHUD(UCanvas* InCanvas, APlayerController*
 
 	if (PC->GetPawn())
 	{
-		UAttributeComponent * AttributeComponent = PC->GetPawn()->FindComponentByClass<UAttributeComponent>();
-		if (AttributeComponent)
+		UAbilitySystemComponent * AbilitySystemComponent = PC->GetPawn()->FindComponentByClass<UAbilitySystemComponent>();
+		if (AbilitySystemComponent)
 		{
-			DrawDebugAttributeComponent(AttributeComponent);
+			DrawDebugAbilitySystemComponent(AbilitySystemComponent);
 		}
 	}
 }
 
-void AAbilitySystemDebugHUD::DrawDebugAttributeComponent(UAttributeComponent *Component)
+void AAbilitySystemDebugHUD::DrawDebugAbilitySystemComponent(UAbilitySystemComponent *Component)
 {
 	UWorld *World = GetWorld();
 	float GameWorldTime = World->GetTimeSeconds();
diff --git a/Engine/Source/Runtime/GameplayAbilities/Private/AbilitySystemGlobals.cpp b/Engine/Source/Runtime/GameplayAbilities/Private/AbilitySystemGlobals.cpp
--- a/Engine/Source/Runtime/GameplayAbilities/Private/AbilitySystemGlobals.cpp
+++ b/Engine/Source/Runtime/GameplayAbilities/Private/AbilitySystemGlobals.cpp
@@ -1,7 +1,7 @@
 // Copyright 1998-2014 Epic Games, Inc. All Rights Reserved.
 
 #include "AbilitySystemPrivatePCH.h"
-#include "AttributeComponent.h"
+#include "AbilitySystemComponent.h"
 #include "AbilitySystemTypes.h"
 
 UAbilitySystemGlobals::UAbilitySystemGlobals(const class FPostConstructInitializeProperties& PCIP)
@@ -63,11 +63,11 @@ void UAbilitySystemGlobals::OnDataTableReimported(UObject* InObject)
 #endif
 
 
-UAttributeComponent * UAbilitySystemGlobals::GetAttributeComponentFromActor(AActor *Actor) const
+UAbilitySystemComponent * UAbilitySystemGlobals::GetAbilitySystemComponentFromActor(AActor *Actor) const
 {
 	if (Actor)
 	{
-		return Actor->FindComponentByClass<UAttributeComponent>();
+		return Actor->FindComponentByClass<UAbilitySystemComponent>();
 	}
 
 	// Caller should check this
diff --git a/Engine/Source/Runtime/GameplayAbilities/Private/AbilitySystemTestAttributeSet.cpp b/Engine/Source/Runtime/GameplayAbilities/Private/AbilitySystemTestAttributeSet.cpp
--- a/Engine/Source/Runtime/GameplayAbilities/Private/AbilitySystemTestAttributeSet.cpp
+++ b/Engine/Source/Runtime/GameplayAbilities/Private/AbilitySystemTestAttributeSet.cpp
@@ -4,7 +4,7 @@
 #include "AbilitySystemTestAttributeSet.h"
 #include "GameplayTagsModule.h"
 #include "GameplayEffectExtension.h"
-#include "AttributeComponent.h"
+#include "AbilitySystemComponent.h"
 
 
 UAbilitySystemTestAttributeSet::UAbilitySystemTestAttributeSet(const class FPostConstructInitializeProperties& PCIP)
@@ -56,7 +56,7 @@ void UAbilitySystemTestAttributeSet::PreAttributeModify(struct FGameplayEffectMo
 		if (Data.EvaluatedData.Magnitude > 0.f)
 		{
 			// Check the source - does he have Crit?
-			UAbilitySystemTestAttributeSet * SourceAttributes = Data.EffectSpec.InstigatorContext.GetOriginalInstigatorAttributeComponent()->GetSet<UAbilitySystemTestAttributeSet>();
+			UAbilitySystemTestAttributeSet * SourceAttributes = Data.EffectSpec.InstigatorContext.GetOriginalInstigatorAbilitySystemComponent()->GetSet<UAbilitySystemTestAttributeSet>();
 			if (SourceAttributes && SourceAttributes->CritChance > 0.f)
 			{
 				if (FMath::FRand() <= SourceAttributes->CritChance)
diff --git a/Engine/Source/Runtime/GameplayAbilities/Private/AbilitySystemTestPawn.cpp b/Engine/Source/Runtime/GameplayAbilities/Private/AbilitySystemTestPawn.cpp
--- a/Engine/Source/Runtime/GameplayAbilities/Private/AbilitySystemTestPawn.cpp
+++ b/Engine/Source/Runtime/GameplayAbilities/Private/AbilitySystemTestPawn.cpp
@@ -3,15 +3,15 @@
 #include "AbilitySystemPrivatePCH.h"
 #include "AbilitySystemTestPawn.h"
 #include "AbilitySystemTestAttributeSet.h"
-#include "AttributeComponent.h"
+#include "AbilitySystemComponent.h"
 
-FName  AAbilitySystemTestPawn::AttributeComponentName(TEXT("AttributeComponent0"));
+FName  AAbilitySystemTestPawn::AbilitySystemComponentName(TEXT("AbilitySystemComponent0"));
 
 AAbilitySystemTestPawn::AAbilitySystemTestPawn(const class FPostConstructInitializeProperties& PCIP)
 : Super(PCIP)
 {
-	AttributeComponent = PCIP.CreateDefaultSubobject<UAttributeComponent>(this, AAbilitySystemTestPawn::AttributeComponentName);
-	AttributeComponent->SetIsReplicated(true);
+	AbilitySystemComponent = PCIP.CreateDefaultSubobject<UAbilitySystemComponent>(this, AAbilitySystemTestPawn::AbilitySystemComponentName);
+	AbilitySystemComponent->SetIsReplicated(true);
 
 	//DefaultAbilitySet = NULL;
 }
@@ -22,12 +22,12 @@ void AAbilitySystemTestPawn::PostInitializeComponents()
 
 	Super::PostInitializeComponents();
 	GameplayCueHandler.Owner = this;
-	AttributeComponent->InitStats(UAbilitySystemTestAttributeSet::StaticClass(), NULL);
+	AbilitySystemComponent->InitStats(UAbilitySystemTestAttributeSet::StaticClass(), NULL);
 
 	/*
 	if (DefaultAbilitySet != NULL)
 	{
-		AttributeComponent->InitializeAbilities(DefaultAbilitySet);
+		AbilitySystemComponent->InitializeAbilities(DefaultAbilitySet);
 	}
 	*/
 }
diff --git a/Engine/Source/Runtime/GameplayAbilities/Private/AttributeComponent.cpp b/Engine/Source/Runtime/GameplayAbilities/Private/AttributeComponent.cpp
deleted file mode 100644
diff --git a/Engine/Source/Runtime/GameplayAbilities/Private/AttributeComponent_Abilities.cpp b/Engine/Source/Runtime/GameplayAbilities/Private/AttributeComponent_Abilities.cpp
deleted file mode 100644
diff --git a/Engine/Source/Runtime/GameplayAbilities/Private/GameplayAbilitySet.cpp b/Engine/Source/Runtime/GameplayAbilities/Private/GameplayAbilitySet.cpp
--- a/Engine/Source/Runtime/GameplayAbilities/Private/GameplayAbilitySet.cpp
+++ b/Engine/Source/Runtime/GameplayAbilities/Private/GameplayAbilitySet.cpp
@@ -1,7 +1,7 @@
 // Copyright 1998-2014 Epic Games, Inc. All Rights Reserved.
 
 #include "AbilitySystemPrivatePCH.h"
-#include "AttributeComponent.h"
+#include "AbilitySystemComponent.h"
 #include "Abilities/GameplayAbility.h"
 #include "GameplayAbilitySet.h"
 
@@ -16,15 +16,15 @@ void InputPressed(int32 AbilityIdx, int32 InputID, TWeakObjectPtr<AActor> ActorO
 	FGameplayAbilityActorInfo	ActorInfo;
 	ActorInfo.InitFromActor(ActorOwner.Get());
 
-	UAttributeComponent * AttributeComponent = ActorInfo.AttributeComponent.Get();
-	if (AttributeComponent)
+	UAbilitySystemComponent * AbilitySystemComponent = ActorInfo.AbilitySystemComponent.Get();
+	if (AbilitySystemComponent)
 	{
-		if (AttributeComponent->ActivatableAbilities.IsValidIndex(AbilityIdx))
+		if (AbilitySystemComponent->ActivatableAbilities.IsValidIndex(AbilityIdx))
 		{
-			UGameplayAbility * Ability = AttributeComponent->ActivatableAbilities[AbilityIdx];
+			UGameplayAbility * Ability = AbilitySystemComponent->ActivatableAbilities[AbilityIdx];
 			if (Ability)
 			{
-				Ability->InputPressed(InputID, AttributeComponent->AbilityActorInfo.Get());
+				Ability->InputPressed(InputID, AbilitySystemComponent->AbilityActorInfo.Get());
 			}
 		}
 	}
@@ -35,15 +35,15 @@ void InputReleased(int32 AbilityIdx, int32 InputID, TWeakObjectPtr<AActor> Actor
 	FGameplayAbilityActorInfo	ActorInfo;
 	ActorInfo.InitFromActor(ActorOwner.Get());
 
-	UAttributeComponent * AttributeComponent = ActorInfo.AttributeComponent.Get();
-	if (AttributeComponent)
+	UAbilitySystemComponent * AbilitySystemComponent = ActorInfo.AbilitySystemComponent.Get();
+	if (AbilitySystemComponent)
 	{
-		if (AttributeComponent->ActivatableAbilities.IsValidIndex(AbilityIdx))
+		if (AbilitySystemComponent->ActivatableAbilities.IsValidIndex(AbilityIdx))
 		{
-			UGameplayAbility * Ability = AttributeComponent->ActivatableAbilities[AbilityIdx];
+			UGameplayAbility * Ability = AbilitySystemComponent->ActivatableAbilities[AbilityIdx];
 			if (Ability)
 			{
-				Ability->InputReleased(InputID, AttributeComponent->AbilityActorInfo.Get());
+				Ability->InputReleased(InputID, AbilitySystemComponent->AbilityActorInfo.Get());
 			}
 		}
 	}
@@ -62,7 +62,7 @@ UGameplayAbilitySet::UGameplayAbilitySet(const class FPostConstructInitializePro
 
 }
 
-void UGameplayAbilitySet::InitializeAbilities(UAttributeComponent *AttributeComponent) const
+void UGameplayAbilitySet::InitializeAbilities(UAbilitySystemComponent *AbilitySystemComponent) const
 {
 	for (const FGameplayAbilityBindInfo &BindInfo : this->Abilities)
 	{
@@ -71,23 +71,23 @@ void UGameplayAbilitySet::InitializeAbilities(UAttributeComponent *AttributeComp
 		{
 			if (Ability->GetInstancingPolicy() == EGameplayAbilityInstancingPolicy::InstancedPerActor)
 			{
-				Ability = AttributeComponent->CreateNewInstanceOfAbility(Ability);
+				Ability = AbilitySystemComponent->CreateNewInstanceOfAbility(Ability);
 			}
 
-			AttributeComponent->ActivatableAbilities.Add(Ability);
+			AbilitySystemComponent->ActivatableAbilities.Add(Ability);
 		}
 	}
 }
 
-void UGameplayAbilitySet::BindInputComponentToAbilities(UInputComponent *InputComponent, UAttributeComponent *AttributeComponent) const
+void UGameplayAbilitySet::BindInputComponentToAbilities(UInputComponent *InputComponent, UAbilitySystemComponent *AbilitySystemComponent) const
 {
-	AActor * OwnerActor = AttributeComponent->GetOwner();
+	AActor * OwnerActor = AbilitySystemComponent->GetOwner();
 	check(OwnerActor);
 
 	for (int32 idx = 0; idx < this->Abilities.Num(); ++idx)
 	{
 		const FGameplayAbilityBindInfo & BindInfo = this->Abilities[idx];
-		UGameplayAbility * Ability = AttributeComponent->ActivatableAbilities[idx];	
+		UGameplayAbility * Ability = AbilitySystemComponent->ActivatableAbilities[idx];	
 		
 		for (const FGameplayAbilityBindInfoCommandIDPair &CommandPair : BindInfo.CommandBindings)
 		{
diff --git a/Engine/Source/Runtime/GameplayAbilities/Private/GameplayEffect.cpp b/Engine/Source/Runtime/GameplayAbilities/Private/GameplayEffect.cpp
--- a/Engine/Source/Runtime/GameplayAbilities/Private/GameplayEffect.cpp
+++ b/Engine/Source/Runtime/GameplayAbilities/Private/GameplayEffect.cpp
@@ -4,7 +4,7 @@
 #include "GameplayEffect.h"
 #include "GameplayEffectExtension.h"
 #include "GameplayEffectStackingExtension.h"
-#include "AttributeComponent.h"
+#include "AbilitySystemComponent.h"
 
 const float UGameplayEffect::INFINITE_DURATION = -1.f;
 const float UGameplayEffect::INSTANT_APPLICATION = 0.f;
@@ -444,9 +444,9 @@ bool FModifierSpec::AreTagRequirementsSatisfied(const FModifierSpec &ModifierToB
 void FGameplayEffectInstigatorContext::AddInstigator(class AActor *InInstigator)
 {
 	Instigator = InInstigator;
-	InstigatorAttributeComponent = NULL;
+	InstigatorAbilitySystemComponent = NULL;
 
-	// Cache off his attribute component.
+	// Cache off his AbilitySystemComponent.
 	if (Instigator)
 	{
 		TArray<UObject*> ChildObjects;
@@ -454,9 +454,9 @@ void FGameplayEffectInstigatorContext::AddInstigator(class AActor *InInstigator)
 
 		for (UObject * Obj : ChildObjects)
 		{
-			if (Obj && Obj->GetClass()->IsChildOf(UAttributeComponent::StaticClass()))
+			if (Obj && Obj->GetClass()->IsChildOf(UAbilitySystemComponent::StaticClass()))
 			{
-				InstigatorAttributeComponent = Cast<UAttributeComponent>(Obj);
+				InstigatorAbilitySystemComponent = Cast<UAbilitySystemComponent>(Obj);
 				break;
 			}
 		}
@@ -483,7 +483,7 @@ bool FGameplayEffectInstigatorContext::NetSerialize(FArchive& Ar, class UPackage
 			{
 				HitResult = TSharedPtr<FHitResult>(new FHitResult());
 			}
-			AddInstigator(Instigator); // Just to initialize InstigatorAttributeComponent
+			AddInstigator(Instigator); // Just to initialize InstigatorAbilitySystemComponent
 		}
 	}
 
@@ -1419,7 +1419,7 @@ float FActiveGameplayEffectsContainer::GetGameplayEffectDuration(FActiveGameplay
 	return UGameplayEffect::INFINITE_DURATION;
 }
 
-void FActiveGameplayEffectsContainer::OnDurationAggregatorDirty(FAggregator* Aggregator, UAttributeComponent* Owner, FActiveGameplayEffectHandle Handle)
+void FActiveGameplayEffectsContainer::OnDurationAggregatorDirty(FAggregator* Aggregator, UAbilitySystemComponent* Owner, FActiveGameplayEffectHandle Handle)
 {
 	for (FActiveGameplayEffect& Effect : GameplayEffects)
 	{
@@ -1428,7 +1428,7 @@ void FActiveGameplayEffectsContainer::OnDurationAggregatorDirty(FAggregator* Agg
 			FTimerManager& TimerManager = Owner->GetWorld()->GetTimerManager();
 			float CurrDuration = Effect.GetDuration();
 			float TimeRemaining = CurrDuration - TimerManager.GetTimerElapsed(Effect.DurationHandle);
-			FTimerDelegate Delegate = FTimerDelegate::CreateUObject(Owner, &UAttributeComponent::CheckDurationExpired, Handle);
+			FTimerDelegate Delegate = FTimerDelegate::CreateUObject(Owner, &UAbilitySystemComponent::CheckDurationExpired, Handle);
 			TimerManager.SetTimer(Effect.DurationHandle, Delegate, CurrDuration, false, FMath::Max(TimeRemaining, 0.f));
 
 			return;
@@ -1518,7 +1518,7 @@ void FActiveGameplayEffectsContainer::StacksNeedToRecalculate()
 	{
 		bNeedToRecalculateStacks = true;
 		FTimerManager& TimerManager = Owner->GetWorld()->GetTimerManager();
-		FTimerDelegate Delegate = FTimerDelegate::CreateUObject(Owner, &UAttributeComponent::OnRestackGameplayEffects);
+		FTimerDelegate Delegate = FTimerDelegate::CreateUObject(Owner, &UAbilitySystemComponent::OnRestackGameplayEffects);
 
 		TimerManager.SetTimerForNextTick(Delegate);
 	}
@@ -1535,7 +1535,7 @@ FActiveGameplayEffect & FActiveGameplayEffectsContainer::CreateNewActiveGameplay
 		if (Spec.GetDuration() > 0.f)
 		{
 			FTimerManager& TimerManager = Owner->GetWorld()->GetTimerManager();
-			FTimerDelegate Delegate = FTimerDelegate::CreateUObject(Owner, &UAttributeComponent::CheckDurationExpired, LastAssignedHandle);
+			FTimerDelegate Delegate = FTimerDelegate::CreateUObject(Owner, &UAbilitySystemComponent::CheckDurationExpired, LastAssignedHandle);
 			TimerManager.SetTimer(NewEffect.DurationHandle, Delegate, Spec.GetDuration(), false, Spec.GetDuration());
 		}
 
@@ -1543,7 +1543,7 @@ FActiveGameplayEffect & FActiveGameplayEffectsContainer::CreateNewActiveGameplay
 		if (Spec.GetPeriod() != UGameplayEffect::NO_PERIOD)
 		{
 			FTimerManager& TimerManager = Owner->GetWorld()->GetTimerManager();
-			FTimerDelegate Delegate = FTimerDelegate::CreateUObject(Owner, &UAttributeComponent::ExecutePeriodicEffect, LastAssignedHandle);
+			FTimerDelegate Delegate = FTimerDelegate::CreateUObject(Owner, &UAbilitySystemComponent::ExecutePeriodicEffect, LastAssignedHandle);
 
 			// If this is a periodic stacking effect, make sure that it's in sync with the others.
 			float FirstDelay = -1.f;
@@ -1581,7 +1581,7 @@ FActiveGameplayEffect & FActiveGameplayEffectsContainer::CreateNewActiveGameplay
 		// Clients predicting should call MarkArrayDirty to force the internal replication map to be rebuilt.
 		MarkArrayDirty();
 
-		Owner->GetPredictionKeyDelegate(InPredictionKey).AddUObject(Owner, &UAttributeComponent::RemoveActiveGameplayEffect_NoReturn, NewEffect.Handle);
+		Owner->GetPredictionKeyDelegate(InPredictionKey).AddUObject(Owner, &UAbilitySystemComponent::RemoveActiveGameplayEffect_NoReturn, NewEffect.Handle);
 	}
 	return NewEffect;
 }
@@ -1704,7 +1704,7 @@ void FActiveGameplayEffectsContainer::CheckDuration(FActiveGameplayEffectHandle
 				float TimeRemaining = TimerManager.GetTimerRemaining(Effect.DurationHandle);
 				if (TimeRemaining <= 0.f)
 				{
-					FTimerDelegate Delegate = FTimerDelegate::CreateUObject(Owner, &UAttributeComponent::CheckDurationExpired, Effect.Handle);
+					FTimerDelegate Delegate = FTimerDelegate::CreateUObject(Owner, &UAbilitySystemComponent::CheckDurationExpired, Effect.Handle);
 					TimerManager.SetTimer(Effect.DurationHandle, Delegate, (Effect.StartWorldTime + Duration) - CurrentTime, false);
 				}
 			}
@@ -1963,7 +1963,7 @@ void FGameplayEffectLevelSpec::RegisterLevelDependancy(TWeakPtr<FAggregator> Own
 {
 	if (Source.IsValid() && Attribute.GetUProperty() != NULL)
 	{
-		UAttributeComponent * SourceComponent = Source->FindComponentByClass<UAttributeComponent>();
+		UAbilitySystemComponent * SourceComponent = Source->FindComponentByClass<UAbilitySystemComponent>();
 		if (SourceComponent)
 		{
 			SourceComponent->AddDependancyToAttribute(Attribute, OwningAggregator);
diff --git a/Engine/Source/Runtime/GameplayAbilities/Private/GameplayEffectExtension_LifestealTest.cpp b/Engine/Source/Runtime/GameplayAbilities/Private/GameplayEffectExtension_LifestealTest.cpp
--- a/Engine/Source/Runtime/GameplayAbilities/Private/GameplayEffectExtension_LifestealTest.cpp
+++ b/Engine/Source/Runtime/GameplayAbilities/Private/GameplayEffectExtension_LifestealTest.cpp
@@ -4,7 +4,7 @@
 #include "GameplayEffectExtension_LifestealTest.h"
 #include "AbilitySystemTestAttributeSet.h"
 #include "GameplayTagsModule.h"
-#include "AttributeComponent.h"
+#include "AbilitySystemComponent.h"
 
 UGameplayEffectExtension_LifestealTest::UGameplayEffectExtension_LifestealTest(const class FPostConstructInitializeProperties& PCIP)
 : Super(PCIP)
@@ -28,7 +28,7 @@ void UGameplayEffectExtension_LifestealTest::PostGameplayEffectExecute(const FGa
 	float HealthToRestore = -DamageDone * LifestealPCT;
 	if (HealthToRestore > 0.f)
 	{
-		UAttributeComponent *Source = Data.EffectSpec.InstigatorContext.GetOriginalInstigatorAttributeComponent();
+		UAbilitySystemComponent *Source = Data.EffectSpec.InstigatorContext.GetOriginalInstigatorAbilitySystemComponent();
 
 		UGameplayEffect * LocalHealthRestore = HealthRestoreGameplayEffect;
 		if (!LocalHealthRestore)
diff --git a/Engine/Source/Runtime/GameplayAbilities/Private/GameplayEffectExtension_ShieldTest.cpp b/Engine/Source/Runtime/GameplayAbilities/Private/GameplayEffectExtension_ShieldTest.cpp
--- a/Engine/Source/Runtime/GameplayAbilities/Private/GameplayEffectExtension_ShieldTest.cpp
+++ b/Engine/Source/Runtime/GameplayAbilities/Private/GameplayEffectExtension_ShieldTest.cpp
@@ -3,7 +3,7 @@
 #include "AbilitySystemPrivatePCH.h"
 #include "GameplayEffectExtension_ShieldTest.h"
 #include "GameplayTagsModule.h"
-#include "AttributeComponent.h"
+#include "AbilitySystemComponent.h"
 #include "AbilitySystemTestAttributeSet.h"
 
 UGameplayEffectExtension_ShieldTest::UGameplayEffectExtension_ShieldTest(const class FPostConstructInitializeProperties& PCIP)
@@ -16,7 +16,7 @@ UGameplayEffectExtension_ShieldTest::UGameplayEffectExtension_ShieldTest(const c
 void UGameplayEffectExtension_ShieldTest::PreGameplayEffectExecute(const FGameplayModifierEvaluatedData &SelfData, FGameplayEffectModCallbackData &Data) const
 {
 	IGameplayTagsModule& GameplayTagsModule = IGameplayTagsModule::Get();
-	UAttributeComponent *Source = Data.EffectSpec.InstigatorContext.GetOriginalInstigatorAttributeComponent();
+	UAbilitySystemComponent *Source = Data.EffectSpec.InstigatorContext.GetOriginalInstigatorAbilitySystemComponent();
 
 	// FIXME: some annoyances here: Damage about to be applied = Data.EvaluatedData.Magnitude = negative. Do some sign flipping here that would make more sense if we were dealing with
 	// 'damage' (positive) instead of 'health' (negative)
diff --git a/Engine/Source/Runtime/GameplayAbilities/Private/GameplayEffectTests.cpp b/Engine/Source/Runtime/GameplayAbilities/Private/GameplayEffectTests.cpp
--- a/Engine/Source/Runtime/GameplayAbilities/Private/GameplayEffectTests.cpp
+++ b/Engine/Source/Runtime/GameplayAbilities/Private/GameplayEffectTests.cpp
@@ -6,7 +6,7 @@
 #include "GameplayEffect.h"
 #include "AttributeSet.h"
 #include "GameplayTagsModule.h"
-#include "AttributeComponent.h"
+#include "AbilitySystemComponent.h"
 #include "GameplayEffectExtension_LifestealTest.h"
 #include "GameplayEffectExtension_ShieldTest.h"
 #include "GameplayEffectStackingExtension_CappedNumberTest.h"
@@ -110,8 +110,8 @@ bool GameplayEffectsTest_InstantDamage(UWorld *World, FAutomationTestBase * Test
 
 	UProperty *HealthProperty = FindFieldChecked<UProperty>(UAbilitySystemTestAttributeSet::StaticClass(), GET_MEMBER_NAME_CHECKED(UAbilitySystemTestAttributeSet, Health));
 
-	UAttributeComponent * SourceComponent = SourceActor->AttributeComponent;
-	UAttributeComponent * DestComponent = DestActor->AttributeComponent;
+	UAbilitySystemComponent * SourceComponent = SourceActor->AbilitySystemComponent;
+	UAbilitySystemComponent * DestComponent = DestActor->AbilitySystemComponent;
 	SourceComponent->GetSet<UAbilitySystemTestAttributeSet>()->Health = StartHealth;
 	DestComponent->GetSet<UAbilitySystemTestAttributeSet>()->Health = StartHealth;
 
@@ -152,8 +152,8 @@ bool GameplayEffectsTest_InstantDamageRemap(UWorld *World, FAutomationTestBase *
 	UProperty *HealthProperty = FindFieldChecked<UProperty>(UAbilitySystemTestAttributeSet::StaticClass(), GET_MEMBER_NAME_CHECKED(UAbilitySystemTestAttributeSet, Health));
 	UProperty *DamageProperty = FindFieldChecked<UProperty>(UAbilitySystemTestAttributeSet::StaticClass(), GET_MEMBER_NAME_CHECKED(UAbilitySystemTestAttributeSet, Damage));
 
-	UAttributeComponent * SourceComponent = SourceActor->AttributeComponent;
-	UAttributeComponent * DestComponent = DestActor->AttributeComponent;
+	UAbilitySystemComponent * SourceComponent = SourceActor->AbilitySystemComponent;
+	UAbilitySystemComponent * DestComponent = DestActor->AbilitySystemComponent;
 	SourceComponent->GetSet<UAbilitySystemTestAttributeSet>()->Health = StartHealth;
 	DestComponent->GetSet<UAbilitySystemTestAttributeSet>()->Health = StartHealth;
 
@@ -203,8 +203,8 @@ bool GameplayEffectsTest_InstantDamage_Buffed(UWorld *World, FAutomationTestBase
 
 	UProperty *HealthProperty = FindFieldChecked<UProperty>(UAbilitySystemTestAttributeSet::StaticClass(), GET_MEMBER_NAME_CHECKED(UAbilitySystemTestAttributeSet, Health));
 
-	UAttributeComponent * SourceComponent = SourceActor->AttributeComponent;
-	UAttributeComponent * DestComponent = DestActor->AttributeComponent;
+	UAbilitySystemComponent * SourceComponent = SourceActor->AbilitySystemComponent;
+	UAbilitySystemComponent * DestComponent = DestActor->AbilitySystemComponent;
 	SourceComponent->GetSet<UAbilitySystemTestAttributeSet>()->Health = StartHealth;
 	DestComponent->GetSet<UAbilitySystemTestAttributeSet>()->Health = StartHealth;
 
@@ -267,8 +267,8 @@ bool GameplayEffectsTest_TemporaryDamage(UWorld *World, FAutomationTestBase * Te
 
 	UProperty *HealthProperty = FindFieldChecked<UProperty>(UAbilitySystemTestAttributeSet::StaticClass(), GET_MEMBER_NAME_CHECKED(UAbilitySystemTestAttributeSet, Health));
 
-	UAttributeComponent * SourceComponent = SourceActor->AttributeComponent;
-	UAttributeComponent * DestComponent = DestActor->AttributeComponent;
+	UAbilitySystemComponent * SourceComponent = SourceActor->AbilitySystemComponent;
+	UAbilitySystemComponent * DestComponent = DestActor->AbilitySystemComponent;
 	SourceComponent->GetSet<UAbilitySystemTestAttributeSet>()->Health = StartHealth;
 	DestComponent->GetSet<UAbilitySystemTestAttributeSet>()->Health = StartHealth;
 
@@ -326,8 +326,8 @@ bool GameplayEffectsTest_TemporaryDamageBuffed(UWorld *World, FAutomationTestBas
 
 	UProperty *HealthProperty = FindFieldChecked<UProperty>(UAbilitySystemTestAttributeSet::StaticClass(), GET_MEMBER_NAME_CHECKED(UAbilitySystemTestAttributeSet, Health));
 
-	UAttributeComponent * SourceComponent = SourceActor->AttributeComponent;
-	UAttributeComponent * DestComponent = DestActor->AttributeComponent;
+	UAbilitySystemComponent * SourceComponent = SourceActor->AbilitySystemComponent;
+	UAbilitySystemComponent * DestComponent = DestActor->AbilitySystemComponent;
 	SourceComponent->GetSet<UAbilitySystemTestAttributeSet>()->Health = StartHealth;
 	DestComponent->GetSet<UAbilitySystemTestAttributeSet>()->Health = StartHealth;
 
@@ -416,8 +416,8 @@ bool GameplayEffectsTest_TemporaryDamageTemporaryBuff(UWorld *World, FAutomation
 
 	UProperty *HealthProperty = FindFieldChecked<UProperty>(UAbilitySystemTestAttributeSet::StaticClass(), GET_MEMBER_NAME_CHECKED(UAbilitySystemTestAttributeSet, Health));
 
-	UAttributeComponent * SourceComponent = SourceActor->AttributeComponent;
-	UAttributeComponent * DestComponent = DestActor->AttributeComponent;
+	UAbilitySystemComponent * SourceComponent = SourceActor->AbilitySystemComponent;
+	UAbilitySystemComponent * DestComponent = DestActor->AbilitySystemComponent;
 	SourceComponent->GetSet<UAbilitySystemTestAttributeSet>()->Health = StartHealth;
 	DestComponent->GetSet<UAbilitySystemTestAttributeSet>()->Health = StartHealth;
 
@@ -525,8 +525,8 @@ bool GameplayEffectsTest_LinkedBuffDestroy(UWorld *World, FAutomationTestBase *
 
 	UProperty *HealthProperty = FindFieldChecked<UProperty>(UAbilitySystemTestAttributeSet::StaticClass(), GET_MEMBER_NAME_CHECKED(UAbilitySystemTestAttributeSet, Health));
 
-	UAttributeComponent * SourceComponent = SourceActor->AttributeComponent;
-	UAttributeComponent * DestComponent = DestActor->AttributeComponent;
+	UAbilitySystemComponent * SourceComponent = SourceActor->AbilitySystemComponent;
+	UAbilitySystemComponent * DestComponent = DestActor->AbilitySystemComponent;
 	SourceComponent->GetSet<UAbilitySystemTestAttributeSet>()->Health = StartHealth;
 	DestComponent->GetSet<UAbilitySystemTestAttributeSet>()->Health = StartHealth;
 
@@ -612,8 +612,8 @@ bool GameplayEffectsTest_SnapshotBuffDestroy(UWorld *World, FAutomationTestBase
 
 	UProperty *HealthProperty = FindFieldChecked<UProperty>(UAbilitySystemTestAttributeSet::StaticClass(), GET_MEMBER_NAME_CHECKED(UAbilitySystemTestAttributeSet, Health));
 
-	UAttributeComponent * SourceComponent = SourceActor->AttributeComponent;
-	UAttributeComponent * DestComponent = DestActor->AttributeComponent;
+	UAbilitySystemComponent * SourceComponent = SourceActor->AbilitySystemComponent;
+	UAbilitySystemComponent * DestComponent = DestActor->AbilitySystemComponent;
 	
 	SourceComponent->GetSet<UAbilitySystemTestAttributeSet>()->Health = StartHealth;
 	DestComponent->GetSet<UAbilitySystemTestAttributeSet>()->Health = StartHealth;
@@ -704,8 +704,8 @@ bool GameplayEffectsTest_DurationBuff(UWorld *World, FAutomationTestBase * Test)
 
 	UProperty *HealthProperty = FindFieldChecked<UProperty>(UAbilitySystemTestAttributeSet::StaticClass(), GET_MEMBER_NAME_CHECKED(UAbilitySystemTestAttributeSet, Health));
 
-	UAttributeComponent * SourceComponent = SourceActor->AttributeComponent;
-	UAttributeComponent * DestComponent = DestActor->AttributeComponent;
+	UAbilitySystemComponent * SourceComponent = SourceActor->AbilitySystemComponent;
+	UAbilitySystemComponent * DestComponent = DestActor->AbilitySystemComponent;
 
 	SourceComponent->GetSet<UAbilitySystemTestAttributeSet>()->Health = StartHealth;
 	DestComponent->GetSet<UAbilitySystemTestAttributeSet>()->Health = StartHealth;
@@ -959,8 +959,8 @@ bool GameplayEffectsTest_DamageBuffBuff_Basic(UWorld *World, FAutomationTestBase
 
 	UProperty *HealthProperty = FindFieldChecked<UProperty>(UAbilitySystemTestAttributeSet::StaticClass(), GET_MEMBER_NAME_CHECKED(UAbilitySystemTestAttributeSet, Health));
 
-	UAttributeComponent * SourceComponent = SourceActor->AttributeComponent;
-	UAttributeComponent * DestComponent = DestActor->AttributeComponent;
+	UAbilitySystemComponent * SourceComponent = SourceActor->AbilitySystemComponent;
+	UAbilitySystemComponent * DestComponent = DestActor->AbilitySystemComponent;
 	SourceComponent->GetSet<UAbilitySystemTestAttributeSet>()->Health = StartHealth;
 	DestComponent->GetSet<UAbilitySystemTestAttributeSet>()->Health = StartHealth;
 
@@ -1100,8 +1100,8 @@ bool GameplayEffectsTest_DamageBuffBuff_FullLink(UWorld *World, FAutomationTestB
 
 	UProperty *HealthProperty = FindFieldChecked<UProperty>(UAbilitySystemTestAttributeSet::StaticClass(), GET_MEMBER_NAME_CHECKED(UAbilitySystemTestAttributeSet, Health));
 
-	UAttributeComponent * SourceComponent = SourceActor->AttributeComponent;
-	UAttributeComponent * DestComponent = DestActor->AttributeComponent;
+	UAbilitySystemComponent * SourceComponent = SourceActor->AbilitySystemComponent;
+	UAbilitySystemComponent * DestComponent = DestActor->AbilitySystemComponent;
 	SourceComponent->GetSet<UAbilitySystemTestAttributeSet>()->Health = StartHealth;
 	DestComponent->GetSet<UAbilitySystemTestAttributeSet>()->Health = StartHealth;
 
@@ -1238,8 +1238,8 @@ bool GameplayEffectsTest_DamageBuffBuff_FullSnapshot(UWorld *World, FAutomationT
 
 	UProperty *HealthProperty = FindFieldChecked<UProperty>(UAbilitySystemTestAttributeSet::StaticClass(), GET_MEMBER_NAME_CHECKED(UAbilitySystemTestAttributeSet, Health));
 
-	UAttributeComponent * SourceComponent = SourceActor->AttributeComponent;
-	UAttributeComponent * DestComponent = DestActor->AttributeComponent;
+	UAbilitySystemComponent * SourceComponent = SourceActor->AbilitySystemComponent;
+	UAbilitySystemComponent * DestComponent = DestActor->AbilitySystemComponent;
 	SourceComponent->GetSet<UAbilitySystemTestAttributeSet>()->Health = StartHealth;
 	DestComponent->GetSet<UAbilitySystemTestAttributeSet>()->Health = StartHealth;
 
@@ -1376,8 +1376,8 @@ bool GameplayEffectsTest_DamageBuffBuff_SnapshotLink(UWorld *World, FAutomationT
 
 	UProperty *HealthProperty = FindFieldChecked<UProperty>(UAbilitySystemTestAttributeSet::StaticClass(), GET_MEMBER_NAME_CHECKED(UAbilitySystemTestAttributeSet, Health));
 
-	UAttributeComponent * SourceComponent = SourceActor->AttributeComponent;
-	UAttributeComponent * DestComponent = DestActor->AttributeComponent;
+	UAbilitySystemComponent * SourceComponent = SourceActor->AbilitySystemComponent;
+	UAbilitySystemComponent * DestComponent = DestActor->AbilitySystemComponent;
 	SourceComponent->GetSet<UAbilitySystemTestAttributeSet>()->Health = StartHealth;
 	DestComponent->GetSet<UAbilitySystemTestAttributeSet>()->Health = StartHealth;
 
@@ -1510,8 +1510,8 @@ bool GameplayEffectsTest_DamageAppliesBuff(UWorld *World, FAutomationTestBase *
 
 	UProperty *DamageProperty = FindFieldChecked<UProperty>(UAbilitySystemTestAttributeSet::StaticClass(), GET_MEMBER_NAME_CHECKED(UAbilitySystemTestAttributeSet, Damage));
 
-	UAttributeComponent * SourceComponent = SourceActor->AttributeComponent;
-	UAttributeComponent * DestComponent = DestActor->AttributeComponent;
+	UAbilitySystemComponent * SourceComponent = SourceActor->AbilitySystemComponent;
+	UAbilitySystemComponent * DestComponent = DestActor->AbilitySystemComponent;
 	SourceComponent->GetSet<UAbilitySystemTestAttributeSet>()->Health = StartHealth;
 	DestComponent->GetSet<UAbilitySystemTestAttributeSet>()->Health = StartHealth;
 
@@ -1587,8 +1587,8 @@ bool GameplayEffectsTest_BuffAppliesBuff(UWorld *World, FAutomationTestBase * Te
 
 	UProperty *DamageProperty = FindFieldChecked<UProperty>(UAbilitySystemTestAttributeSet::StaticClass(), GET_MEMBER_NAME_CHECKED(UAbilitySystemTestAttributeSet, Damage));
 
-	UAttributeComponent * SourceComponent = SourceActor->AttributeComponent;
-	UAttributeComponent * DestComponent = DestActor->AttributeComponent;
+	UAbilitySystemComponent * SourceComponent = SourceActor->AbilitySystemComponent;
+	UAbilitySystemComponent * DestComponent = DestActor->AbilitySystemComponent;
 	SourceComponent->GetSet<UAbilitySystemTestAttributeSet>()->Health = StartHealth;
 	DestComponent->GetSet<UAbilitySystemTestAttributeSet>()->Health = StartHealth;
 
@@ -1718,8 +1718,8 @@ bool GameplayEffectsTest_BuffIndirection(UWorld *World, FAutomationTestBase * Te
 
 	UProperty *DamageProperty = FindFieldChecked<UProperty>(UAbilitySystemTestAttributeSet::StaticClass(), GET_MEMBER_NAME_CHECKED(UAbilitySystemTestAttributeSet, Damage));
 
-	UAttributeComponent * SourceComponent = SourceActor->AttributeComponent;
-	UAttributeComponent * DestComponent = DestActor->AttributeComponent;
+	UAbilitySystemComponent * SourceComponent = SourceActor->AbilitySystemComponent;
+	UAbilitySystemComponent * DestComponent = DestActor->AbilitySystemComponent;
 	SourceComponent->GetSet<UAbilitySystemTestAttributeSet>()->Health = StartHealth;
 	DestComponent->GetSet<UAbilitySystemTestAttributeSet>()->Health = StartHealth;
 
@@ -1836,8 +1836,8 @@ bool GameplayEffectsTest_DurationDamage(UWorld *World, FAutomationTestBase * Tes
 
 	UProperty *HealthProperty = FindFieldChecked<UProperty>(UAbilitySystemTestAttributeSet::StaticClass(), GET_MEMBER_NAME_CHECKED(UAbilitySystemTestAttributeSet, Health));
 
-	UAttributeComponent * SourceComponent = SourceActor->AttributeComponent;
-	UAttributeComponent * DestComponent = DestActor->AttributeComponent;
+	UAbilitySystemComponent * SourceComponent = SourceActor->AbilitySystemComponent;
+	UAbilitySystemComponent * DestComponent = DestActor->AbilitySystemComponent;
 
 	SourceComponent->GetSet<UAbilitySystemTestAttributeSet>()->Health = StartHealth;
 	DestComponent->GetSet<UAbilitySystemTestAttributeSet>()->Health = StartHealth;
@@ -1913,8 +1913,8 @@ bool GameplayEffectsTest_PeriodicDamage(UWorld *World, FAutomationTestBase * Tes
 
 	UProperty *HealthProperty = FindFieldChecked<UProperty>(UAbilitySystemTestAttributeSet::StaticClass(), GET_MEMBER_NAME_CHECKED(UAbilitySystemTestAttributeSet, Health));
 
-	UAttributeComponent * SourceComponent = SourceActor->AttributeComponent;
-	UAttributeComponent * DestComponent = DestActor->AttributeComponent;
+	UAbilitySystemComponent * SourceComponent = SourceActor->AbilitySystemComponent;
+	UAbilitySystemComponent * DestComponent = DestActor->AbilitySystemComponent;
 
 	SourceComponent->GetSet<UAbilitySystemTestAttributeSet>()->Health = StartHealth;
 	DestComponent->GetSet<UAbilitySystemTestAttributeSet>()->Health = StartHealth;
@@ -1988,8 +1988,8 @@ bool GameplayEffectsTest_LifestealExtension(UWorld *World, FAutomationTestBase *
 	UProperty *HealthProperty = FindFieldChecked<UProperty>(UAbilitySystemTestAttributeSet::StaticClass(), GET_MEMBER_NAME_CHECKED(UAbilitySystemTestAttributeSet, Health));
 	UProperty *DamageProperty = FindFieldChecked<UProperty>(UAbilitySystemTestAttributeSet::StaticClass(), GET_MEMBER_NAME_CHECKED(UAbilitySystemTestAttributeSet, Damage));
 
-	UAttributeComponent * SourceComponent = SourceActor->AttributeComponent;
-	UAttributeComponent * DestComponent = DestActor->AttributeComponent;
+	UAbilitySystemComponent * SourceComponent = SourceActor->AbilitySystemComponent;
+	UAbilitySystemComponent * DestComponent = DestActor->AbilitySystemComponent;
 
 	SourceComponent->GetSet<UAbilitySystemTestAttributeSet>()->Health = StartHealth;
 	DestComponent->GetSet<UAbilitySystemTestAttributeSet>()->Health = StartHealth;
@@ -2063,8 +2063,8 @@ bool GameplayEffectsTest_ShieldExtension(UWorld *World, FAutomationTestBase * Te
 	UProperty *HealthProperty = FindFieldChecked<UProperty>(UAbilitySystemTestAttributeSet::StaticClass(), GET_MEMBER_NAME_CHECKED(UAbilitySystemTestAttributeSet, Health));
 	UProperty *DamageProperty = FindFieldChecked<UProperty>(UAbilitySystemTestAttributeSet::StaticClass(), GET_MEMBER_NAME_CHECKED(UAbilitySystemTestAttributeSet, Damage));
 
-	UAttributeComponent * SourceComponent = SourceActor->AttributeComponent;
-	UAttributeComponent * DestComponent = DestActor->AttributeComponent;
+	UAbilitySystemComponent * SourceComponent = SourceActor->AbilitySystemComponent;
+	UAbilitySystemComponent * DestComponent = DestActor->AbilitySystemComponent;
 
 	SourceComponent->GetSet<UAbilitySystemTestAttributeSet>()->Health = StartHealth;
 	DestComponent->GetSet<UAbilitySystemTestAttributeSet>()->Health = StartHealth;
@@ -2178,8 +2178,8 @@ bool GameplayEffectsTest_ShieldExtensionMultiple(UWorld *World, FAutomationTestB
 	UProperty *HealthProperty = FindFieldChecked<UProperty>(UAbilitySystemTestAttributeSet::StaticClass(), GET_MEMBER_NAME_CHECKED(UAbilitySystemTestAttributeSet, Health));
 	UProperty *DamageProperty = FindFieldChecked<UProperty>(UAbilitySystemTestAttributeSet::StaticClass(), GET_MEMBER_NAME_CHECKED(UAbilitySystemTestAttributeSet, Damage));
 
-	UAttributeComponent * SourceComponent = SourceActor->AttributeComponent;
-	UAttributeComponent * DestComponent = DestActor->AttributeComponent;
+	UAbilitySystemComponent * SourceComponent = SourceActor->AbilitySystemComponent;
+	UAbilitySystemComponent * DestComponent = DestActor->AbilitySystemComponent;
 
 	SourceComponent->GetSet<UAbilitySystemTestAttributeSet>()->Health = StartHealth;
 	DestComponent->GetSet<UAbilitySystemTestAttributeSet>()->Health = StartHealth;
@@ -2397,8 +2397,8 @@ bool GameplayEffectsTest_InstantDamage_ScalingExplicit(UWorld *World, FAutomatio
 
 	UProperty *DamageProperty = FindFieldChecked<UProperty>(UAbilitySystemTestAttributeSet::StaticClass(), GET_MEMBER_NAME_CHECKED(UAbilitySystemTestAttributeSet, Damage));
 
-	UAttributeComponent * SourceComponent = SourceActor->AttributeComponent;
-	UAttributeComponent * DestComponent = DestActor->AttributeComponent;
+	UAbilitySystemComponent * SourceComponent = SourceActor->AbilitySystemComponent;
+	UAbilitySystemComponent * DestComponent = DestActor->AbilitySystemComponent;
 	SourceComponent->GetSet<UAbilitySystemTestAttributeSet>()->Health = StartHealth;
 	DestComponent->GetSet<UAbilitySystemTestAttributeSet>()->Health = StartHealth;
 
@@ -2446,8 +2446,8 @@ bool GameplayEffectsTest_InstantDamage_ScalingGlobal(UWorld *World, FAutomationT
 	
 	UProperty *DamageProperty = FindFieldChecked<UProperty>(UAbilitySystemTestAttributeSet::StaticClass(), GET_MEMBER_NAME_CHECKED(UAbilitySystemTestAttributeSet, Damage));
 
-	UAttributeComponent * SourceComponent = SourceActor->AttributeComponent;
-	UAttributeComponent * DestComponent = DestActor->AttributeComponent;
+	UAbilitySystemComponent * SourceComponent = SourceActor->AbilitySystemComponent;
+	UAbilitySystemComponent * DestComponent = DestActor->AbilitySystemComponent;
 	SourceComponent->GetSet<UAbilitySystemTestAttributeSet>()->Health = StartHealth;
 	DestComponent->GetSet<UAbilitySystemTestAttributeSet>()->Health = StartHealth;
 
@@ -2495,8 +2495,8 @@ bool GameplayEffectsTest_InstantDamage_OverrideScaling(UWorld *World, FAutomatio
 
 	UProperty *DamageProperty = FindFieldChecked<UProperty>(UAbilitySystemTestAttributeSet::StaticClass(), GET_MEMBER_NAME_CHECKED(UAbilitySystemTestAttributeSet, Damage));
 
-	UAttributeComponent * SourceComponent = SourceActor->AttributeComponent;
-	UAttributeComponent * DestComponent = DestActor->AttributeComponent;
+	UAbilitySystemComponent * SourceComponent = SourceActor->AbilitySystemComponent;
+	UAbilitySystemComponent * DestComponent = DestActor->AbilitySystemComponent;
 	SourceComponent->GetSet<UAbilitySystemTestAttributeSet>()->Health = StartHealth;
 	DestComponent->GetSet<UAbilitySystemTestAttributeSet>()->Health = StartHealth;
 
@@ -2541,8 +2541,8 @@ bool GameplayEffectsTest_InstantDamageRequiredTag(UWorld *World, FAutomationTest
 
 	UProperty *HealthProperty = FindFieldChecked<UProperty>(UAbilitySystemTestAttributeSet::StaticClass(), GET_MEMBER_NAME_CHECKED(UAbilitySystemTestAttributeSet, Health));
 
-	UAttributeComponent * SourceComponent = SourceActor->AttributeComponent;
-	UAttributeComponent * DestComponent = DestActor->AttributeComponent;
+	UAbilitySystemComponent * SourceComponent = SourceActor->AbilitySystemComponent;
+	UAbilitySystemComponent * DestComponent = DestActor->AbilitySystemComponent;
 	SourceComponent->GetSet<UAbilitySystemTestAttributeSet>()->Health = StartHealth;
 	DestComponent->GetSet<UAbilitySystemTestAttributeSet>()->Health = StartHealth;
 
@@ -2631,8 +2631,8 @@ bool GameplayEffectsTest_InstantDamageIgnoreTag(UWorld *World, FAutomationTestBa
 
 	UProperty *HealthProperty = FindFieldChecked<UProperty>(UAbilitySystemTestAttributeSet::StaticClass(), GET_MEMBER_NAME_CHECKED(UAbilitySystemTestAttributeSet, Health));
 
-	UAttributeComponent * SourceComponent = SourceActor->AttributeComponent;
-	UAttributeComponent * DestComponent = DestActor->AttributeComponent;
+	UAbilitySystemComponent * SourceComponent = SourceActor->AbilitySystemComponent;
+	UAbilitySystemComponent * DestComponent = DestActor->AbilitySystemComponent;
 	SourceComponent->GetSet<UAbilitySystemTestAttributeSet>()->Health = StartHealth;
 	DestComponent->GetSet<UAbilitySystemTestAttributeSet>()->Health = StartHealth;
 
@@ -2722,8 +2722,8 @@ bool GameplayEffectsTest_InstantDamageModifierPassesTag(UWorld *World, FAutomati
 
 	UProperty *HealthProperty = FindFieldChecked<UProperty>(UAbilitySystemTestAttributeSet::StaticClass(), GET_MEMBER_NAME_CHECKED(UAbilitySystemTestAttributeSet, Health));
 
-	UAttributeComponent * SourceComponent = SourceActor->AttributeComponent;
-	UAttributeComponent * DestComponent = DestActor->AttributeComponent;
+	UAbilitySystemComponent * SourceComponent = SourceActor->AbilitySystemComponent;
+	UAbilitySystemComponent * DestComponent = DestActor->AbilitySystemComponent;
 	SourceComponent->GetSet<UAbilitySystemTestAttributeSet>()->Health = StartHealth;
 	DestComponent->GetSet<UAbilitySystemTestAttributeSet>()->Health = StartHealth;
 
@@ -2804,8 +2804,8 @@ bool GameplayEffectsTest_InstantDamageModifierTag(UWorld *World, FAutomationTest
 
 	UProperty *HealthProperty = FindFieldChecked<UProperty>(UAbilitySystemTestAttributeSet::StaticClass(), GET_MEMBER_NAME_CHECKED(UAbilitySystemTestAttributeSet, Health));
 
-	UAttributeComponent * SourceComponent = SourceActor->AttributeComponent;
-	UAttributeComponent * DestComponent = DestActor->AttributeComponent;
+	UAbilitySystemComponent * SourceComponent = SourceActor->AbilitySystemComponent;
+	UAbilitySystemComponent * DestComponent = DestActor->AbilitySystemComponent;
 	SourceComponent->GetSet<UAbilitySystemTestAttributeSet>()->Health = StartHealth;
 	DestComponent->GetSet<UAbilitySystemTestAttributeSet>()->Health = StartHealth;
 
@@ -2890,8 +2890,8 @@ bool GameplayEffectsTest_InstantDamage_ScalingProperty(UWorld *World, FAutomatio
 	UProperty *DamageProperty = FindFieldChecked<UProperty>(UAbilitySystemTestAttributeSet::StaticClass(), GET_MEMBER_NAME_CHECKED(UAbilitySystemTestAttributeSet, Damage));
 	UProperty *PhysicalDamageProperty = FindFieldChecked<UProperty>(UAbilitySystemTestAttributeSet::StaticClass(), GET_MEMBER_NAME_CHECKED(UAbilitySystemTestAttributeSet, PhysicalDamage));
 
-	UAttributeComponent * SourceComponent = SourceActor->AttributeComponent;
-	UAttributeComponent * DestComponent = DestActor->AttributeComponent;
+	UAbilitySystemComponent * SourceComponent = SourceActor->AbilitySystemComponent;
+	UAbilitySystemComponent * DestComponent = DestActor->AbilitySystemComponent;
 	SourceComponent->GetSet<UAbilitySystemTestAttributeSet>()->Health = StartHealth;
 	DestComponent->GetSet<UAbilitySystemTestAttributeSet>()->Health = StartHealth;
 
@@ -2946,8 +2946,8 @@ bool GameplayEffectsTest_InstantDamage_ScalingPropertyNested(UWorld *World, FAut
 	UProperty *DamageProperty = FindFieldChecked<UProperty>(UAbilitySystemTestAttributeSet::StaticClass(), GET_MEMBER_NAME_CHECKED(UAbilitySystemTestAttributeSet, Damage));
 	UProperty *PhysicalDamageProperty = FindFieldChecked<UProperty>(UAbilitySystemTestAttributeSet::StaticClass(), GET_MEMBER_NAME_CHECKED(UAbilitySystemTestAttributeSet, PhysicalDamage));
 
-	UAttributeComponent * SourceComponent = SourceActor->AttributeComponent;
-	UAttributeComponent * DestComponent = DestActor->AttributeComponent;
+	UAbilitySystemComponent * SourceComponent = SourceActor->AbilitySystemComponent;
+	UAbilitySystemComponent * DestComponent = DestActor->AbilitySystemComponent;
 	SourceComponent->GetSet<UAbilitySystemTestAttributeSet>()->Health = StartHealth;
 	DestComponent->GetSet<UAbilitySystemTestAttributeSet>()->Health = StartHealth;
 
@@ -3003,8 +3003,8 @@ bool GameplayEffectsTest_DotDamage_ScalingProperty_Snapshot(UWorld *World, FAuto
 	UProperty *DamageProperty = FindFieldChecked<UProperty>(UAbilitySystemTestAttributeSet::StaticClass(), GET_MEMBER_NAME_CHECKED(UAbilitySystemTestAttributeSet, Damage));
 	UProperty *SpellDamageProperty = FindFieldChecked<UProperty>(UAbilitySystemTestAttributeSet::StaticClass(), GET_MEMBER_NAME_CHECKED(UAbilitySystemTestAttributeSet, SpellDamage));
 
-	UAttributeComponent * SourceComponent = SourceActor->AttributeComponent;
-	UAttributeComponent * DestComponent = DestActor->AttributeComponent;
+	UAbilitySystemComponent * SourceComponent = SourceActor->AbilitySystemComponent;
+	UAbilitySystemComponent * DestComponent = DestActor->AbilitySystemComponent;
 	SourceComponent->GetSet<UAbilitySystemTestAttributeSet>()->Health = StartHealth;
 	DestComponent->GetSet<UAbilitySystemTestAttributeSet>()->Health = StartHealth;
 
@@ -3089,8 +3089,8 @@ bool GameplayEffectsTest_DotDamage_ScalingProperty_Dynamic(UWorld *World, FAutom
 	UProperty *DamageProperty = FindFieldChecked<UProperty>(UAbilitySystemTestAttributeSet::StaticClass(), GET_MEMBER_NAME_CHECKED(UAbilitySystemTestAttributeSet, Damage));
 	UProperty *SpellDamageProperty = FindFieldChecked<UProperty>(UAbilitySystemTestAttributeSet::StaticClass(), GET_MEMBER_NAME_CHECKED(UAbilitySystemTestAttributeSet, SpellDamage));
 
-	UAttributeComponent * SourceComponent = SourceActor->AttributeComponent;
-	UAttributeComponent * DestComponent = DestActor->AttributeComponent;
+	UAbilitySystemComponent * SourceComponent = SourceActor->AbilitySystemComponent;
+	UAbilitySystemComponent * DestComponent = DestActor->AbilitySystemComponent;
 	SourceComponent->GetSet<UAbilitySystemTestAttributeSet>()->Health = StartHealth;
 	DestComponent->GetSet<UAbilitySystemTestAttributeSet>()->Health = StartHealth;
 
@@ -3188,7 +3188,7 @@ bool GameplayEffectsTest_MetaAttributes(UWorld *World, FAutomationTestBase * Tes
 	UProperty *MaxHealthProperty = FindFieldChecked<UProperty>(UAbilitySystemTestAttributeSet::StaticClass(), GET_MEMBER_NAME_CHECKED(UAbilitySystemTestAttributeSet, MaxHealth));
 	UProperty *StrengthProperty = FindFieldChecked<UProperty>(UAbilitySystemTestAttributeSet::StaticClass(), GET_MEMBER_NAME_CHECKED(UAbilitySystemTestAttributeSet, Strength));
 
-	UAttributeComponent * SourceComponent = SourceActor->AttributeComponent;
+	UAbilitySystemComponent * SourceComponent = SourceActor->AbilitySystemComponent;
 	SourceComponent->GetSet<UAbilitySystemTestAttributeSet>()->MaxHealth = StartHealth;
 	SourceComponent->GetSet<UAbilitySystemTestAttributeSet>()->Strength = 0.f;
 
@@ -3265,8 +3265,8 @@ bool GameplayEffectsTest_TagOrdering(UWorld *World, FAutomationTestBase * Test)
 	UProperty *HealthProperty = FindFieldChecked<UProperty>(UAbilitySystemTestAttributeSet::StaticClass(), GET_MEMBER_NAME_CHECKED(UAbilitySystemTestAttributeSet, Health));
 	UProperty *DamageProperty = FindFieldChecked<UProperty>(UAbilitySystemTestAttributeSet::StaticClass(), GET_MEMBER_NAME_CHECKED(UAbilitySystemTestAttributeSet, Damage));
 
-	UAttributeComponent * SourceComponent = SourceActor->AttributeComponent;
-	UAttributeComponent * DestComponent = DestActor->AttributeComponent;
+	UAbilitySystemComponent * SourceComponent = SourceActor->AbilitySystemComponent;
+	UAbilitySystemComponent * DestComponent = DestActor->AbilitySystemComponent;
 	SourceComponent->GetSet<UAbilitySystemTestAttributeSet>()->Health = StartHealth;
 	DestComponent->GetSet<UAbilitySystemTestAttributeSet>()->Health = StartHealth;
 
@@ -3354,8 +3354,8 @@ bool GameplayEffectsTest_StackingHighest(UWorld *World, FAutomationTestBase * Te
 
 	UProperty *StackingProperty = FindFieldChecked<UProperty>(UAbilitySystemTestAttributeSet::StaticClass(), GET_MEMBER_NAME_CHECKED(UAbilitySystemTestAttributeSet, StackingAttribute1));
 
-	UAttributeComponent * SourceComponent = SourceActor->AttributeComponent;
-	UAttributeComponent * DestComponent = DestActor->AttributeComponent;
+	UAbilitySystemComponent * SourceComponent = SourceActor->AbilitySystemComponent;
+	UAbilitySystemComponent * DestComponent = DestActor->AbilitySystemComponent;
 
 	{
 		ABILITY_LOG_SCOPE(TEXT("Apply StackingEffect"));
@@ -3446,8 +3446,8 @@ bool GameplayEffectsTest_StackingLowest(UWorld *World, FAutomationTestBase * Tes
 
 	UProperty *StackingProperty = FindFieldChecked<UProperty>(UAbilitySystemTestAttributeSet::StaticClass(), GET_MEMBER_NAME_CHECKED(UAbilitySystemTestAttributeSet, StackingAttribute1));
 
-	UAttributeComponent * SourceComponent = SourceActor->AttributeComponent;
-	UAttributeComponent * DestComponent = DestActor->AttributeComponent;
+	UAbilitySystemComponent * SourceComponent = SourceActor->AbilitySystemComponent;
+	UAbilitySystemComponent * DestComponent = DestActor->AbilitySystemComponent;
 
 	{
 		ABILITY_LOG_SCOPE(TEXT("Apply StackingEffect"));
@@ -3533,8 +3533,8 @@ bool GameplayEffectsTest_StackingUnlimited(UWorld *World, FAutomationTestBase *
 
 	UProperty *StackingProperty = FindFieldChecked<UProperty>(UAbilitySystemTestAttributeSet::StaticClass(), GET_MEMBER_NAME_CHECKED(UAbilitySystemTestAttributeSet, StackingAttribute1));
 
-	UAttributeComponent * SourceComponent = SourceActor->AttributeComponent;
-	UAttributeComponent * DestComponent = DestActor->AttributeComponent;
+	UAbilitySystemComponent * SourceComponent = SourceActor->AbilitySystemComponent;
+	UAbilitySystemComponent * DestComponent = DestActor->AbilitySystemComponent;
 
 	{
 		ABILITY_LOG_SCOPE(TEXT("Apply StackingEffect"));
@@ -3608,8 +3608,8 @@ bool GameplayEffectsTest_StackingRemoval(UWorld *World, FAutomationTestBase * Te
 
 	UProperty *StackingProperty = FindFieldChecked<UProperty>(UAbilitySystemTestAttributeSet::StaticClass(), GET_MEMBER_NAME_CHECKED(UAbilitySystemTestAttributeSet, StackingAttribute1));
 
-	UAttributeComponent * SourceComponent = SourceActor->AttributeComponent;
-	UAttributeComponent * DestComponent = DestActor->AttributeComponent;
+	UAbilitySystemComponent * SourceComponent = SourceActor->AbilitySystemComponent;
+	UAbilitySystemComponent * DestComponent = DestActor->AbilitySystemComponent;
 
 	{
 		ABILITY_LOG_SCOPE(TEXT("Apply StackingEffect"));
@@ -3690,8 +3690,8 @@ bool GameplayEffectsTest_StackingNoStack(UWorld *World, FAutomationTestBase * Te
 
 	UProperty *NoStackProperty = FindFieldChecked<UProperty>(UAbilitySystemTestAttributeSet::StaticClass(), GET_MEMBER_NAME_CHECKED(UAbilitySystemTestAttributeSet, NoStackAttribute));
 
-	UAttributeComponent * SourceComponent = SourceActor->AttributeComponent;
-	UAttributeComponent * DestComponent = DestActor->AttributeComponent;
+	UAbilitySystemComponent * SourceComponent = SourceActor->AbilitySystemComponent;
+	UAbilitySystemComponent * DestComponent = DestActor->AbilitySystemComponent;
 
 	{
 		ABILITY_LOG_SCOPE(TEXT("Apply StackingEffect"));
@@ -3765,8 +3765,8 @@ bool GameplayEffectsTest_StackingCustomCapped(UWorld *World, FAutomationTestBase
 
 	UProperty *StackingProperty = FindFieldChecked<UProperty>(UAbilitySystemTestAttributeSet::StaticClass(), GET_MEMBER_NAME_CHECKED(UAbilitySystemTestAttributeSet, StackingAttribute1));
 
-	UAttributeComponent * SourceComponent = SourceActor->AttributeComponent;
-	UAttributeComponent * DestComponent = DestActor->AttributeComponent;
+	UAbilitySystemComponent * SourceComponent = SourceActor->AbilitySystemComponent;
+	UAbilitySystemComponent * DestComponent = DestActor->AbilitySystemComponent;
 
 	{
 		ABILITY_LOG_SCOPE(TEXT("Apply StackingEffect"));
@@ -3916,8 +3916,8 @@ bool GameplayEffectsTest_StackingCustomDiminishingReturns(UWorld *World, FAutoma
 
 	UProperty *StackingProperty = FindFieldChecked<UProperty>(UAbilitySystemTestAttributeSet::StaticClass(), GET_MEMBER_NAME_CHECKED(UAbilitySystemTestAttributeSet, StackingAttribute1));
 
-	UAttributeComponent * SourceComponent = SourceActor->AttributeComponent;
-	UAttributeComponent * DestComponent = DestActor->AttributeComponent;
+	UAbilitySystemComponent * SourceComponent = SourceActor->AbilitySystemComponent;
+	UAbilitySystemComponent * DestComponent = DestActor->AbilitySystemComponent;
 
 	{
 		ABILITY_LOG_SCOPE(TEXT("Apply StackingEffect"));
@@ -4073,8 +4073,8 @@ bool GameplayEffectsTest_StackingDifferentRules(UWorld *World, FAutomationTestBa
 
 	UProperty *StackingProperty = FindFieldChecked<UProperty>(UAbilitySystemTestAttributeSet::StaticClass(), GET_MEMBER_NAME_CHECKED(UAbilitySystemTestAttributeSet, StackingAttribute1));
 
-	UAttributeComponent * SourceComponent = SourceActor->AttributeComponent;
-	UAttributeComponent * DestComponent = DestActor->AttributeComponent;
+	UAbilitySystemComponent * SourceComponent = SourceActor->AbilitySystemComponent;
+	UAbilitySystemComponent * DestComponent = DestActor->AbilitySystemComponent;
 
 	{
 		ABILITY_LOG_SCOPE(TEXT("Apply StackingEffect"));
@@ -4149,8 +4149,8 @@ bool GameplayEffectsTest_StackingDifferentAttributes(UWorld *World, FAutomationT
 	UProperty *StackingProperty1 = FindFieldChecked<UProperty>(UAbilitySystemTestAttributeSet::StaticClass(), GET_MEMBER_NAME_CHECKED(UAbilitySystemTestAttributeSet, StackingAttribute1));
 	UProperty *StackingProperty2 = FindFieldChecked<UProperty>(UAbilitySystemTestAttributeSet::StaticClass(), GET_MEMBER_NAME_CHECKED(UAbilitySystemTestAttributeSet, StackingAttribute2));
 
-	UAttributeComponent * SourceComponent = SourceActor->AttributeComponent;
-	UAttributeComponent * DestComponent = DestActor->AttributeComponent;
+	UAbilitySystemComponent * SourceComponent = SourceActor->AbilitySystemComponent;
+	UAbilitySystemComponent * DestComponent = DestActor->AbilitySystemComponent;
 
 	{
 		ABILITY_LOG_SCOPE(TEXT("Apply StackingEffect"));
@@ -4227,8 +4227,8 @@ bool GameplayEffectsTest_StackingCustomTwoRules(UWorld *World, FAutomationTestBa
 
 	UProperty *StackingProperty = FindFieldChecked<UProperty>(UAbilitySystemTestAttributeSet::StaticClass(), GET_MEMBER_NAME_CHECKED(UAbilitySystemTestAttributeSet, StackingAttribute1));
 
-	UAttributeComponent * SourceComponent = SourceActor->AttributeComponent;
-	UAttributeComponent * DestComponent = DestActor->AttributeComponent;
+	UAbilitySystemComponent * SourceComponent = SourceActor->AbilitySystemComponent;
+	UAbilitySystemComponent * DestComponent = DestActor->AbilitySystemComponent;
 
 	{
 		ABILITY_LOG_SCOPE(TEXT("Apply StackingEffect"));
@@ -4328,8 +4328,8 @@ bool GameplayEffectsTest_StackingCustomTwoAttributes(UWorld *World, FAutomationT
 	UProperty *StackingProperty1 = FindFieldChecked<UProperty>(UAbilitySystemTestAttributeSet::StaticClass(), GET_MEMBER_NAME_CHECKED(UAbilitySystemTestAttributeSet, StackingAttribute1));
 	UProperty *StackingProperty2 = FindFieldChecked<UProperty>(UAbilitySystemTestAttributeSet::StaticClass(), GET_MEMBER_NAME_CHECKED(UAbilitySystemTestAttributeSet, StackingAttribute2));
 
-	UAttributeComponent * SourceComponent = SourceActor->AttributeComponent;
-	UAttributeComponent * DestComponent = DestActor->AttributeComponent;
+	UAbilitySystemComponent * SourceComponent = SourceActor->AbilitySystemComponent;
+	UAbilitySystemComponent * DestComponent = DestActor->AbilitySystemComponent;
 
 	{
 		ABILITY_LOG_SCOPE(TEXT("Apply StackingEffect"));
@@ -4432,8 +4432,8 @@ bool GameplayEffectsTest_StackingRemovingModifiers(UWorld *World, FAutomationTes
 
 	UProperty *StackingProperty = FindFieldChecked<UProperty>(UAbilitySystemTestAttributeSet::StaticClass(), GET_MEMBER_NAME_CHECKED(UAbilitySystemTestAttributeSet, StackingAttribute1));
 
-	UAttributeComponent * SourceComponent = SourceActor->AttributeComponent;
-	UAttributeComponent * DestComponent = DestActor->AttributeComponent;
+	UAbilitySystemComponent * SourceComponent = SourceActor->AbilitySystemComponent;
+	UAbilitySystemComponent * DestComponent = DestActor->AbilitySystemComponent;
 
 	// Setup a GE to modify IncomingGEs
 	{
@@ -4534,8 +4534,8 @@ bool GameplayEffectsTest_StackingAddingModifiers(UWorld *World, FAutomationTestB
 
 	UProperty *StackingProperty = FindFieldChecked<UProperty>(UAbilitySystemTestAttributeSet::StaticClass(), GET_MEMBER_NAME_CHECKED(UAbilitySystemTestAttributeSet, StackingAttribute1));
 
-	UAttributeComponent * SourceComponent = SourceActor->AttributeComponent;
-	UAttributeComponent * DestComponent = DestActor->AttributeComponent;
+	UAbilitySystemComponent * SourceComponent = SourceActor->AbilitySystemComponent;
+	UAbilitySystemComponent * DestComponent = DestActor->AbilitySystemComponent;
 
 	{
 		ABILITY_LOG_SCOPE(TEXT("Apply StackingEffect"));
@@ -4639,8 +4639,8 @@ bool GameplayEffectsTest_ImmunityIncoming(UWorld *World, FAutomationTestBase * T
 
 	UProperty *DamageProperty = FindFieldChecked<UProperty>(UAbilitySystemTestAttributeSet::StaticClass(), GET_MEMBER_NAME_CHECKED(UAbilitySystemTestAttributeSet, Damage));
 
-	UAttributeComponent * SourceComponent = SourceActor->AttributeComponent;
-	UAttributeComponent * DestComponent = DestActor->AttributeComponent;
+	UAbilitySystemComponent * SourceComponent = SourceActor->AbilitySystemComponent;
+	UAbilitySystemComponent * DestComponent = DestActor->AbilitySystemComponent;
 	SourceComponent->GetSet<UAbilitySystemTestAttributeSet>()->Health = StartHealth;
 	DestComponent->GetSet<UAbilitySystemTestAttributeSet>()->Health = StartHealth;
 
@@ -4697,8 +4697,8 @@ bool GameplayEffectsTest_ImmunityOutgoing(UWorld *World, FAutomationTestBase * T
 
 	UProperty *DamageProperty = FindFieldChecked<UProperty>(UAbilitySystemTestAttributeSet::StaticClass(), GET_MEMBER_NAME_CHECKED(UAbilitySystemTestAttributeSet, Damage));
 
-	UAttributeComponent * SourceComponent = SourceActor->AttributeComponent;
-	UAttributeComponent * DestComponent = DestActor->AttributeComponent;
+	UAbilitySystemComponent * SourceComponent = SourceActor->AbilitySystemComponent;
+	UAbilitySystemComponent * DestComponent = DestActor->AbilitySystemComponent;
 	SourceComponent->GetSet<UAbilitySystemTestAttributeSet>()->Health = StartHealth;
 	DestComponent->GetSet<UAbilitySystemTestAttributeSet>()->Health = StartHealth;
 
@@ -4758,8 +4758,8 @@ bool GameplayEffectsTest_ImmunityMod(UWorld *World, FAutomationTestBase * Test)
 
 	UProperty *DamageProperty = FindFieldChecked<UProperty>(UAbilitySystemTestAttributeSet::StaticClass(), GET_MEMBER_NAME_CHECKED(UAbilitySystemTestAttributeSet, Damage));
 
-	UAttributeComponent * SourceComponent = SourceActor->AttributeComponent;
-	UAttributeComponent * DestComponent = DestActor->AttributeComponent;
+	UAbilitySystemComponent * SourceComponent = SourceActor->AbilitySystemComponent;
+	UAbilitySystemComponent * DestComponent = DestActor->AbilitySystemComponent;
 	SourceComponent->GetSet<UAbilitySystemTestAttributeSet>()->Health = StartHealth;
 	DestComponent->GetSet<UAbilitySystemTestAttributeSet>()->Health = StartHealth;
 
@@ -4834,8 +4834,8 @@ bool GameplayEffectsTest_ImmunityBlockedBuff(UWorld *World, FAutomationTestBase
 
 	UProperty *DamageProperty = FindFieldChecked<UProperty>(UAbilitySystemTestAttributeSet::StaticClass(), GET_MEMBER_NAME_CHECKED(UAbilitySystemTestAttributeSet, Damage));
 
-	UAttributeComponent * SourceComponent = SourceActor->AttributeComponent;
-	UAttributeComponent * DestComponent = DestActor->AttributeComponent;
+	UAbilitySystemComponent * SourceComponent = SourceActor->AbilitySystemComponent;
+	UAbilitySystemComponent * DestComponent = DestActor->AbilitySystemComponent;
 	SourceComponent->GetSet<UAbilitySystemTestAttributeSet>()->Health = StartHealth;
 	DestComponent->GetSet<UAbilitySystemTestAttributeSet>()->Health = StartHealth;
 
@@ -4903,8 +4903,8 @@ bool GameplayEffectsTest_ImmunityBlockedBaseAndBuff(UWorld *World, FAutomationTe
 
 	UProperty *DamageProperty = FindFieldChecked<UProperty>(UAbilitySystemTestAttributeSet::StaticClass(), GET_MEMBER_NAME_CHECKED(UAbilitySystemTestAttributeSet, Damage));
 
-	UAttributeComponent * SourceComponent = SourceActor->AttributeComponent;
-	UAttributeComponent * DestComponent = DestActor->AttributeComponent;
+	UAbilitySystemComponent * SourceComponent = SourceActor->AbilitySystemComponent;
+	UAbilitySystemComponent * DestComponent = DestActor->AbilitySystemComponent;
 	SourceComponent->GetSet<UAbilitySystemTestAttributeSet>()->Health = StartHealth;
 	DestComponent->GetSet<UAbilitySystemTestAttributeSet>()->Health = StartHealth;
 
@@ -4972,8 +4972,8 @@ bool GameplayEffectsTest_ImmunityActiveGE(UWorld *World, FAutomationTestBase * T
 
 	UProperty *DamageProperty = FindFieldChecked<UProperty>(UAbilitySystemTestAttributeSet::StaticClass(), GET_MEMBER_NAME_CHECKED(UAbilitySystemTestAttributeSet, Damage));
 
-	UAttributeComponent * SourceComponent = SourceActor->AttributeComponent;
-	UAttributeComponent * DestComponent = DestActor->AttributeComponent;
+	UAbilitySystemComponent * SourceComponent = SourceActor->AbilitySystemComponent;
+	UAbilitySystemComponent * DestComponent = DestActor->AbilitySystemComponent;
 	SourceComponent->GetSet<UAbilitySystemTestAttributeSet>()->Health = StartHealth;
 	DestComponent->GetSet<UAbilitySystemTestAttributeSet>()->Health = StartHealth;
 
@@ -5069,8 +5069,8 @@ bool GameplayEffectsTest_ImmunityActiveMod(UWorld *World, FAutomationTestBase *
 
 	UProperty *DamageProperty = FindFieldChecked<UProperty>(UAbilitySystemTestAttributeSet::StaticClass(), GET_MEMBER_NAME_CHECKED(UAbilitySystemTestAttributeSet, Damage));
 
-	UAttributeComponent * SourceComponent = SourceActor->AttributeComponent;
-	UAttributeComponent * DestComponent = DestActor->AttributeComponent;
+	UAbilitySystemComponent * SourceComponent = SourceActor->AbilitySystemComponent;
+	UAbilitySystemComponent * DestComponent = DestActor->AbilitySystemComponent;
 	SourceComponent->GetSet<UAbilitySystemTestAttributeSet>()->Health = StartHealth;
 	DestComponent->GetSet<UAbilitySystemTestAttributeSet>()->Health = StartHealth;
 
@@ -5165,8 +5165,8 @@ bool GameplayEffectsTest_ChanceToApplyToTarget(UWorld *World, FAutomationTestBas
 
 	UProperty *DamageProperty = FindFieldChecked<UProperty>(UAbilitySystemTestAttributeSet::StaticClass(), GET_MEMBER_NAME_CHECKED(UAbilitySystemTestAttributeSet, Damage));
 
-	UAttributeComponent * SourceComponent = SourceActor->AttributeComponent;
-	UAttributeComponent * DestComponent = DestActor->AttributeComponent;
+	UAbilitySystemComponent * SourceComponent = SourceActor->AbilitySystemComponent;
+	UAbilitySystemComponent * DestComponent = DestActor->AbilitySystemComponent;
 
 	SourceComponent->GetSet<UAbilitySystemTestAttributeSet>()->Health = StartHealth;
 	DestComponent->GetSet<UAbilitySystemTestAttributeSet>()->Health = StartHealth;
@@ -5234,8 +5234,8 @@ bool GameplayEffectsTest_ChanceToExecuteOnActiveGEMod(UWorld *World, FAutomation
 
 	UProperty *DamageProperty = FindFieldChecked<UProperty>(UAbilitySystemTestAttributeSet::StaticClass(), GET_MEMBER_NAME_CHECKED(UAbilitySystemTestAttributeSet, Damage));
 
-	UAttributeComponent * SourceComponent = SourceActor->AttributeComponent;
-	UAttributeComponent * DestComponent = DestActor->AttributeComponent;
+	UAbilitySystemComponent * SourceComponent = SourceActor->AbilitySystemComponent;
+	UAbilitySystemComponent * DestComponent = DestActor->AbilitySystemComponent;
 
 	SourceComponent->GetSet<UAbilitySystemTestAttributeSet>()->Health = StartHealth;
 	DestComponent->GetSet<UAbilitySystemTestAttributeSet>()->Health = StartHealth;
@@ -5351,8 +5351,8 @@ bool GameplayEffectsTest_ChanceToExecuteOnActiveGEImmunity(UWorld *World, FAutom
 
 	UProperty *DamageProperty = FindFieldChecked<UProperty>(UAbilitySystemTestAttributeSet::StaticClass(), GET_MEMBER_NAME_CHECKED(UAbilitySystemTestAttributeSet, Damage));
 
-	UAttributeComponent * SourceComponent = SourceActor->AttributeComponent;
-	UAttributeComponent * DestComponent = DestActor->AttributeComponent;
+	UAbilitySystemComponent * SourceComponent = SourceActor->AbilitySystemComponent;
+	UAbilitySystemComponent * DestComponent = DestActor->AbilitySystemComponent;
 
 	SourceComponent->GetSet<UAbilitySystemTestAttributeSet>()->Health = StartHealth;
 	DestComponent->GetSet<UAbilitySystemTestAttributeSet>()->Health = StartHealth;
@@ -5462,8 +5462,8 @@ bool GameplayEffectsTest_ChanceToExecuteOnOutgoingGEMod(UWorld *World, FAutomati
 
 	UProperty *DamageProperty = FindFieldChecked<UProperty>(UAbilitySystemTestAttributeSet::StaticClass(), GET_MEMBER_NAME_CHECKED(UAbilitySystemTestAttributeSet, Damage));
 	
-	UAttributeComponent * SourceComponent = SourceActor->AttributeComponent;
-	UAttributeComponent * DestComponent = DestActor->AttributeComponent;
+	UAbilitySystemComponent * SourceComponent = SourceActor->AbilitySystemComponent;
+	UAbilitySystemComponent * DestComponent = DestActor->AbilitySystemComponent;
 	SourceComponent->GetSet<UAbilitySystemTestAttributeSet>()->Health = StartHealth;
 	DestComponent->GetSet<UAbilitySystemTestAttributeSet>()->Health = StartHealth;
 
@@ -5565,8 +5565,8 @@ bool GameplayEffectsTest_ChanceToExecuteOnOutgoingGEImmunity(UWorld *World, FAut
 
 	UProperty *DamageProperty = FindFieldChecked<UProperty>(UAbilitySystemTestAttributeSet::StaticClass(), GET_MEMBER_NAME_CHECKED(UAbilitySystemTestAttributeSet, Damage));
 
-	UAttributeComponent * SourceComponent = SourceActor->AttributeComponent;
-	UAttributeComponent * DestComponent = DestActor->AttributeComponent;
+	UAbilitySystemComponent * SourceComponent = SourceActor->AbilitySystemComponent;
+	UAbilitySystemComponent * DestComponent = DestActor->AbilitySystemComponent;
 	SourceComponent->GetSet<UAbilitySystemTestAttributeSet>()->Health = StartHealth;
 	DestComponent->GetSet<UAbilitySystemTestAttributeSet>()->Health = StartHealth;
 
@@ -5658,8 +5658,8 @@ bool GameplayEffectsTest_ChanceToExecuteOnIncomingGEMod(UWorld *World, FAutomati
 
 	UProperty *DamageProperty = FindFieldChecked<UProperty>(UAbilitySystemTestAttributeSet::StaticClass(), GET_MEMBER_NAME_CHECKED(UAbilitySystemTestAttributeSet, Damage));
 
-	UAttributeComponent * SourceComponent = SourceActor->AttributeComponent;
-	UAttributeComponent * DestComponent = DestActor->AttributeComponent;
+	UAbilitySystemComponent * SourceComponent = SourceActor->AbilitySystemComponent;
+	UAbilitySystemComponent * DestComponent = DestActor->AbilitySystemComponent;
 	SourceComponent->GetSet<UAbilitySystemTestAttributeSet>()->Health = StartHealth;
 	DestComponent->GetSet<UAbilitySystemTestAttributeSet>()->Health = StartHealth;
 
@@ -5761,8 +5761,8 @@ bool GameplayEffectsTest_ChanceToExecuteOnIncomingGEImmunity(UWorld *World, FAut
 
 	UProperty *DamageProperty = FindFieldChecked<UProperty>(UAbilitySystemTestAttributeSet::StaticClass(), GET_MEMBER_NAME_CHECKED(UAbilitySystemTestAttributeSet, Damage));
 
-	UAttributeComponent * SourceComponent = SourceActor->AttributeComponent;
-	UAttributeComponent * DestComponent = DestActor->AttributeComponent;
+	UAbilitySystemComponent * SourceComponent = SourceActor->AbilitySystemComponent;
+	UAbilitySystemComponent * DestComponent = DestActor->AbilitySystemComponent;
 	SourceComponent->GetSet<UAbilitySystemTestAttributeSet>()->Health = StartHealth;
 	DestComponent->GetSet<UAbilitySystemTestAttributeSet>()->Health = StartHealth;
 
@@ -5853,8 +5853,8 @@ bool GameplayEffectsTest_ModifyChanceToApplyToTarget(UWorld *World, FAutomationT
 
 	UProperty *DamageProperty = FindFieldChecked<UProperty>(UAbilitySystemTestAttributeSet::StaticClass(), GET_MEMBER_NAME_CHECKED(UAbilitySystemTestAttributeSet, Damage));
 
-	UAttributeComponent * SourceComponent = SourceActor->AttributeComponent;
-	UAttributeComponent * DestComponent = DestActor->AttributeComponent;
+	UAbilitySystemComponent * SourceComponent = SourceActor->AbilitySystemComponent;
+	UAbilitySystemComponent * DestComponent = DestActor->AbilitySystemComponent;
 	SourceComponent->GetSet<UAbilitySystemTestAttributeSet>()->Health = StartHealth;
 	DestComponent->GetSet<UAbilitySystemTestAttributeSet>()->Health = StartHealth;
 
@@ -5911,8 +5911,8 @@ bool GameplayEffectsTest_ModifyChanceToExecuteOnGE(UWorld *World, FAutomationTes
 
 	UProperty *DamageProperty = FindFieldChecked<UProperty>(UAbilitySystemTestAttributeSet::StaticClass(), GET_MEMBER_NAME_CHECKED(UAbilitySystemTestAttributeSet, Damage));
 
-	UAttributeComponent * SourceComponent = SourceActor->AttributeComponent;
-	UAttributeComponent * DestComponent = DestActor->AttributeComponent;
+	UAbilitySystemComponent * SourceComponent = SourceActor->AbilitySystemComponent;
+	UAbilitySystemComponent * DestComponent = DestActor->AbilitySystemComponent;
 	SourceComponent->GetSet<UAbilitySystemTestAttributeSet>()->Health = StartHealth;
 	DestComponent->GetSet<UAbilitySystemTestAttributeSet>()->Health = StartHealth;
 
diff --git a/Engine/Source/Runtime/GameplayAbilities/Public/AbilitySystemTypes.h b/Engine/Source/Runtime/GameplayAbilities/Public/AbilitySystemTypes.h
--- a/Engine/Source/Runtime/GameplayAbilities/Public/AbilitySystemTypes.h
+++ b/Engine/Source/Runtime/GameplayAbilities/Public/AbilitySystemTypes.h
@@ -4,7 +4,7 @@
 
 #include "AbilitySystemTypes.generated.h"
 
-class UAttributeComponent;
+class UAbilitySystemComponent;
 class UGameplayEffect;
 
 UENUM(BlueprintType)
@@ -44,7 +44,7 @@ struct GAMEPLAYABILITIES_API FGameplayAbilityActorInfo
 	TWeakObjectPtr<UAnimInstance>	AnimInstance;
 
 	UPROPERTY(BlueprintReadOnly, Category = "ActorInfo")
-	TWeakObjectPtr<UAttributeComponent>	AttributeComponent;
+	TWeakObjectPtr<UAbilitySystemComponent>	AbilitySystemComponent;
 
 	UPROPERTY(BlueprintReadOnly, Category = "ActorInfo")
 	TWeakObjectPtr<UMovementComponent>	MovementComponent;
@@ -89,7 +89,7 @@ struct GAMEPLAYABILITIES_API FGameplayAbilityActivationInfo
 	}
 
 
-	void GeneratePredictionKey(UAttributeComponent * Component) const;
+	void GeneratePredictionKey(UAbilitySystemComponent * Component) const;
 
 	void SetPredictionKey(int32 InPredictionKey);
 
