remove use of 'classes' folder in AbilitySystem -> just use public folder.

diff --git a/Engine/Source/Runtime/GameplayAbilities/Classes/Abilities/GameplayAbility.h b/Engine/Source/Runtime/GameplayAbilities/Classes/Abilities/GameplayAbility.h
deleted file mode 100644
diff --git a/Engine/Source/Runtime/GameplayAbilities/Classes/Abilities/GameplayAbility_CharacterJump.h b/Engine/Source/Runtime/GameplayAbilities/Classes/Abilities/GameplayAbility_CharacterJump.h
deleted file mode 100644
diff --git a/Engine/Source/Runtime/GameplayAbilities/Classes/Abilities/GameplayAbility_Instanced.h b/Engine/Source/Runtime/GameplayAbilities/Classes/Abilities/GameplayAbility_Instanced.h
deleted file mode 100644
diff --git a/Engine/Source/Runtime/GameplayAbilities/Classes/Abilities/GameplayAbility_Montage.h b/Engine/Source/Runtime/GameplayAbilities/Classes/Abilities/GameplayAbility_Montage.h
deleted file mode 100644
diff --git a/Engine/Source/Runtime/GameplayAbilities/Classes/Abilities/Tasks/AbilityTask.h b/Engine/Source/Runtime/GameplayAbilities/Classes/Abilities/Tasks/AbilityTask.h
deleted file mode 100644
diff --git a/Engine/Source/Runtime/GameplayAbilities/Classes/Abilities/Tasks/AbilityTask_PlayMontageAndWait.h b/Engine/Source/Runtime/GameplayAbilities/Classes/Abilities/Tasks/AbilityTask_PlayMontageAndWait.h
deleted file mode 100644
diff --git a/Engine/Source/Runtime/GameplayAbilities/Classes/Abilities/Tasks/AbilityTask_WaitMovementModeChange.h b/Engine/Source/Runtime/GameplayAbilities/Classes/Abilities/Tasks/AbilityTask_WaitMovementModeChange.h
deleted file mode 100644
diff --git a/Engine/Source/Runtime/GameplayAbilities/Classes/Abilities/Tasks/AbilityTask_WaitOverlap.h b/Engine/Source/Runtime/GameplayAbilities/Classes/Abilities/Tasks/AbilityTask_WaitOverlap.h
deleted file mode 100644
diff --git a/Engine/Source/Runtime/GameplayAbilities/Classes/AbilitySystemBlueprintLibrary.h b/Engine/Source/Runtime/GameplayAbilities/Classes/AbilitySystemBlueprintLibrary.h
deleted file mode 100644
diff --git a/Engine/Source/Runtime/GameplayAbilities/Classes/AbilitySystemComponent.h b/Engine/Source/Runtime/GameplayAbilities/Classes/AbilitySystemComponent.h
deleted file mode 100644
diff --git a/Engine/Source/Runtime/GameplayAbilities/Classes/AbilitySystemDebugHUD.h b/Engine/Source/Runtime/GameplayAbilities/Classes/AbilitySystemDebugHUD.h
deleted file mode 100644
diff --git a/Engine/Source/Runtime/GameplayAbilities/Classes/AbilitySystemGlobals.h b/Engine/Source/Runtime/GameplayAbilities/Classes/AbilitySystemGlobals.h
deleted file mode 100644
diff --git a/Engine/Source/Runtime/GameplayAbilities/Classes/AbilitySystemTestAttributeSet.h b/Engine/Source/Runtime/GameplayAbilities/Classes/AbilitySystemTestAttributeSet.h
deleted file mode 100644
diff --git a/Engine/Source/Runtime/GameplayAbilities/Classes/AbilitySystemTestPawn.h b/Engine/Source/Runtime/GameplayAbilities/Classes/AbilitySystemTestPawn.h
deleted file mode 100644
diff --git a/Engine/Source/Runtime/GameplayAbilities/Classes/AttributeSet.h b/Engine/Source/Runtime/GameplayAbilities/Classes/AttributeSet.h
deleted file mode 100644
diff --git a/Engine/Source/Runtime/GameplayAbilities/Classes/GameplayAbilitySet.h b/Engine/Source/Runtime/GameplayAbilities/Classes/GameplayAbilitySet.h
deleted file mode 100644
diff --git a/Engine/Source/Runtime/GameplayAbilities/Classes/GameplayCueInterface.h b/Engine/Source/Runtime/GameplayAbilities/Classes/GameplayCueInterface.h
deleted file mode 100644
diff --git a/Engine/Source/Runtime/GameplayAbilities/Classes/GameplayCueView.h b/Engine/Source/Runtime/GameplayAbilities/Classes/GameplayCueView.h
deleted file mode 100644
diff --git a/Engine/Source/Runtime/GameplayAbilities/Classes/GameplayEffect.h b/Engine/Source/Runtime/GameplayAbilities/Classes/GameplayEffect.h
deleted file mode 100644
diff --git a/Engine/Source/Runtime/GameplayAbilities/Classes/GameplayEffectExtension.h b/Engine/Source/Runtime/GameplayAbilities/Classes/GameplayEffectExtension.h
deleted file mode 100644
diff --git a/Engine/Source/Runtime/GameplayAbilities/Classes/GameplayEffectExtension_LifestealTest.h b/Engine/Source/Runtime/GameplayAbilities/Classes/GameplayEffectExtension_LifestealTest.h
deleted file mode 100644
diff --git a/Engine/Source/Runtime/GameplayAbilities/Classes/GameplayEffectExtension_ShieldTest.h b/Engine/Source/Runtime/GameplayAbilities/Classes/GameplayEffectExtension_ShieldTest.h
deleted file mode 100644
diff --git a/Engine/Source/Runtime/GameplayAbilities/Classes/GameplayEffectStackingExtension.h b/Engine/Source/Runtime/GameplayAbilities/Classes/GameplayEffectStackingExtension.h
deleted file mode 100644
diff --git a/Engine/Source/Runtime/GameplayAbilities/Classes/GameplayEffectStackingExtension_CappedNumberTest.h b/Engine/Source/Runtime/GameplayAbilities/Classes/GameplayEffectStackingExtension_CappedNumberTest.h
deleted file mode 100644
diff --git a/Engine/Source/Runtime/GameplayAbilities/Classes/GameplayEffectStackingExtension_DiminishingReturnsTest.h b/Engine/Source/Runtime/GameplayAbilities/Classes/GameplayEffectStackingExtension_DiminishingReturnsTest.h
deleted file mode 100644
diff --git a/Engine/Source/Runtime/GameplayAbilities/Public/Abilities/GameplayAbility.h b/Engine/Source/Runtime/GameplayAbilities/Public/Abilities/GameplayAbility.h
new file mode 100644
--- /dev/null
+++ b/Engine/Source/Runtime/GameplayAbilities/Public/Abilities/GameplayAbility.h
@@ -0,0 +1,209 @@
+// Copyright 1998-2014 Epic Games, Inc. All Rights Reserved.
+
+#pragma once
+
+#include "GameplayTagContainer.h"
+#include "GameplayTagAssetInterface.h"
+#include "Runtime/Engine/Classes/Animation/AnimInstance.h"
+#include "AbilitySystemTypes.h"
+#include "GameplayAbilityTypes.h"
+#include "GameplayAbility.generated.h"
+
+
+/**
+* UGameplayAbility
+*	The GameplayAbility baseclass. This base class is a data asset and can be used to create non-instanced abilities.
+*	Instanced abilities should derive off UGameplyAbility_Instanced.
+*/
+UCLASS()
+class GAMEPLAYABILITIES_API UGameplayAbility : public UDataAsset
+{
+	GENERATED_UCLASS_BODY()
+
+public:
+
+	UFUNCTION(Client, Reliable)
+	void ClientActivateAbilitySucceed(int32 PredictionKey);
+
+	virtual void ClientActivateAbilitySucceed_Internal(int32 PredictionKey);
+
+	bool IsSupportedForNetworking() const OVERRIDE;
+	
+	// ----------------------------------------------------------------------------------------------------------------
+	//
+	//	The important functions:
+	//	
+	//		CanActivateAbility()	- const function to see if ability is activatable. Callable by UI etc
+	//
+	//		TryActivateAbility()	- Attempts to activate the ability. Calls CanActivateAbility(). Input events can call this directly.
+	//								- Also handles instancing-per-execution logic and replication/prediction calls.
+	//		
+	//		CallActivate()			- Protected, non virtual function. Does some boilerplate 'pre activate' stuff, then calls Activate()
+	//
+	//		Activate()				- What the abilities *does*. This is what child classes want to override.
+	//	
+	//		Commit()				- Commits reources/cooldowns etc. Activate() must call this!
+	//		
+	//		CancelAbility()			- Interrupts the ability (from an outside source).
+	//									-We may want to add some info on what/who cancelled.
+	//
+	//		EndAbility()			- The ability has ended. This is intended to be called by the ability to end itself.
+	//	
+	// ----------------------------------------------------------------------------------------------------------------
+
+	/**
+	 * Attempts to activate the ability.
+	 *	-This function calls CanActivateAbility
+	 *	-This function handles instancing
+	 *	-This function handles networking and prediction
+	 *	-If all goes well, CallActivateAbility is called next.
+	 */
+	bool TryActivateAbility(const FGameplayAbilityActorInfo* ActorInfo, int32 PredictionKey = 0, UGameplayAbility ** OutInstancedAbility = NULL);
+	
+	/** Returns true if this ability can be activated right now. Has no side effects */
+	virtual bool CanActivateAbility(const FGameplayAbilityActorInfo* ActorInfo) const;
+
+	// -----------------------------------------------
+
+	/** Input binding. Base implementation calls TryActivateAbility */
+	virtual void InputPressed(int32 InputID, const FGameplayAbilityActorInfo* ActorInfo);
+
+	/** Input binding. Base implementation does nothing */
+	virtual void InputReleased(int32 InputID, const FGameplayAbilityActorInfo* ActorInfo);
+
+	// ----------------------------------------------------------------------------------------------------------------
+	//
+	//	Outward facing API for other systems
+	//
+	// ----------------------------------------------------------------------------------------------------------------
+
+	float GetCooldownTimeRemaining(const FGameplayAbilityActorInfo* ActorInfo) const;
+
+	// ----------------------------------------------------------------------------------------------------------------
+	//
+	//	Standardized State
+	//
+	// ----------------------------------------------------------------------------------------------------------------	
+
+	UPROPERTY(EditDefaultsOnly, Category=Advanced)
+	TEnumAsByte<EGameplayAbilityNetExecutionPolicy::Type>	NetExecutionPolicy;
+
+	/** This GameplayEffect represents the cooldown. It will be applied when the ability is commited and the ability cannot be used again until it is expired. */
+	UPROPERTY(EditDefaultsOnly, Category=Cooldowns)
+	class UGameplayEffect * CooldownGameplayEffect;
+
+	/** This GameplayEffect represents the cost (mana, stamina, etc) of the ability. It will be applied when the ability is commited. */
+	UPROPERTY(EditDefaultsOnly, Category=Costs)
+	class UGameplayEffect * CostGameplayEffect;
+
+	// ----------------------------------------------------------------------------------------------------------------
+	//
+	//	Ability exclusion / cancelling
+	//
+	// ----------------------------------------------------------------------------------------------------------------
+	
+	/** Abilities with these tags are canelled when this ability is executed */
+	UPROPERTY(EditDefaultsOnly, Category = Tags)
+	FGameplayTagContainer CancelAbilitiesWithTag;
+
+	/** This ability has these tags */
+	UPROPERTY(EditDefaultsOnly, Category = Tags)
+	FGameplayTagContainer AbilityTags;
+
+	// ----------------------------------------------------------------------------------------------------------------
+	//
+	//	Animation callbacks
+	//
+	// ----------------------------------------------------------------------------------------------------------------
+
+	void MontageBranchPoint_AbilityDecisionStop(const FGameplayAbilityActorInfo* ActorInfo) const;
+
+	void MontageBranchPoint_AbilityDecisionStart(const FGameplayAbilityActorInfo* ActorInfo) const;
+
+	// ----------------------------------------------------------------------------------------------------------------
+
+	virtual void PredictiveActivateAbility(const FGameplayAbilityActorInfo* ActorInfo, const FGameplayAbilityActivationInfo ActivationInfo);
+
+	virtual void ServerTryActivateAbility(const FGameplayAbilityActorInfo* ActorInfo, const FGameplayAbilityActivationInfo ActivationInfo);
+
+	virtual void ClientActivateAbility(const FGameplayAbilityActorInfo* ActorInfo, const FGameplayAbilityActivationInfo ActivationInfo);
+
+	// ----------------------------------------------------------------------------------------------------------------
+
+	virtual EGameplayAbilityInstancingPolicy::Type GetInstancingPolicy() const
+	{
+		return EGameplayAbilityInstancingPolicy::NonInstanced;
+	}
+
+	virtual EGameplayAbilityReplicationPolicy::Type GetReplicationPolicy() const
+	{
+		return EGameplayAbilityReplicationPolicy::ReplicateNone;
+	}
+
+	// ----------------------------------------------------------------------------------------------------------------
+
+	virtual UWorld* GetWorld() const OVERRIDE
+	{
+		return GetOuter()->GetWorld();
+	}
+
+protected:
+
+	/**
+	 * The main function that defines what an ability does.
+	 *  -Child classes will want to override this
+	 *  -This function must call CommitAbility()
+	 */
+
+	virtual void ActivateAbility(const FGameplayAbilityActorInfo * ActorInfo, const FGameplayAbilityActivationInfo ActivationInfo);
+
+	/** 
+	 * Attempts to commit the ability (spend resources, etc). This our last chance to fail.
+	 *	-Child classes that override ActivateAbility must call this themselves!
+	 */
+
+	virtual bool CommitAbility(const FGameplayAbilityActorInfo* ActorInfo, const FGameplayAbilityActivationInfo ActivationInfo);
+	
+	/** Do boilerplate init stuff and then call ActivateAbility */
+	void PreActivate(const FGameplayAbilityActorInfo* ActorInfo, const FGameplayAbilityActivationInfo ActivationInfo);
+
+	void CallActivateAbility(const FGameplayAbilityActorInfo* ActorInfo, const FGameplayAbilityActivationInfo ActivationInfo);
+
+	void CallPredictiveActivateAbility(const FGameplayAbilityActorInfo* ActorInfo, const FGameplayAbilityActivationInfo ActivationInfo);
+
+	virtual void EndAbility(const FGameplayAbilityActorInfo* ActorInfo);
+
+	/** 
+	 * The last chance to fail before commiting
+	 *	-This will usually be the same as CanActivateAbility. Some abilities may need to do extra checks here if they are consuming extra stuff in CommitExecute
+	 */
+	virtual bool CommitCheck(const FGameplayAbilityActorInfo* ActorInfo, const FGameplayAbilityActivationInfo ActivationInfo);
+
+	/** Does the commmit atomically (consume resources, do cooldowns, etc) */
+	virtual void CommitExecute(const FGameplayAbilityActorInfo* ActorInfo, const FGameplayAbilityActivationInfo ActivationInfo);
+
+	/** Destroys intanced-per-execution abilities. Instance-per-actor abilities should 'reset'. Non instance abilities - what can we do? */
+	virtual void CancelAbility(const FGameplayAbilityActorInfo* ActorInfo, const FGameplayAbilityActivationInfo ActivationInfo) { }
+
+	// --------------------------------------------------------------------------------
+
+
+	// --------------------------------------------------------------------------------
+
+	/** Checks cooldown. returns true if we can be used again. False if not */
+	bool	CheckCooldown(const FGameplayAbilityActorInfo* ActorInfo) const;
+
+	/** Applies CooldownGameplayEffect to the target */
+	void	ApplyCooldown(const FGameplayAbilityActorInfo* ActorInfo, const FGameplayAbilityActivationInfo ActivationInfo);
+
+	bool	CheckCost(const FGameplayAbilityActorInfo* ActorInfo) const;
+
+	void	ApplyCost(const FGameplayAbilityActorInfo* ActorInfo, const FGameplayAbilityActivationInfo ActivationInfo);
+
+	// --------------------------------------------------------------------------------
+
+	UPROPERTY(EditDefaultsOnly, Category = Tags)
+	bool Prediction;
+
+	friend class UAbilitySystemComponent;
+};
\ No newline at end of file
diff --git a/Engine/Source/Runtime/GameplayAbilities/Public/Abilities/GameplayAbility_CharacterJump.h b/Engine/Source/Runtime/GameplayAbilities/Public/Abilities/GameplayAbility_CharacterJump.h
new file mode 100644
--- /dev/null
+++ b/Engine/Source/Runtime/GameplayAbilities/Public/Abilities/GameplayAbility_CharacterJump.h
@@ -0,0 +1,27 @@
+// Copyright 1998-2013 Epic Games, Inc. All Rights Reserved.
+
+#pragma once
+
+#include "Abilities/GameplayAbility.h"
+#include "GameplayAbility_CharacterJump.generated.h"
+
+/**
+ *	A gameplay ability that plays a single montage and applies a GameplayEffect
+ */
+UCLASS()
+class GAMEPLAYABILITIES_API UGameplayAbility_CharacterJump : public UGameplayAbility
+{
+	GENERATED_UCLASS_BODY()
+
+public:
+
+	virtual bool CanActivateAbility(const FGameplayAbilityActorInfo* ActorInfo) const OVERRIDE;
+	
+	virtual void ActivateAbility(const FGameplayAbilityActorInfo* OwnerInfo, const FGameplayAbilityActivationInfo ActivationInfo) OVERRIDE;
+
+	virtual void InputReleased(int32 InputID, const FGameplayAbilityActorInfo* ActorInfo) OVERRIDE;
+
+	virtual void PredictiveActivateAbility(const FGameplayAbilityActorInfo* ActorInfo, const FGameplayAbilityActivationInfo ActivationInfo) OVERRIDE;
+
+	virtual void CancelAbility(const FGameplayAbilityActorInfo* ActorInfo, const FGameplayAbilityActivationInfo ActivationInfo) OVERRIDE;
+};
\ No newline at end of file
diff --git a/Engine/Source/Runtime/GameplayAbilities/Public/Abilities/GameplayAbility_Instanced.h b/Engine/Source/Runtime/GameplayAbilities/Public/Abilities/GameplayAbility_Instanced.h
new file mode 100644
--- /dev/null
+++ b/Engine/Source/Runtime/GameplayAbilities/Public/Abilities/GameplayAbility_Instanced.h
@@ -0,0 +1,107 @@
+// Copyright 1998-2013 Epic Games, Inc. All Rights Reserved.
+
+#pragma once
+#include "GameplayAbility.h"
+#include "GameplayAbility_Instanced.generated.h"
+
+class UAnimInstance;
+class UAbilitySystemComponent;
+
+/**
+* UGameplayEffect
+*	The GameplayEffect definition. This is the data asset defined in the editor that drives everything.
+*/
+UCLASS(Blueprintable)
+class GAMEPLAYABILITIES_API UGameplayAbility_Instanced : public UGameplayAbility
+{
+	GENERATED_UCLASS_BODY()
+
+public:
+
+	virtual void ClientActivateAbilitySucceed_Internal(int32 PredictionKey) OVERRIDE;
+	
+	// -------------------------------------------------
+
+	UFUNCTION(BlueprintImplementableEvent, Category = Ability)
+	virtual bool K2_CanActivateAbility() const;
+
+	virtual bool CanActivateAbility(const FGameplayAbilityActorInfo* ActorInfo) const OVERRIDE;
+
+	// -------------------------------------------------
+
+	UFUNCTION(BlueprintImplementableEvent, Category = Ability)
+	virtual void K2_CommitExecute();
+
+	virtual void CommitExecute(const FGameplayAbilityActorInfo* ActorInfo, const FGameplayAbilityActivationInfo ActivationInfo) OVERRIDE;
+
+	// -------------------------------------------------
+
+	UFUNCTION(BlueprintImplementableEvent, Category = Ability)
+	virtual void K2_ActivateAbility();
+
+	virtual void ActivateAbility(const FGameplayAbilityActorInfo* ActorInfo, const FGameplayAbilityActivationInfo ActivationInfo) OVERRIDE;
+
+	// -------------------------------------------------
+
+	UFUNCTION(BlueprintImplementableEvent, Category = Ability)
+	virtual void K2_PredictiveActivateAbility();
+
+	virtual void PredictiveActivateAbility(const FGameplayAbilityActorInfo* ActorInfo, const FGameplayAbilityActivationInfo ActivationInfo) OVERRIDE;
+
+	// -------------------------------------------------
+
+	UFUNCTION(BlueprintPure, Category = Ability)
+	FGameplayAbilityActorInfo GetActorInfo();
+
+	// -------------------------------------------------
+
+	UFUNCTION(BlueprintCallable, Category = Ability)
+	virtual bool K2_CommitAbility();
+
+	UFUNCTION(BlueprintCallable, Category = Ability)
+	virtual void K2_EndAbility();
+	
+
+	// -------------------------------------------------
+
+
+	virtual EGameplayAbilityInstancingPolicy::Type GetInstancingPolicy() const OVERRIDE
+	{
+		return InstancedPerExecution ? EGameplayAbilityInstancingPolicy::InstancedPerExecution : EGameplayAbilityInstancingPolicy::InstancedPerActor;
+	}
+
+	virtual EGameplayAbilityReplicationPolicy::Type GetReplicationPolicy() const OVERRIDE
+	{
+		return ReplicationPolicy;
+	}
+
+	UPROPERTY(EditDefaultsOnly, Category=Advanced)
+	TEnumAsByte<EGameplayAbilityReplicationPolicy::Type> ReplicationPolicy;
+
+	UPROPERTY(EditDefaultsOnly, Category=Advanced)
+	bool	InstancedPerExecution;
+
+	int32 GetFunctionCallspace(UFunction* Function, void* Parameters, FFrame* Stack);
+
+	bool CallRemoteFunction(UFunction* Function, void* Parameters, FOutParmRec* OutParms, FFrame* Stack);
+
+
+	// -------------------------------------------------
+
+
+	/** This is shared, cached information about the thing using us
+	 *		E.g, Actor*, MovementComponent*, AnimInstance, etc.
+	 *		This is hopefully allocated once per actor and shared by many abilities.
+	 *		The actual struct may be overridden per game to include game specific data.
+	 *		(E.g, child classes may want to cast to FMyGameAbilityActorInfo)
+	 */
+	mutable const FGameplayAbilityActorInfo * CurrentActorInfo;
+
+	/** This is information specific to this instance of the ability. E.g, whether it is predicting, authorting, confirmed, etc. */
+	UPROPERTY(BlueprintReadOnly, Category=Ability)
+	FGameplayAbilityActivationInfo	CurrentActivationInfo;
+
+protected:
+
+	virtual void EndAbility(const FGameplayAbilityActorInfo* ActorInfo) OVERRIDE;
+};
\ No newline at end of file
diff --git a/Engine/Source/Runtime/GameplayAbilities/Public/Abilities/GameplayAbility_Montage.h b/Engine/Source/Runtime/GameplayAbilities/Public/Abilities/GameplayAbility_Montage.h
new file mode 100644
--- /dev/null
+++ b/Engine/Source/Runtime/GameplayAbilities/Public/Abilities/GameplayAbility_Montage.h
@@ -0,0 +1,34 @@
+// Copyright 1998-2013 Epic Games, Inc. All Rights Reserved.
+
+#pragma once
+
+#include "Abilities/GameplayAbility.h"
+#include "GameplayAbility_Montage.generated.h"
+
+/**
+ *	A gameplay ability that plays a single montage and applies a GameplayEffect
+ */
+UCLASS()
+class GAMEPLAYABILITIES_API UGameplayAbility_Montage : public UGameplayAbility
+{
+	GENERATED_UCLASS_BODY()
+
+public:
+	
+	virtual void ActivateAbility(const FGameplayAbilityActorInfo* OwnerInfo, const FGameplayAbilityActivationInfo ActivationInfo) OVERRIDE;
+
+	UPROPERTY(EditDefaultsOnly, Category = MontageAbility)
+	UAnimMontage *	MontageToPlay;
+
+	UPROPERTY(EditDefaultsOnly, Category = MontageAbility)
+	float	PlayRate;
+
+	UPROPERTY(EditDefaultsOnly, Category = MontageAbility)
+	FName	SectionName;
+
+	/** GameplayEffects to apply and then remove while the animation is playing */
+	UPROPERTY(EditDefaultsOnly, Category = MontageAbility)
+	TArray<const UGameplayEffect*>	GameplayEffectsWhileAnimating;
+
+	void OnMontageEnded(UAnimMontage* Montage, bool bInterrupted, const FGameplayAbilityActorInfo* OwnerInfo, TArray<struct FActiveGameplayEffectHandle>	AppliedEffects);
+};
\ No newline at end of file
diff --git a/Engine/Source/Runtime/GameplayAbilities/Public/Abilities/Tasks/AbilityTask.h b/Engine/Source/Runtime/GameplayAbilities/Public/Abilities/Tasks/AbilityTask.h
new file mode 100644
--- /dev/null
+++ b/Engine/Source/Runtime/GameplayAbilities/Public/Abilities/Tasks/AbilityTask.h
@@ -0,0 +1,12 @@
+// Copyright 1998-2014 Epic Games, Inc. All Rights Reserved.
+#pragma once
+#include "AbilityTask.Generated.h"
+
+UCLASS(MinimalAPI)
+class UAbilityTask : public UObject
+{
+	GENERATED_UCLASS_BODY()	
+
+public:
+	
+};
\ No newline at end of file
diff --git a/Engine/Source/Runtime/GameplayAbilities/Public/Abilities/Tasks/AbilityTask_PlayMontageAndWait.h b/Engine/Source/Runtime/GameplayAbilities/Public/Abilities/Tasks/AbilityTask_PlayMontageAndWait.h
new file mode 100644
--- /dev/null
+++ b/Engine/Source/Runtime/GameplayAbilities/Public/Abilities/Tasks/AbilityTask_PlayMontageAndWait.h
@@ -0,0 +1,24 @@
+// Copyright 1998-2014 Epic Games, Inc. All Rights Reserved.
+#pragma once
+#include "AbilityTask.h"
+#include "AbilityTask_PlayMontageAndWait.generated.h"
+
+DECLARE_DYNAMIC_MULTICAST_DELEGATE(FMontageWaitSimpleDelegate);
+
+UCLASS(MinimalAPI)
+class UAbilityTask_PlayMontageAndWait : public UAbilityTask
+{
+	GENERATED_UCLASS_BODY()
+
+	UPROPERTY(BlueprintAssignable)
+	FMontageWaitSimpleDelegate	OnComplete;
+
+	UPROPERTY(BlueprintAssignable)
+	FMontageWaitSimpleDelegate	OnInterrupted;
+
+	UFUNCTION()
+	void OnMontageEnded(UAnimMontage* Montage, bool bInterrupted);
+
+public:
+	
+};
\ No newline at end of file
diff --git a/Engine/Source/Runtime/GameplayAbilities/Public/Abilities/Tasks/AbilityTask_WaitMovementModeChange.h b/Engine/Source/Runtime/GameplayAbilities/Public/Abilities/Tasks/AbilityTask_WaitMovementModeChange.h
new file mode 100644
--- /dev/null
+++ b/Engine/Source/Runtime/GameplayAbilities/Public/Abilities/Tasks/AbilityTask_WaitMovementModeChange.h
@@ -0,0 +1,23 @@
+// Copyright 1998-2014 Epic Games, Inc. All Rights Reserved.
+#pragma once
+#include "AbilityTask.h"
+#include "AbilityTask_WaitMovementModeChange.generated.h"
+
+DECLARE_DYNAMIC_MULTICAST_DELEGATE_OneParam(FMovementModeChangedDelegate, EMovementMode, NewMovementMode);
+
+UCLASS(MinimalAPI)
+class UAbilityTask_WaitMovementModeChange : public UAbilityTask
+{
+	GENERATED_UCLASS_BODY()
+
+	UPROPERTY(BlueprintAssignable)
+	FMovementModeChangedDelegate	OnChange;
+
+	UFUNCTION()
+	void OnMovementModeChange(ACharacter * Character, EMovementMode PrevMovementMode, uint8 PreviousCustomMode);
+
+	EMovementMode	RequiredMode;
+
+public:
+	
+};
\ No newline at end of file
diff --git a/Engine/Source/Runtime/GameplayAbilities/Public/Abilities/Tasks/AbilityTask_WaitOverlap.h b/Engine/Source/Runtime/GameplayAbilities/Public/Abilities/Tasks/AbilityTask_WaitOverlap.h
new file mode 100644
--- /dev/null
+++ b/Engine/Source/Runtime/GameplayAbilities/Public/Abilities/Tasks/AbilityTask_WaitOverlap.h
@@ -0,0 +1,28 @@
+// Copyright 1998-2014 Epic Games, Inc. All Rights Reserved.
+#pragma once
+#include "AbilityTask.h"
+#include "GameplayAbilityTypes.h"
+#include "AbilityTask_WaitOverlap.generated.h"
+
+DECLARE_DYNAMIC_MULTICAST_DELEGATE_OneParam(FWaitOverlapDelegate, FGameplayAbilityTargetDataHandle, TargetData);
+
+class AActor;
+class UPrimitiveComponent;
+
+UCLASS(MinimalAPI)
+class UAbilityTask_WaitOverlap : public UAbilityTask
+{
+	GENERATED_UCLASS_BODY()
+
+	UPROPERTY(BlueprintAssignable)
+	FWaitOverlapDelegate	OnOverlap;
+
+	UFUNCTION()
+	void OnOverlapCallback(AActor* OtherActor, UPrimitiveComponent* OtherComp, int32 OtherBodyIndex, bool bFromSweep, const FHitResult & SweepResult);
+
+	UFUNCTION()
+	void OnHitCallback(AActor* OtherActor, UPrimitiveComponent* OtherComp, FVector NormalImpulse, const FHitResult& Hit);
+
+public:
+	
+};
\ No newline at end of file
diff --git a/Engine/Source/Runtime/GameplayAbilities/Public/AbilitySystemBlueprintLibrary.h b/Engine/Source/Runtime/GameplayAbilities/Public/AbilitySystemBlueprintLibrary.h
new file mode 100644
--- /dev/null
+++ b/Engine/Source/Runtime/GameplayAbilities/Public/AbilitySystemBlueprintLibrary.h
@@ -0,0 +1,35 @@
+// Copyright 1998-2014 Epic Games, Inc. All Rights Reserved.
+
+#pragma once
+
+#include "Abilities/GameplayAbility.h"
+#include "AbilitySystemTypes.h"
+#include "AbilitySystemBlueprintLibrary.generated.h"
+
+class UAnimMontage;
+class UAbilityTask_PlayMontageAndWait;
+class UAbilityTask_WaitMovementModeChange;
+class UAbilityTask_WaitOverlap;
+
+UCLASS(MinimalAPI)
+class UAbilitySystemBlueprintLibrary : public UBlueprintFunctionLibrary
+{
+	GENERATED_UCLASS_BODY()
+
+	UFUNCTION(BlueprintPure, Category = Ability)
+	static UAbilitySystemComponent* GetAbilitySystemComponent(AActor *Actor);
+
+	UFUNCTION(BlueprintCallable, Category=Ability)
+	static void ApplyGameplayEffectToTargetData(FGameplayAbilityTargetDataHandle Target, UGameplayEffect *GameplayEffect, const FGameplayAbilityActorInfo InstigatorInfo);
+
+	// -------------------------------------------------------------------------------
+
+	UFUNCTION(BlueprintCallable, meta=(HidePin="WorldContextObject", DefaultToSelf="WorldContextObject", BlueprintInternalUseOnly = "TRUE"))
+	static UAbilityTask_PlayMontageAndWait* CreatePlayMontageAndWaitProxy(UObject* WorldContextObject, UAnimMontage *MontageToPlay);
+
+	UFUNCTION(BlueprintCallable, meta = (HidePin = "WorldContextObject", DefaultToSelf = "WorldContextObject", BlueprintInternalUseOnly = "TRUE"))
+	static UAbilityTask_WaitMovementModeChange* CreateWaitMovementModeChange(UObject* WorldContextObject, EMovementMode NewMode);
+
+	UFUNCTION(BlueprintCallable, meta = (HidePin = "WorldContextObject", DefaultToSelf = "WorldContextObject", BlueprintInternalUseOnly = "TRUE"))
+	static UAbilityTask_WaitOverlap* CreateWaitOverlap(UObject* WorldContextObject, EMovementMode NewMode);
+};
\ No newline at end of file
diff --git a/Engine/Source/Runtime/GameplayAbilities/Public/AbilitySystemComponent.h b/Engine/Source/Runtime/GameplayAbilities/Public/AbilitySystemComponent.h
new file mode 100644
--- /dev/null
+++ b/Engine/Source/Runtime/GameplayAbilities/Public/AbilitySystemComponent.h
@@ -0,0 +1,368 @@
+// Copyright 1998-2014 Epic Games, Inc. All Rights Reserved.
+#pragma once
+
+#include "GameplayEffect.h"
+#include "GameplayAbilityTypes.h"
+#include "AbilitySystemComponent.generated.h"
+
+GAMEPLAYABILITIES_API DECLARE_LOG_CATEGORY_EXTERN(LogAbilitySystemComponent, Log, All);
+
+DECLARE_MULTICAST_DELEGATE(FAbilitySystemComponentPredictionKeyClear);
+
+USTRUCT()
+struct GAMEPLAYABILITIES_API FAttributeDefaults
+{
+	GENERATED_USTRUCT_BODY()
+
+	UPROPERTY(EditAnywhere, Category="AttributeTest")
+	TSubclassOf<class UAttributeSet> Attributes;
+
+	UPROPERTY(EditAnywhere, BlueprintReadWrite, Category="AttributeTest")
+	class UDataTable*	DefaultStartingTable;
+};
+
+/** 
+ * 
+ */
+UCLASS(ClassGroup=AbilitySystem, hidecategories=(Object,LOD,Lighting,Transform,Sockets,TextureStreaming), editinlinenew, meta=(BlueprintSpawnableComponent))
+class GAMEPLAYABILITIES_API UAbilitySystemComponent : public UActorComponent
+{
+	GENERATED_UCLASS_BODY()
+
+	friend FGameplayEffectSpec;
+	friend class AAbilitySystemDebugHUD;
+
+	virtual ~UAbilitySystemComponent();
+
+	virtual void InitializeComponent() OVERRIDE;
+
+	/** Finds existing AttributeSet */
+	template <class T >
+	T*	GetSet()
+	{
+		return (T*)GetAttributeSubobject(T::StaticClass());
+	}
+
+	template <class T >
+	T*	GetSetChecked()
+	{
+		return (T*)GetAttributeSubobjectChecked(T::StaticClass());
+	}
+
+	/** Adds a new AttributeSet (initialized to default values) */
+	template <class T >
+	T*  AddSet()
+	{
+		return (T*)GetOrCreateAttributeSubobject(T::StaticClass());
+	}
+
+	/** Adds a new AttributeSet that is a DSO (created by called in their CStor) */
+	template <class T>
+	T*	AddDefaultSubobjectSet(TSubobjectPtr<T> Subobject)
+	{
+		T* Set = Subobject.Get();
+		SpawnedAttributes.Add(Set);
+		return Set;
+	}
+
+	UFUNCTION(BlueprintCallable, Category="Skills")
+	UAttributeSet * InitStats(TSubclassOf<class UAttributeSet> Attributes, const UDataTable* DataTable);
+
+	UFUNCTION(BlueprintCallable, Category="Skills")
+	void ModifyStats(TSubclassOf<class UAttributeSet> Attributes, FDataTableRowHandle ModifierHandle);
+
+	UPROPERTY(EditAnywhere, Category="AttributeTest")
+	TArray<FAttributeDefaults>	DefaultStartingData;
+
+	UPROPERTY(Replicated)
+	TArray<UAttributeSet*>	SpawnedAttributes;
+
+	void SetNumericAttribute(const FGameplayAttribute &Attribute, float NewFloatValue);
+	float GetNumericAttribute(const FGameplayAttribute &Attribute);
+
+	// -- Replication -------------------------------------------------------------------------------------------------
+
+	virtual bool ReplicateSubobjects(class UActorChannel *Channel, class FOutBunch *Bunch, FReplicationFlags *RepFlags) OVERRIDE;
+	
+	virtual void GetSubobjectsWithStableNamesForNetworking(TArray<UObject*>& Objs) OVERRIDE;
+
+	virtual void PostNetReceive() OVERRIDE;
+
+	/**
+	 *	Prediction Keys
+	 *		PredictionKey is a simple system to allow clients to locally predict interactions within the AbilityComponent while having
+	 *		a mechanism for merging results from the server.
+	 *		
+	 *		There are two main cases:
+	 *			-Client predicts something and is wrong. (Mispredictions). He must cleanup stuff that he created or changed.
+	 *			-Client predicts something and is right. 
+	 *			
+	 *			The Client predicts something and asks the server. "Can I do X? PredictionKey=Y".
+	 *			Server can reply immediately and directly with yes or no, and include the PredictionKey with that answer.
+	 *			
+	 *			If No, the client can just immediately cleanup whatever he predicted. This is straightforward.
+	 *			
+	 *			If Yes, there is a small window where the Yes RPC is received but the actor/component properties have not been updated.
+	 *			The client's cannot immediately cleanup his prediction work: he needs to wait for property replication.
+	 *			
+	 *			This is where ReplicatedPredictionKey comes in. 
+	 *			
+	 *			When the server allows a predicted action, he sets ReplicatedPredictionKey to this value. The client then looks
+	 *			at this replicated value. In OnRep_PredictionKey, the client goes through and cleans up all stuff that was predicted
+	 *			and is no longer needed.
+	 *			
+	 *			-Client calls GetNextPredictionKey() at the start of a predicted action. This gives him a key to use for this 'run'.
+	 *			-Client can predict whatever he wants as long as he registers an 'undo' function with GetPredictionKeyDelegate().
+	 *			-This delegate is invoked if the client is denied his prediction, or if he is right and property replication catches up.
+	 *			
+	 *	 
+	 */
+
+	UPROPERTY(ReplicatedUsing=OnRep_PredictionKey)
+	int32	ReplicatedPredictionKey;
+
+	UPROPERTY(transient)
+	int32	LocalPredictionKey;
+
+	int32	GetNextPredictionKey();
+
+	UFUNCTION()
+	void OnRep_PredictionKey();
+
+	TArray<TPair<int32, FAbilitySystemComponentPredictionKeyClear> >	PredictionDelegates;
+	
+
+	FAbilitySystemComponentPredictionKeyClear &	GetPredictionKeyDelegate(int32 PredictionKey);
+
+	/*
+	virtual void OnSubobjectCreatedFromReplication(UObject *NewSubobject) OVERRIDE;
+	
+	virtual void OnSubobjectDestroyFromReplication(UObject *NewSubobject) OVERRIDE;
+	*/
+
+	// ----------------------------------------------------------------------------------------------------------------
+	//
+	//	GameplayEffects
+	//	(maybe should go in a different component?)
+	// ----------------------------------------------------------------------------------------------------------------
+
+	// --------------------------------------------
+	// Primary outward facing API for other systems:
+	// --------------------------------------------
+	FActiveGameplayEffectHandle ApplyGameplayEffectSpecToTarget(OUT FGameplayEffectSpec &GameplayEffect, UAbilitySystemComponent *Target, FModifierQualifier BaseQualifier = FModifierQualifier());
+	FActiveGameplayEffectHandle ApplyGameplayEffectSpecToSelf(const FGameplayEffectSpec &GameplayEffect, FModifierQualifier BaseQualifier = FModifierQualifier());
+
+	UFUNCTION(BlueprintCallable, BlueprintAuthorityOnly, Category = GameplayEffects)
+	bool RemoveActiveGameplayEffect(FActiveGameplayEffectHandle Handle);
+
+	/** This only exists so it can be hooked up to a multicast delegate */
+	void RemoveActiveGameplayEffect_NoReturn(FActiveGameplayEffectHandle Handle)
+	{
+		RemoveActiveGameplayEffect(Handle);
+	}
+
+	float GetGameplayEffectDuration(FActiveGameplayEffectHandle Handle) const;
+
+	float GetGameplayEffectDuration() const;
+
+	// Not happy with this interface but don't see a better way yet. How should outside code (UI, etc) ask things like 'how much is this gameplay effect modifying my damage by'
+	// (most likely we want to catch this on the backend - when damage is applied we can get a full dump/history of how the number got to where it is. But still we may need polling methods like below (how much would my damage be)
+	UFUNCTION(BlueprintCallable, Category = GameplayEffects)
+	float GetGameplayEffectMagnitude(FActiveGameplayEffectHandle Handle, FGameplayAttribute Attribute) const;
+
+	UFUNCTION(BlueprintCallable, Category = GameplayEffects)
+	float GetGameplayEffectMagnitudeByTag(FActiveGameplayEffectHandle InHandle, const FGameplayTag& InTag) const;
+
+	UFUNCTION(BlueprintCallable, Category = GameplayEffects)
+	bool IsGameplayEffectActive(FActiveGameplayEffectHandle InHandle) const;
+
+	// Tags
+	UFUNCTION(BlueprintCallable, Category=GameplayEffects)
+	bool HasAnyTags(FGameplayTagContainer &Tags);
+
+	UFUNCTION(BlueprintCallable, Category = GameplayEffects)
+	bool HasAllTags(FGameplayTagContainer &Tags);
+	
+
+	// --------------------------------------------
+	// Possibly useful but not primary API functions:
+	// --------------------------------------------
+
+	FActiveGameplayEffectHandle ApplyGameplayEffectToTarget(UGameplayEffect *GameplayEffect, UAbilitySystemComponent *Target, float Level = FGameplayEffectLevelSpec::INVALID_LEVEL, FModifierQualifier BaseQualifier = FModifierQualifier());
+
+	UFUNCTION(BlueprintCallable, BlueprintAuthorityOnly, Category = GameplayEffects, meta=(FriendlyName = "ApplyGameplayEffectToSelf"))
+	FActiveGameplayEffectHandle K2_ApplyGameplayEffectToSelf(const UGameplayEffect *GameplayEffect, float Level, AActor *Instigator);
+	
+	FActiveGameplayEffectHandle ApplyGameplayEffectToSelf(const UGameplayEffect *GameplayEffect, float Level, AActor *Instigator, FModifierQualifier BaseQualifier = FModifierQualifier() );
+
+	int32 GetNumActiveGameplayEffect() const;
+
+	void AddDependancyToAttribute(FGameplayAttribute Attribute, const TWeakPtr<FAggregator> InDependant);
+
+	/** Tests if all modifiers in this GameplayEffect will leave the attribute > 0.f */
+	bool CanApplyAttributeModifiers(const UGameplayEffect *GameplayEffect, float Level, AActor *Instigator);
+
+	// Generic 'Get expected magnitude (list) if I was to apply this outgoing or incoming'
+
+	// Get duration or magnitude (list) of active effects
+	//		-Get duration of CD
+	//		-Get magnitude + duration of a movespeed buff
+
+	TArray<float> GetActiveEffectsTimeRemaining(const FActiveGameplayEffectQuery Query) const;
+
+	void OnRestackGameplayEffects();
+
+	// --------------------------------------------
+	// Temp / Debug
+	// --------------------------------------------
+
+	void TEMP_ApplyActiveGameplayEffects();
+	
+	void PrintAllGameplayEffects() const;
+
+	void PushGlobalCurveOveride(UCurveTable *OverrideTable)
+	{
+		if (OverrideTable)
+		{
+			GlobalCurveDataOverride.Overrides.Push(OverrideTable);
+		}
+	}
+
+	// ----------------------------------------------------------------------------------------------------------------
+
+	/**
+	 *	GameplayAbilities
+	 *	
+	 *	The role of the AbilitySystemComponent wrt Abilities is to provide:
+	 *		-Management of abiltiy instances (whether per actor or per execution instance).
+	 *			-Someone *has* to keep track of these instances.
+	 *			-Non instanced abilities *could* be executed without any ability stuff in AbilitySystemComponent.
+	 *				They should be able to operate on an GameplayAbilityActorInfo + GameplayAbility.
+	 *		
+	 *	As convenience it may provide some other features:
+	 *		-Some basic input binding (whether instanced or non instanced abilities).
+	 *		-Concepts like "this component has these abilities
+	 *	
+	 */
+
+	class UGameplayAbility * CreateNewInstanceOfAbility(UGameplayAbility* Ability);
+
+	void CancelAbilitiesWithTags(const FGameplayTagContainer Tags, const FGameplayAbilityActorInfo* ActorInfo, const FGameplayAbilityActivationInfo ActivationInfo, UGameplayAbility* Ignore);
+	
+	/** References to non-replicating abilities that we instanced. We need to keep these refernces to avoid GC */
+	UPROPERTY()
+	TArray<class UGameplayAbility *>	NonReplicatedInstancedAbilities;
+
+	/** References to replicating abilities that we instanced. We need to keep these refernces to avoid GC */
+	UPROPERTY(Replicated)
+	TArray<class UGameplayAbility *>	ReplicatedInstancedAbilities;
+
+	/**
+	 *	The abilities we can activate. 
+	 *		-This will include CDOs for non instanced abilities and per-execution instanced abilities. 
+	 *		-Actor-instanced abilities will be the actual instance (not CDO)
+	 *		
+	 *	This array is not vital for things to work. It is a convenience thing for 'giving abilties to the actor'. But abilities could also work on things
+	 *	without an AbilitySystemComponent. For example an ability could be written to execute on a StaticMeshActor. As long as the ability doesn't require 
+	 *	instancing or anything else that the AbilitySystemComponent would provide, then it doesn't need the component to function.
+	 */ 
+	UPROPERTY(Replicated, BlueprintReadOnly, Category = "Abilities")
+	TArray<class UGameplayAbility *>	ActivatableAbilities;
+	
+
+	UFUNCTION(Server, WithValidation, Reliable)
+	void	ServerTryActivateAbility(class UGameplayAbility* AbilityToActivate, int32 PredictionKey);
+
+	UFUNCTION(Client, Reliable)
+	void	ClientActivateAbilityFailed(class UGameplayAbility* AbilityToActivate, int32 PredictionKey);
+
+	UFUNCTION(Client, Reliable)
+	void	ClientActivateAbilitySucceed(class UGameplayAbility* AbilityToActivate, int32 PredictionKey);
+
+	// ----------------------------------------------------------------------------------------------------------------
+
+	/** This is probably temp for testing. We want to think a bit more on the API for outside stuff activating abilities */
+	UFUNCTION(BlueprintCallable, Category = "Abilities")
+	bool	ActivateAbility(int32 AbilityIdx);
+
+	/** There needs to be a concept of an animating ability. Only one may exist at a time. New requestrs can be queued up, overridden, or ignored. */
+	UPROPERTY()
+	class UGameplayAbility *	AnimatingAbility;
+
+	void	MontageBranchPoint_AbilityDecisionStop();
+
+	void	MontageBranchPoint_AbilityDecisionStart();
+
+	
+
+	// -----------------------------------------------------------------------------
+
+	/** Cached off data about the owning actor that abilities will need to frequently access (movement component, mesh component, anim instance, etc) */
+	TSharedPtr<FGameplayAbilityActorInfo>	AbilityActorInfo;
+
+	void InitAbilityActorInfo();
+
+private:
+
+	// --------------------------------------------
+	// Important internal functions
+	// --------------------------------------------
+
+	void ExecutePeriodicEffect(FActiveGameplayEffectHandle	Handle);
+
+	void ExecuteGameplayEffect(const FGameplayEffectSpec &Spec, const FModifierQualifier &QualifierContext);
+
+	void CheckDurationExpired(FActiveGameplayEffectHandle Handle);
+
+	// --------------------------------------------
+	// Internal Utility / helper functions
+	// --------------------------------------------
+
+	bool AreGameplayEffectApplicationRequirementsSatisfied(const class UGameplayEffect* EffectToAdd, FGameplayEffectInstigatorContext& Instigator) const;
+
+	bool IsOwnerActorAuthoritative() const;
+
+	void OnAttributeGameplayEffectSpecExected(const FGameplayAttribute &Attribute, const struct FGameplayEffectSpec &Spec, struct FGameplayModifierEvaluatedData &Data);
+
+	const FGlobalCurveDataOverride * GetCurveDataOverride() const
+	{
+		// only return data if we have overrides. NULL if we dont.
+		return (GlobalCurveDataOverride.Overrides.Num() > 0 ? &GlobalCurveDataOverride : NULL);
+	}
+
+	FGlobalCurveDataOverride	GlobalCurveDataOverride;
+	
+
+	// --------------------------------------------
+	
+	UPROPERTY(ReplicatedUsing=OnRep_GameplayEffects)
+	FActiveGameplayEffectsContainer	ActiveGameplayEffects;
+
+	UFUNCTION()
+	void OnRep_GameplayEffects();
+
+	// ---------------------------------------------
+
+	UFUNCTION(NetMulticast, unreliable)
+	void NetMulticast_InvokeGameplayCueExecuted(const FGameplayEffectSpec Spec);
+
+	void InvokeGameplayCueExecute(const FGameplayEffectSpec &Spec);
+	void InvokeGameplayCueActivated(const FGameplayEffectSpec &Spec);
+	void InvokeGameplayCueAdded(const FGameplayEffectSpec &Spec);
+	void InvokeGameplayCueRemoved(const FGameplayEffectSpec &Spec);
+
+	// ---------------------------------------------
+
+	
+
+protected:
+
+	virtual void OnRegister() OVERRIDE;
+
+	UAttributeSet*	GetAttributeSubobject(const TSubclassOf<UAttributeSet> AttributeClass) const;
+	UAttributeSet*	GetAttributeSubobjectChecked(const TSubclassOf<UAttributeSet> AttributeClass) const;
+	UAttributeSet*	GetOrCreateAttributeSubobject(const TSubclassOf<UAttributeSet> AttributeClass);
+
+	friend struct FActiveGameplayEffectsContainer;
+	friend struct FActiveGameplayEffect;
+};
diff --git a/Engine/Source/Runtime/GameplayAbilities/Public/AbilitySystemDebugHUD.h b/Engine/Source/Runtime/GameplayAbilities/Public/AbilitySystemDebugHUD.h
new file mode 100644
--- /dev/null
+++ b/Engine/Source/Runtime/GameplayAbilities/Public/AbilitySystemDebugHUD.h
@@ -0,0 +1,43 @@
+// Copyright 1998-2014 Epic Games, Inc. All Rights Reserved.
+
+#pragma once
+
+#include "AbilitySystemDebugHUD.generated.h"
+
+namespace EAlignHorizontal
+{
+	enum Type
+	{
+		Left,
+		Center,
+		Right,
+	};
+}
+
+namespace EAlignVertical
+{
+	enum Type
+	{
+		Top,
+		Center,
+		Bottom,
+	};
+}
+
+UCLASS()
+class AAbilitySystemDebugHUD : public AHUD
+{
+	GENERATED_UCLASS_BODY()
+
+	/** main HUD update loop */
+	virtual void DrawHUD() OVERRIDE;
+
+	void DrawWithBackground(UFont* InFont, const FString& Text, const FColor& TextColor, EAlignHorizontal::Type HAlign, float& OffsetX, EAlignVertical::Type VAlign, float& OffsetY, float Alpha=1.f);
+
+	void DrawDebugHUD(UCanvas* Canvas, APlayerController* PC);
+
+private:
+
+	void DrawDebugAbilitySystemComponent(UAbilitySystemComponent *Component);
+
+};
\ No newline at end of file
diff --git a/Engine/Source/Runtime/GameplayAbilities/Public/AbilitySystemGlobals.h b/Engine/Source/Runtime/GameplayAbilities/Public/AbilitySystemGlobals.h
new file mode 100644
--- /dev/null
+++ b/Engine/Source/Runtime/GameplayAbilities/Public/AbilitySystemGlobals.h
@@ -0,0 +1,60 @@
+// Copyright 1998-2014 Epic Games, Inc. All Rights Reserved.
+
+#pragma once
+
+#include "AbilitySystemGlobals.generated.h"
+
+class AActor;
+class UAbilitySystemComponent;
+struct FGameplayAbilityActorInfo;
+
+/** Holds global data for the skill system. Can be configured per project via config file */
+UCLASS(config=Game)
+class GAMEPLAYABILITIES_API UAbilitySystemGlobals : public UObject
+{
+	GENERATED_UCLASS_BODY()
+
+	/** Holds all of the valid gameplay-related tags that can be applied to assets */
+	UPROPERTY(config)
+	FString GlobalCurveTableName;
+
+	/** Holds information about the valid attributes' min and max values and stacking rules */
+	UPROPERTY(config)
+	FString GlobalAttributeDataTableName;
+
+	class UCurveTable *	GetGlobalCurveTable();
+
+	class UDataTable * GetGlobalAttributeDataTable();
+
+	void AutomationTestOnly_SetGlobalCurveTable(class UCurveTable *InTable)
+	{
+		GlobalCurveTable = InTable;
+	}
+
+	void AutomationTestOnly_SetGlobalAttributeDataTable(class UDataTable *InTable)
+	{
+		GlobalAttributeDataTable = InTable;
+	}
+
+	static UAbilitySystemGlobals&	Get();
+
+	/**
+	 *	Games may want to override this in a UMyProjectAbilitySystemsGlobals class and provide
+	 *	a faster lookup for attribute compnents (E.g. Cast<MyProjPawn>(Actor)->AbilitySystemComponent; etc)
+	 */
+	virtual UAbilitySystemComponent * GetAbilitySystemComponentFromActor(AActor *Actor) const;
+
+	virtual FGameplayAbilityActorInfo * AllocAbilityActorInfo(AActor *Actor) const;
+
+private:
+
+	class UCurveTable* GlobalCurveTable;
+
+	class UDataTable* GlobalAttributeDataTable;
+
+#if WITH_EDITOR
+	void OnCurveTableReimported(UObject* InObject);
+	void OnDataTableReimported(UObject* InObject);
+#endif
+
+};
diff --git a/Engine/Source/Runtime/GameplayAbilities/Public/AbilitySystemTestAttributeSet.h b/Engine/Source/Runtime/GameplayAbilities/Public/AbilitySystemTestAttributeSet.h
new file mode 100644
--- /dev/null
+++ b/Engine/Source/Runtime/GameplayAbilities/Public/AbilitySystemTestAttributeSet.h
@@ -0,0 +1,65 @@
+// Copyright 1998-2014 Epic Games, Inc. All Rights Reserved.
+
+#pragma once
+#include "AttributeSet.h"
+#include "AbilitySystemTestAttributeSet.generated.h"
+
+UCLASS(Blueprintable, BlueprintType)
+class GAMEPLAYABILITIES_API UAbilitySystemTestAttributeSet : public UAttributeSet
+{
+	GENERATED_UCLASS_BODY()
+
+	UPROPERTY(EditAnywhere, BlueprintReadWrite, Replicated, Category = "AttributeTest", meta = (HideFromModifiers))			// You can't make a GameplayEffect modify Health Directly (go through)
+	float	MaxHealth;
+
+	UPROPERTY(EditAnywhere, BlueprintReadWrite, Replicated, Category = "AttributeTest", meta = (HideFromModifiers))			// You can't make a GameplayEffect modify Health Directly (go through)
+	float	Health;
+
+	UPROPERTY(EditAnywhere, BlueprintReadWrite, Replicated, Category = "AttributeTest")
+	float	Mana;
+
+	UPROPERTY(EditAnywhere, BlueprintReadWrite, Replicated, Category = "AttributeTest")
+	float	MaxMana;
+
+	/** This Damage is just used for applying negative health mods. Its not a 'persistent' attribute. */
+	UPROPERTY(EditAnywhere, BlueprintReadWrite, Category = "AttributeTest", meta = (HideFromLevelInfos))		// You can't make a GameplayEffect 'powered' by Damage (Its transient)
+	float	Damage;
+
+	/** This Attribute is the actual spell damage for this actor. It will power spell based GameplayEffects */
+	UPROPERTY(EditAnywhere, BlueprintReadWrite, Replicated, Category = "AttributeTest")
+	float	SpellDamage;
+
+	/** This Attribute is the actual physical damage for this actor. It will power physical based GameplayEffects */
+	UPROPERTY(EditAnywhere, BlueprintReadWrite, Replicated, Category = "AttributeTest")
+	float	PhysicalDamage;
+
+	UPROPERTY(EditAnywhere, BlueprintReadWrite, Replicated, Category = "AttributeTest")
+	float	CritChance;
+
+	UPROPERTY(EditAnywhere, BlueprintReadWrite, Replicated, Category = "AttributeTest")
+	float	CritMultiplier;
+
+	UPROPERTY(EditAnywhere, BlueprintReadWrite, Replicated, Category = "AttributeTest")
+	float	ArmorDamageReduction;
+
+	UPROPERTY(EditAnywhere, BlueprintReadWrite, Replicated, Category = "AttributeTest")
+	float	DodgeChance;
+
+	UPROPERTY(EditAnywhere, BlueprintReadWrite, Replicated, Category = "AttributeTest")
+	float	LifeSteal;
+
+	UPROPERTY(EditAnywhere, BlueprintReadWrite, Replicated, Category = "AttributeTest")
+	float	Strength;
+
+	UPROPERTY(EditAnywhere, BlueprintReadWrite, Category = "AttributeTest")
+	float	StackingAttribute1;
+
+	UPROPERTY(EditAnywhere, BlueprintReadWrite, Category = "AttributeTest")
+	float	StackingAttribute2;
+
+	UPROPERTY(EditAnywhere, BlueprintReadWrite, Category = "AttributeTest")
+	float	NoStackAttribute;
+
+	virtual void PreAttributeModify(struct FGameplayEffectModCallbackData &Data) OVERRIDE;
+	virtual void PostAttributeModify(const struct FGameplayEffectModCallbackData &Data) OVERRIDE;
+};
diff --git a/Engine/Source/Runtime/GameplayAbilities/Public/AbilitySystemTestPawn.h b/Engine/Source/Runtime/GameplayAbilities/Public/AbilitySystemTestPawn.h
new file mode 100644
--- /dev/null
+++ b/Engine/Source/Runtime/GameplayAbilities/Public/AbilitySystemTestPawn.h
@@ -0,0 +1,39 @@
+// Copyright 1998-2014 Epic Games, Inc. All Rights Reserved.
+
+#pragma once
+#include "GameplayCueInterface.h"
+#include "GameplayCueView.h"
+#include "AbilitySystemTestPawn.generated.h"
+
+UCLASS(Blueprintable, BlueprintType)
+class GAMEPLAYABILITIES_API AAbilitySystemTestPawn : public ADefaultPawn, public IGameplayCueInterface
+{
+	GENERATED_UCLASS_BODY()
+
+	virtual void PostInitializeComponents() OVERRIDE;
+
+	UFUNCTION()
+	virtual void GameplayCueActivated(const FGameplayTagContainer & GameplayCueTags, float NormalizedMagnitude, const FGameplayEffectInstigatorContext InstigatorContext) OVERRIDE;
+
+	UFUNCTION()
+	virtual void GameplayCueExecuted(const FGameplayTagContainer & GameplayCueTags, float NormalizedMagnitude, const FGameplayEffectInstigatorContext InstigatorContext) OVERRIDE;
+
+	UFUNCTION()
+	virtual void GameplayCueAdded(const FGameplayTagContainer & GameplayCueTags, float NormalizedMagnitude, const FGameplayEffectInstigatorContext InstigatorContext) OVERRIDE;
+
+	UFUNCTION()
+	virtual void GameplayCueRemoved(const FGameplayTagContainer & GameplayCueTags, float NormalizedMagnitude, const FGameplayEffectInstigatorContext InstigatorContext) OVERRIDE;
+
+
+	/** DefaultPawn collision component */
+	UPROPERTY(Category = AbilitySystem, VisibleAnywhere, BlueprintReadOnly)
+	TSubobjectPtr<class UAbilitySystemComponent> AbilitySystemComponent;
+
+	UPROPERTY(EditDefaultsOnly, Category=GameplayEffects)
+	FGameplayCueHandler	GameplayCueHandler;
+
+	//UPROPERTY(EditDefaultsOnly, Category=GameplayEffects)
+	//UGameplayAbilitySet * DefaultAbilitySet;
+
+	static FName AbilitySystemComponentName;
+};
\ No newline at end of file
diff --git a/Engine/Source/Runtime/GameplayAbilities/Public/AbilitySystemTypes.h b/Engine/Source/Runtime/GameplayAbilities/Public/AbilitySystemTypes.h
--- a/Engine/Source/Runtime/GameplayAbilities/Public/AbilitySystemTypes.h
+++ b/Engine/Source/Runtime/GameplayAbilities/Public/AbilitySystemTypes.h
@@ -1,218 +1,3 @@
 ﻿// Copyright 1998-2014 Epic Games, Inc. All Rights Reserved.
 
 #pragma once
-
-#include "AbilitySystemTypes.generated.h"
-
-class UAbilitySystemComponent;
-class UGameplayEffect;
-
-UENUM(BlueprintType)
-namespace EGameplayAbilityActivationMode
-{
-	enum Type
-	{
-		Authority,			// We are the authority activating this ability
-		NonAuthority,		// We are not the authority but aren't predicting yet. This is a mostly invalid state to be in.
-
-		Predicting,			// We are predicting the activation of this ability
-		Confirmed,			// We are not the authority, but the authority has confirmed this activation
-	};
-}
-
-/**
- *	FGameplayAbilityActorInfo 
- *	
- *	Cached data associated with an Actor using an Ability.
- *		-Initialized from an AActor* in InitFromActor
- *		-Abilities use this to know what to actor upon. E.g., instead of being coupled to a specific actor class.
- *		-These are generally passed around as pointers to support polymorphism.
- *		-Projects can override UAbilitySystemGlobals::AllocAbilityActorInfo to override the default struct type that is created. 
- *			
- */
-USTRUCT(BlueprintType)
-struct GAMEPLAYABILITIES_API FGameplayAbilityActorInfo
-{
-	GENERATED_USTRUCT_BODY()
-
-	virtual ~FGameplayAbilityActorInfo() {}
-
-	UPROPERTY(BlueprintReadOnly, Category = "ActorInfo")
-	TWeakObjectPtr<AActor>	Actor;
-
-	UPROPERTY(BlueprintReadOnly, Category = "ActorInfo")
-	TWeakObjectPtr<UAnimInstance>	AnimInstance;
-
-	UPROPERTY(BlueprintReadOnly, Category = "ActorInfo")
-	TWeakObjectPtr<UAbilitySystemComponent>	AbilitySystemComponent;
-
-	UPROPERTY(BlueprintReadOnly, Category = "ActorInfo")
-	TWeakObjectPtr<UMovementComponent>	MovementComponent;
-
-	virtual void InitFromActor(AActor *Actor);
-};
-
-/**
- *	FGameplayAbilityActivationInfo
- *	
- *	Data tied to a specific activation of an ability. 
- *		-Tell us whether we are the authority, if we are predicting, confirmed, etc.
- *		-Holds PredictionKey
- *		-Generally not meant to be subclassed in projects.
- *		-Passed around by value since the struct is small.
- */
-
-USTRUCT(BlueprintType)
-struct GAMEPLAYABILITIES_API FGameplayAbilityActivationInfo
-{
-	GENERATED_USTRUCT_BODY()
-
-	FGameplayAbilityActivationInfo()
-		: ActivationMode(EGameplayAbilityActivationMode::Authority),
-		PredictionKey(0)
-	{
-
-	}
-
-	FGameplayAbilityActivationInfo(AActor * InActor, int32 InPredictionKey)
-		: PredictionKey(InPredictionKey)
-	{
-		// On Init, we are either Authority or NonAuthority. We haven't been given a PredictionKey and we haven't been confirmed.
-		// NonAuthority essentially means 'I'm not sure what how I'm going to do this yet'.
-		ActivationMode = (InActor->Role == ROLE_Authority ? EGameplayAbilityActivationMode::Authority : EGameplayAbilityActivationMode::NonAuthority);
-	}
-
-	FGameplayAbilityActivationInfo(EGameplayAbilityActivationMode::Type InType, int32 InPredictionKey)
-		: ActivationMode(InType)
-		, PredictionKey(InPredictionKey)
-	{
-	}
-
-
-	void GeneratePredictionKey(UAbilitySystemComponent * Component) const;
-
-	void SetPredictionKey(int32 InPredictionKey);
-
-	void SetActivationConfirmed();
-
-	UPROPERTY(BlueprintReadOnly, Category = "ActorInfo")
-	mutable TEnumAsByte<EGameplayAbilityActivationMode::Type>	ActivationMode;
-
-	UPROPERTY()
-	mutable int32 PredictionKey;
-};
-
-
-// ----------------------------------------------------
-
-/**
- *	A generic structure for targeting data. We want generic functions to produce this data and other generic
- *	functions to consume this data.
- *	
- *	We expect this to be able to hold specific actors/object reference and also generic location/direction/origin 
- *	information.
- *	
- *	Some example producers:
- *		-Overlap/Hit collision event generates TargetData about who was hit in a melee attack
- *		-A mouse input causes a hit trace and the actor infront of the crosshair is turned into TargetData
- *		-A mouse input causes TargetData to be generated from the owner's crosshair view origin/direction
- *		-An AOE/aura pulses and all actors in a radius around the instigator are added to TargetData
- *		-Panzer Dragoon style 'painting' targeting mode
- *		-MMORPG style ground AOE targeting style (potentially both a location on the ground and actors that were targeted)
- *		
- *	Some example consumers:
- *		-Apply a GameplayEffect to all actors in TargetData
- *		-Find closest actor from all in TargetData
- *		-Call some function on all actors in TargetData
- *		-Filter or merge TargetDatas
- *		-Spawn a new actor at a TargetData location
- *		
- *		
- *		
- *	Maybe it is better to distinguish between actor list targeting vs positional targeting data?
- *		-AOE/aura type of targeting data blurs the line
- *		
- *		
- */
-
-struct GAMEPLAYABILITIES_API FGameplayAbilityTargetData
-{
-	virtual ~FGameplayAbilityTargetData() { }
-
-	void ApplyGameplayEffect(UGameplayEffect *GameplayEffect, const FGameplayAbilityActorInfo InstigatorInfo);
-
-	virtual TArray<AActor*>	GetActors() const = 0;
-	
-	virtual bool HasHitResult() const 
-	{	
-		return false;
-	}
-
-	virtual const FHitResult * GetHitResult() const
-	{
-		return NULL;
-	}
-
-	virtual FString ToString() const;
-};
-
-struct GAMEPLAYABILITIES_API FGameplayAbilityTargetData_SingleTargetHit : public FGameplayAbilityTargetData
-{
-	FGameplayAbilityTargetData_SingleTargetHit(const FHitResult InHitResult)
-	: HitResult(InHitResult)
-	{ }
-
-	virtual TArray<AActor*>	GetActors() const
-	{
-		TArray<AActor*>	Actors;
-		if (HitResult.Actor.IsValid())
-		{
-			Actors.Push(HitResult.Actor.Get());
-		}
-		return Actors;
-	}
-
-	virtual bool HasHitResult() const
-	{
-		return true;
-	}
-
-	virtual const FHitResult * GetHitResult() const
-	{
-		return &HitResult;
-	}
-
-	FHitResult	HitResult;
-};
-
-/**
- *	Handle for Targeting Data. This servers two main purposes:
- *		-Avoid us having to copy around the full targeting data structure in Blueprints
- *		-Allows us to leverage polymorphism in the target data structure
- *		-Allows us to implement NetSerialize and replicate by value between clients/server	
- *		
- *		-Avoid using UObjects could be used to give us polymorphism and by reference passing in blueprints.
- *		-However we would still be screwed when it came to replication
- *		
- *		-Replication by value
- *		-Pass by reference in blueprints
- *		-Polymophism in TargetData structure
- *		
- */
-
-USTRUCT(BlueprintType)
-struct GAMEPLAYABILITIES_API FGameplayAbilityTargetDataHandle
-{
-	GENERATED_USTRUCT_BODY()
-	
-	TSharedPtr<FGameplayAbilityTargetData>	Data;
-};
-
-template<>
-struct TStructOpsTypeTraits<FGameplayAbilityTargetDataHandle> : public TStructOpsTypeTraitsBase
-{
-	enum
-	{
-		WithCopy = true		// Necessary so that TSharedPtr<FGameplayAbilityTargetData> Data is copied around
-	};
-};
diff --git a/Engine/Source/Runtime/GameplayAbilities/Public/AttributeSet.h b/Engine/Source/Runtime/GameplayAbilities/Public/AttributeSet.h
new file mode 100644
--- /dev/null
+++ b/Engine/Source/Runtime/GameplayAbilities/Public/AttributeSet.h
@@ -0,0 +1,272 @@
+// Copyright 1998-2014 Epic Games, Inc. All Rights Reserved.
+
+#pragma once
+
+#include "AttributeSet.generated.h"
+
+USTRUCT(BlueprintType)
+struct GAMEPLAYABILITIES_API FGameplayAttribute
+{
+	GENERATED_USTRUCT_BODY()
+
+	FGameplayAttribute()
+		: Attribute(NULL)
+	{
+	}
+
+	FGameplayAttribute(UProperty *NewProperty)
+		: Attribute(NewProperty)
+	{
+	}
+
+	void SetUProperty(UProperty *NewProperty)
+	{
+		Attribute = NewProperty;
+	}
+
+	UProperty * GetUProperty() const
+	{
+		return Attribute;
+	}
+
+	UClass * GetAttributeSetClass() const
+	{
+		check(Attribute);
+		return CastChecked<UClass>(Attribute->GetOuter());
+	}
+
+	void SetNumericValueChecked(const float NewValue, class UAttributeSet* Dest) const;
+
+	float GetNumericValueChecked(class UAttributeSet* Src) const;
+	
+	/** Equality/Inequality operators */
+	bool operator==(const FGameplayAttribute& Other) const;
+	bool operator!=(const FGameplayAttribute& Other) const;
+
+	friend uint32 GetTypeHash( const FGameplayAttribute& InAttribute )
+	{
+		// FIXME: Use ObjectID or something to get a better, less collision prone hash
+		return PointerHash(InAttribute.Attribute);
+	}
+
+	FString GetName() const
+	{
+		return Attribute->GetName();
+	}
+
+private:
+	friend class FAttributePropertyDetails;
+
+	UPROPERTY(Category=GameplayAttribute, EditDefaultsOnly)
+	UProperty*	Attribute;
+};
+
+UCLASS(DefaultToInstanced, Blueprintable)
+class GAMEPLAYABILITIES_API UAttributeSet : public UObject
+{
+	GENERATED_UCLASS_BODY()
+
+public:
+
+	bool IsSupportedForNetworking() const OVERRIDE
+	{
+		return true;
+	}
+
+	/** Called just before modifying the value of an attribute. AttributeSet can make additional modifications here. */
+	virtual void PreAttributeModify(struct FGameplayEffectModCallbackData &Data) { }
+	
+	/** Called just after modifying the value of an attribute. No more changes can be made */
+	virtual void PostAttributeModify(const struct FGameplayEffectModCallbackData &Data) { }
+
+	void InitFromMetaDataTable(const UDataTable* DataTable);
+
+	void ClampFromMetaDataTable(const UDataTable *DataTable);
+
+	virtual void PrintDebug();
+
+private:
+
+	void OnAttributeChange(UProperty *Property);
+};
+
+USTRUCT()
+struct GAMEPLAYABILITIES_API FGlobalCurveDataOverride
+{
+	GENERATED_USTRUCT_BODY()
+
+	TArray<UCurveTable*>	Overrides;
+};
+
+USTRUCT()
+struct GAMEPLAYABILITIES_API FScalableFloat
+{
+	GENERATED_USTRUCT_BODY()
+
+	FScalableFloat()
+		: Value(0.f)
+		, FinalCurve(NULL)
+	{
+
+	}
+
+public:
+
+	UPROPERTY(EditDefaultsOnly, Category=ScalableFloat)
+	float	Value;
+
+	UPROPERTY(EditDefaultsOnly, Category=ScalableFloat)
+	FCurveTableRowHandle	Curve;
+
+	void FinalizeCurveData(const FGlobalCurveDataOverride *GlobalOverrides);
+
+	float GetValueAtLevel(float Level) const;
+
+	FScalableFloat MakeFinalizedCopy(const FGlobalCurveDataOverride *GlobalOverrides) const
+	{
+		FScalableFloat Copy(*this);
+		Copy.FinalizeCurveData(GlobalOverrides);
+		return Copy;
+	}
+
+	bool IsStatic() const
+	{
+		return (Curve.RowName == NAME_None);
+	}
+
+	void SetValue(float NewValue)
+	{
+		Value = NewValue;
+		Curve.CurveTable = NULL;
+		Curve.RowName = NAME_None;
+		FinalCurve = NULL;
+	}
+
+	void SetScalingValue(float InCoeffecient, FName InRowName, UCurveTable * InTable)
+	{
+		Value = InCoeffecient;
+		Curve.RowName = InRowName;
+		Curve.CurveTable = InTable;
+		FinalCurve = NULL;
+	}
+
+	void LockValueAtLevel(float Level, FGlobalCurveDataOverride *GlobalOverrides)
+	{
+		SetValue(GetValueAtLevel(Level));
+	}
+
+	float GetValueChecked() const
+	{
+		check(IsStatic());
+		return Value;
+	}
+
+	FString ToSimpleString() const
+	{
+		if (Curve.RowName != NAME_None)
+		{
+			return FString::Printf(TEXT("%.2f - %s@%s"), Value, *Curve.RowName.ToString(), Curve.CurveTable ? *Curve.CurveTable->GetName() : TEXT("None"));
+		}
+		return FString::Printf(TEXT("%.2f"), Value);
+	}
+
+	/** Equality/Inequality operators */
+	bool operator==(const FScalableFloat& Other) const;
+	bool operator!=(const FScalableFloat& Other) const;
+
+private:
+
+	FRichCurve * FinalCurve;
+};
+
+
+/**
+ *	This is an example of using the property name as the row handle name.
+ *	FAttributeMetaData is a table where each row is a single property in a UAttribute class.
+ *
+ *	I need an exmaple of having a property reference in line
+ *
+*/
+USTRUCT(BlueprintType)
+struct GAMEPLAYABILITIES_API FAttributeMetaData : public FTableRowBase
+{
+	GENERATED_USTRUCT_BODY()
+
+public:
+
+	FAttributeMetaData();
+
+	UPROPERTY()
+	float		BaseValue;
+
+	UPROPERTY()
+	float		MinValue;
+
+	UPROPERTY()
+	float		MaxValue;
+
+	UPROPERTY()
+	FString		DerivedAttributeInfo;
+
+	UPROPERTY()
+	bool		bCanStack;
+};
+
+
+UENUM()
+enum EAttributeModifierType
+{
+	EATTRMOD_Add,
+	EATTRMOD_Multiple,
+	EATTRMOD_Override,
+};
+
+USTRUCT(BlueprintType)
+struct FAttributeModifierTest : public FTableRowBase
+{
+	GENERATED_USTRUCT_BODY()
+
+public:
+
+	UPROPERTY()
+	FString	PropertyName;
+
+	UPROPERTY()
+	TEnumAsByte<EAttributeModifierType>	ModifierType;
+
+	UPROPERTY()
+	FString NewValue;
+
+	UProperty*	GetUProperty();
+
+	void	ApplyModifier(UAttributeSet *Object);
+
+private:
+
+	UProperty *CachedUProperty;
+};
+
+USTRUCT()
+struct GAMEPLAYABILITIES_API FAttributeModifier
+{
+	GENERATED_USTRUCT_BODY()
+
+	UPROPERTY()
+	FGameplayAttribute	PropertyToModify;
+
+	UPROPERTY()
+	float		NumericValue;
+
+	UPROPERTY()
+	TEnumAsByte<EAttributeModifierType>	ModifierType;
+};
+
+USTRUCT()
+struct FSomeThingThatHoldsAListOfOtherthings
+{
+	GENERATED_USTRUCT_BODY()
+
+	UPROPERTY()
+	TArray< FAttributeModifier	>	ListOfModifiers;
+
+};
\ No newline at end of file
diff --git a/Engine/Source/Runtime/GameplayAbilities/Public/GameplayAbilitySet.h b/Engine/Source/Runtime/GameplayAbilities/Public/GameplayAbilitySet.h
new file mode 100644
--- /dev/null
+++ b/Engine/Source/Runtime/GameplayAbilities/Public/GameplayAbilitySet.h
@@ -0,0 +1,56 @@
+// Copyright 1998-2013 Epic Games, Inc. All Rights Reserved.
+
+#pragma once
+
+#include "GameplayAbilitySet.generated.h"
+
+USTRUCT()
+struct FGameplayAbilityBindInfoCommandIDPair
+{
+	GENERATED_USTRUCT_BODY()
+
+	// We want this to be a drop down of valid lists but needs to be defined per game
+	UPROPERTY(EditAnywhere, Category = BindInfo)
+	FString		CommandString;
+
+	UPROPERTY(EditAnywhere, Category = BindInfo)
+	int32		InputID;
+};
+
+USTRUCT()
+struct FGameplayAbilityBindInfo
+{
+	GENERATED_USTRUCT_BODY()
+
+	UPROPERTY(EditAnywhere, Category=BindInfo)
+	TArray<FGameplayAbilityBindInfoCommandIDPair>	CommandBindings;
+
+	UPROPERTY(EditAnywhere, Category=BindInfo)
+	TSubclassOf<UGameplayAbility>	GameplayAbilityClass;
+
+	UPROPERTY(EditAnywhere, Category = BindInfo)
+	class UGameplayAbility * GameplayAbilityInstance;
+};
+
+/**
+* UGameplayAbilitySet
+*	This contains a list of abilities along with key bindings. This will be very game specific, so it is expected for games to override this class.
+*	Functions like InitializeAbilities and BindInputComponentToAbilities are setup to act on the AbilitySystemComponent (rather than AbilitySystemComponent processing the ability set).
+*	Ideally this avoids games having to implement their own UAbilitySystemComponent and UGameplayAbilitySet, and they can just implement thei own UGameplayAbilitySet.
+*/
+UCLASS(Blueprintable)
+class GAMEPLAYABILITIES_API UGameplayAbilitySet : public UDataAsset
+{
+	GENERATED_UCLASS_BODY()
+
+public:
+
+	UPROPERTY(EditAnywhere, Category=AbilitySet)
+	TArray<FGameplayAbilityBindInfo>	Abilities;
+
+	/** Give all of the abilities in this set to given AbilitySystemComponent */
+	virtual void InitializeAbilities(UAbilitySystemComponent *AbilitySystemComponent) const;
+
+	/** Binds the ability set directly to this inputcomponent and AbilitySystemComponent. This is a convenience method, not all games will want to bind input directly to abilities. */
+	virtual void BindInputComponentToAbilities(UInputComponent *InputComponent, UAbilitySystemComponent *AbilitySystemComponent) const;
+};
\ No newline at end of file
diff --git a/Engine/Source/Runtime/GameplayAbilities/Public/GameplayAbilityTypes.h b/Engine/Source/Runtime/GameplayAbilities/Public/GameplayAbilityTypes.h
new file mode 100644
--- /dev/null
+++ b/Engine/Source/Runtime/GameplayAbilities/Public/GameplayAbilityTypes.h
@@ -0,0 +1,341 @@
+// Copyright 1998-2014 Epic Games, Inc. All Rights Reserved.
+
+#pragma once
+
+#include "GameplayAbilityTypes.generated.h"
+
+class UGameplayEffect;
+class UAnimInstance;
+class UAbilitySystemComponent;
+
+
+/** 
+ *
+ *
+ */
+
+
+
+UENUM(BlueprintType)
+namespace EGameplayAbilityInstancingPolicy
+{
+	/**
+	*	How the ability is instanced when executed. This limits what an ability can do in its implementation. For example, a NonInstanced
+	*	Ability cannot have state. It is probably unsafe for an InstancedPerActor ability to have latent actions, etc.
+	*/
+
+	enum Type
+	{
+		NonInstanced,			// This ability is never instanced. Anything that executes the ability is operating on the CDO.
+		InstancedPerActor,		// Each actor gets their own instance of this ability. State can be saved, replication is possible.
+		InstancedPerExecution,	// We instance this ability each time it is executed. Replication possible but not recommended.
+	};
+}
+
+UENUM(BlueprintType)
+namespace EGameplayAbilityNetExecutionPolicy
+{
+	/**
+	*	How does an ability execute on the network. Does a client "ask and predict", "ask and wait", "don't ask"
+	*/
+
+	enum Type
+	{
+		Predictive		UMETA(DisplayName = "Predictive"),	// Part of this ability runs predictively on the client.
+		Server			UMETA(DisplayName = "Server"),		// This ability must be OK'd by the server before doing anything on a client.
+		Client			UMETA(DisplayName = "Client"),		// This ability runs as long the client says it does.
+	};
+}
+
+UENUM(BlueprintType)
+namespace EGameplayAbilityReplicationPolicy
+{
+	/**
+	*	How an ability replicates state/events to everyone on the network.
+	*/
+
+	enum Type
+	{
+		ReplicateNone		UMETA(DisplayName = "Replicate None"),	// We don't replicate the instance of the ability to anyone.
+		ReplicateAll		UMETA(DisplayName = "Replicate All"),	// We replicate the instance of the ability to everyone (even simulating clients).
+		ReplicateOwner		UMETA(DisplayName = "Replicate Owner"),	// Only replicate the instance of the ability to the owner.
+	};
+}
+
+USTRUCT(BlueprintType)
+struct FGameplayAbilityHandle
+{
+	GENERATED_USTRUCT_BODY()
+
+	FGameplayAbilityHandle()
+	: Handle(INDEX_NONE)
+	{
+
+	}
+
+	FGameplayAbilityHandle(int32 InHandle)
+		: Handle(InHandle)
+	{
+
+	}
+
+	bool IsValid() const
+	{
+		return Handle != INDEX_NONE;
+	}
+
+	FGameplayAbilityHandle GetNextHandle()
+	{
+		return FGameplayAbilityHandle(Handle + 1);
+	}
+
+	bool operator==(const FGameplayAbilityHandle& Other) const
+	{
+		return Handle == Other.Handle;
+	}
+
+	bool operator!=(const FGameplayAbilityHandle& Other) const
+	{
+		return Handle != Other.Handle;
+	}
+
+	FString ToString() const
+	{
+		return FString::Printf(TEXT("%d"), Handle);
+	}
+
+private:
+
+	UPROPERTY()
+	int32 Handle;
+};
+
+/**
+*	Not implemented yet, but we will need something like this to track ability + levels if we choose to support ability leveling in the base classes.
+*/
+
+USTRUCT()
+struct FGameplayAbilitySpec
+{
+	GENERATED_USTRUCT_BODY()
+
+	UPROPERTY()
+	int32	Level;
+
+	UPROPERTY()
+	class UGameplayAbility * Ability;
+
+	UPROPERTY()
+	FGameplayAbilityHandle	Handle;
+};
+
+
+UENUM(BlueprintType)
+namespace EGameplayAbilityActivationMode
+{
+	enum Type
+	{
+		Authority,			// We are the authority activating this ability
+		NonAuthority,		// We are not the authority but aren't predicting yet. This is a mostly invalid state to be in.
+
+		Predicting,			// We are predicting the activation of this ability
+		Confirmed,			// We are not the authority, but the authority has confirmed this activation
+	};
+}
+
+/**
+*	FGameplayAbilityActorInfo
+*
+*	Cached data associated with an Actor using an Ability.
+*		-Initialized from an AActor* in InitFromActor
+*		-Abilities use this to know what to actor upon. E.g., instead of being coupled to a specific actor class.
+*		-These are generally passed around as pointers to support polymorphism.
+*		-Projects can override UAbilitySystemGlobals::AllocAbilityActorInfo to override the default struct type that is created.
+*
+*/
+USTRUCT(BlueprintType)
+struct GAMEPLAYABILITIES_API FGameplayAbilityActorInfo
+{
+	GENERATED_USTRUCT_BODY()
+
+	virtual ~FGameplayAbilityActorInfo() {}
+
+	UPROPERTY(BlueprintReadOnly, Category = "ActorInfo")
+	TWeakObjectPtr<AActor>	Actor;
+
+	UPROPERTY(BlueprintReadOnly, Category = "ActorInfo")
+	TWeakObjectPtr<UAnimInstance>	AnimInstance;
+
+	UPROPERTY(BlueprintReadOnly, Category = "ActorInfo")
+	TWeakObjectPtr<UAbilitySystemComponent>	AbilitySystemComponent;
+
+	UPROPERTY(BlueprintReadOnly, Category = "ActorInfo")
+	TWeakObjectPtr<UMovementComponent>	MovementComponent;
+
+	virtual void InitFromActor(AActor *Actor);
+};
+
+/**
+*	FGameplayAbilityActivationInfo
+*
+*	Data tied to a specific activation of an ability.
+*		-Tell us whether we are the authority, if we are predicting, confirmed, etc.
+*		-Holds PredictionKey
+*		-Generally not meant to be subclassed in projects.
+*		-Passed around by value since the struct is small.
+*/
+
+USTRUCT(BlueprintType)
+struct GAMEPLAYABILITIES_API FGameplayAbilityActivationInfo
+{
+	GENERATED_USTRUCT_BODY()
+
+	FGameplayAbilityActivationInfo()
+	: ActivationMode(EGameplayAbilityActivationMode::Authority),
+	PredictionKey(0)
+	{
+
+	}
+
+	FGameplayAbilityActivationInfo(AActor * InActor, int32 InPredictionKey)
+		: PredictionKey(InPredictionKey)
+	{
+		// On Init, we are either Authority or NonAuthority. We haven't been given a PredictionKey and we haven't been confirmed.
+		// NonAuthority essentially means 'I'm not sure what how I'm going to do this yet'.
+		ActivationMode = (InActor->Role == ROLE_Authority ? EGameplayAbilityActivationMode::Authority : EGameplayAbilityActivationMode::NonAuthority);
+	}
+
+	FGameplayAbilityActivationInfo(EGameplayAbilityActivationMode::Type InType, int32 InPredictionKey)
+		: ActivationMode(InType)
+		, PredictionKey(InPredictionKey)
+	{
+	}
+
+
+	void GeneratePredictionKey(UAbilitySystemComponent * Component) const;
+
+	void SetPredictionKey(int32 InPredictionKey);
+
+	void SetActivationConfirmed();
+
+	UPROPERTY(BlueprintReadOnly, Category = "ActorInfo")
+	mutable TEnumAsByte<EGameplayAbilityActivationMode::Type>	ActivationMode;
+
+	UPROPERTY()
+	mutable int32 PredictionKey;
+};
+
+
+// ----------------------------------------------------
+
+/**
+*	A generic structure for targeting data. We want generic functions to produce this data and other generic
+*	functions to consume this data.
+*
+*	We expect this to be able to hold specific actors/object reference and also generic location/direction/origin
+*	information.
+*
+*	Some example producers:
+*		-Overlap/Hit collision event generates TargetData about who was hit in a melee attack
+*		-A mouse input causes a hit trace and the actor infront of the crosshair is turned into TargetData
+*		-A mouse input causes TargetData to be generated from the owner's crosshair view origin/direction
+*		-An AOE/aura pulses and all actors in a radius around the instigator are added to TargetData
+*		-Panzer Dragoon style 'painting' targeting mode
+*		-MMORPG style ground AOE targeting style (potentially both a location on the ground and actors that were targeted)
+*
+*	Some example consumers:
+*		-Apply a GameplayEffect to all actors in TargetData
+*		-Find closest actor from all in TargetData
+*		-Call some function on all actors in TargetData
+*		-Filter or merge TargetDatas
+*		-Spawn a new actor at a TargetData location
+*
+*
+*
+*	Maybe it is better to distinguish between actor list targeting vs positional targeting data?
+*		-AOE/aura type of targeting data blurs the line
+*
+*
+*/
+
+struct GAMEPLAYABILITIES_API FGameplayAbilityTargetData
+{
+	virtual ~FGameplayAbilityTargetData() { }
+
+	void ApplyGameplayEffect(UGameplayEffect *GameplayEffect, const FGameplayAbilityActorInfo InstigatorInfo);
+
+	virtual TArray<AActor*>	GetActors() const = 0;
+
+	virtual bool HasHitResult() const
+	{
+		return false;
+	}
+
+	virtual const FHitResult * GetHitResult() const
+	{
+		return NULL;
+	}
+
+	virtual FString ToString() const;
+};
+
+struct GAMEPLAYABILITIES_API FGameplayAbilityTargetData_SingleTargetHit : public FGameplayAbilityTargetData
+{
+	FGameplayAbilityTargetData_SingleTargetHit(const FHitResult InHitResult)
+	: HitResult(InHitResult)
+	{ }
+
+	virtual TArray<AActor*>	GetActors() const
+	{
+		TArray<AActor*>	Actors;
+		if (HitResult.Actor.IsValid())
+		{
+			Actors.Push(HitResult.Actor.Get());
+		}
+		return Actors;
+	}
+
+	virtual bool HasHitResult() const
+	{
+		return true;
+	}
+
+	virtual const FHitResult * GetHitResult() const
+	{
+		return &HitResult;
+	}
+
+	FHitResult	HitResult;
+};
+
+/**
+*	Handle for Targeting Data. This servers two main purposes:
+*		-Avoid us having to copy around the full targeting data structure in Blueprints
+*		-Allows us to leverage polymorphism in the target data structure
+*		-Allows us to implement NetSerialize and replicate by value between clients/server
+*
+*		-Avoid using UObjects could be used to give us polymorphism and by reference passing in blueprints.
+*		-However we would still be screwed when it came to replication
+*
+*		-Replication by value
+*		-Pass by reference in blueprints
+*		-Polymophism in TargetData structure
+*
+*/
+
+USTRUCT(BlueprintType)
+struct GAMEPLAYABILITIES_API FGameplayAbilityTargetDataHandle
+{
+	GENERATED_USTRUCT_BODY()
+
+	TSharedPtr<FGameplayAbilityTargetData>	Data;
+};
+
+template<>
+struct TStructOpsTypeTraits<FGameplayAbilityTargetDataHandle> : public TStructOpsTypeTraitsBase
+{
+	enum
+	{
+		WithCopy = true		// Necessary so that TSharedPtr<FGameplayAbilityTargetData> Data is copied around
+	};
+};
diff --git a/Engine/Source/Runtime/GameplayAbilities/Public/GameplayCueInterface.h b/Engine/Source/Runtime/GameplayAbilities/Public/GameplayCueInterface.h
new file mode 100644
--- /dev/null
+++ b/Engine/Source/Runtime/GameplayAbilities/Public/GameplayCueInterface.h
@@ -0,0 +1,34 @@
+// Copyright 1998-2014 Epic Games, Inc. All Rights Reserved.
+
+#pragma once
+
+#include "GameplayTagContainer.h"
+#include "GameplayEffect.h"
+#include "GameplayCueInterface.generated.h"
+
+/** Interface for actors that wish to handle GameplayCue events from GameplayEffects */
+UINTERFACE(MinimalAPI)
+class UGameplayCueInterface: public UInterface
+{
+	GENERATED_UINTERFACE_BODY()
+};
+
+class GAMEPLAYABILITIES_API IGameplayCueInterface
+{
+	GENERATED_IINTERFACE_BODY()
+
+	// EVENT
+	UFUNCTION()
+	virtual void GameplayCueActivated(const FGameplayTagContainer & GameplayCueTags, float NormalizedMagnitude, const FGameplayEffectInstigatorContext InstigatorContext) = 0;
+
+	UFUNCTION()
+	virtual void GameplayCueExecuted(const FGameplayTagContainer & GameplayCueTags, float NormalizedMagnitude, const FGameplayEffectInstigatorContext InstigatorContext) = 0;
+
+
+	// STATE
+	UFUNCTION()
+	virtual void GameplayCueAdded(const FGameplayTagContainer & GameplayCueTags, float NormalizedMagnitude, const FGameplayEffectInstigatorContext InstigatorContext) = 0;
+
+	UFUNCTION()
+	virtual void GameplayCueRemoved(const FGameplayTagContainer & GameplayCueTags, float NormalizedMagnitude, const FGameplayEffectInstigatorContext InstigatorContext) = 0;
+};
\ No newline at end of file
diff --git a/Engine/Source/Runtime/GameplayAbilities/Public/GameplayCueView.h b/Engine/Source/Runtime/GameplayAbilities/Public/GameplayCueView.h
new file mode 100644
--- /dev/null
+++ b/Engine/Source/Runtime/GameplayAbilities/Public/GameplayCueView.h
@@ -0,0 +1,128 @@
+// Copyright 1998-2014 Epic Games, Inc. All Rights Reserved.
+
+#pragma once
+
+#include "GameplayTags.h"
+#include "GameplayEffect.h"
+#include "GameplayCueView.generated.h"
+
+/** 
+	This is meant to give a base implementation for handling GameplayCues. It will handle very simple things like
+	creating and destroying ParticleSystemComponents, playing audio, etc.
+
+	It is inevitable that specific actors and games will have to handle GameplayCues in their own way.
+*/
+
+UENUM(BlueprintType)
+namespace EGameplayCueEvent
+{
+	enum Type
+	{
+		Applied,
+		Added,
+		Executed,
+		Removed
+	};
+}
+
+USTRUCT()
+struct FGameplayCueViewEffects
+{
+	// The instantied effects
+	GENERATED_USTRUCT_BODY()
+
+	UPROPERTY()
+	TWeakObjectPtr<UParticleSystemComponent> ParticleSystemComponent;
+
+	UPROPERTY()
+	TWeakObjectPtr<UAudioComponent>	AudioComponent;
+
+	UPROPERTY()
+	TWeakObjectPtr<AActor>	SpawnedActor;
+};
+
+USTRUCT()
+struct FGameplayCueViewInfo
+{
+	GENERATED_USTRUCT_BODY()
+
+	// -----------------------------
+	//	Qualifiers
+	// -----------------------------
+
+	UPROPERTY(EditDefaultsOnly, Category = GameplayCue)
+	FGameplayTagContainer Tags;
+
+	UPROPERTY(EditDefaultsOnly, Category = GameplayModifier)
+	TEnumAsByte<EGameplayCueEvent::Type> CueType;
+
+	// This needs to be fleshed out more:
+	bool	InstigatorLocalOnly;
+	bool	TargetLocalOnly;
+
+	// -----------------------------
+	//	Effects
+	// -----------------------------
+
+	/** Local/remote sounds to play for weapon attacks against specific surfaces */
+	UPROPERTY(EditAnywhere, BlueprintReadWrite, Category = GameplayCue)
+	class USoundBase* Sound;
+
+	/** Effects to play for weapon attacks against specific surfaces */
+	UPROPERTY(EditAnywhere, BlueprintReadWrite, Category = GameplayCue)
+	class UParticleSystem* ParticleSystem;
+
+	/** Effects to play for weapon attacks against specific surfaces */
+	UPROPERTY(EditAnywhere, BlueprintReadWrite, Category = GameplayCue)
+	class TSubclassOf<AActor> ActorClass;
+
+	virtual TSharedPtr<FGameplayCueViewEffects> SpawnViewEffects(AActor *Owner, TArray<UObject*> *SpawnedObjects, const FGameplayEffectInstigatorContext InstigatorContext) const;
+};
+
+UCLASS(BlueprintType)
+class GAMEPLAYABILITIES_API UGameplayCueView : public UDataAsset
+{
+	// Defines how a GameplayCue is translated into viewable components
+	GENERATED_UCLASS_BODY()
+
+	UPROPERTY(EditDefaultsOnly, Category = GameplayCue)
+	TArray<FGameplayCueViewInfo>	Views;
+};
+
+USTRUCT()
+struct GAMEPLAYABILITIES_API FGameplayCueHandler
+{
+	GENERATED_USTRUCT_BODY()
+
+	FGameplayCueHandler()
+		: Owner(NULL)
+	{
+	}
+
+	UPROPERTY()
+	AActor * Owner;
+
+	UPROPERTY(EditDefaultsOnly, Category = GameplayCue)
+	TArray<UGameplayCueView*>	Definitions;
+	
+	TMap<FGameplayTag, TArray< TSharedPtr<FGameplayCueViewEffects> > >	SpawnedViewEffects;
+
+	UPROPERTY()
+	TArray<UObject*> SpawnedObjects;
+
+	virtual void GameplayCueActivated(const FGameplayTagContainer & GameplayCueTags, float NormalizedMagnitude, const FGameplayEffectInstigatorContext InstigatorContext);
+	
+	virtual void GameplayCueExecuted(const FGameplayTagContainer & GameplayCueTags, float NormalizedMagnitude, const FGameplayEffectInstigatorContext InstigatorContext);
+	
+	virtual void GameplayCueAdded(const FGameplayTagContainer & GameplayCueTags, float NormalizedMagnitude, const FGameplayEffectInstigatorContext InstigatorContext);
+	
+	virtual void GameplayCueRemoved(const FGameplayTagContainer & GameplayCueTags, float NormalizedMagnitude, const FGameplayEffectInstigatorContext InstigatorContext);
+
+private:
+
+	FGameplayCueViewInfo * GetBestMatchingView(EGameplayCueEvent::Type Type, const FGameplayTag BaseTag);
+
+	void ClearEffects(TArray< TSharedPtr<FGameplayCueViewEffects > > &Effects);
+
+
+};
\ No newline at end of file
diff --git a/Engine/Source/Runtime/GameplayAbilities/Public/GameplayEffect.h b/Engine/Source/Runtime/GameplayAbilities/Public/GameplayEffect.h
new file mode 100644
--- /dev/null
+++ b/Engine/Source/Runtime/GameplayAbilities/Public/GameplayEffect.h
@@ -0,0 +1,1426 @@
+// Copyright 1998-2013 Epic Games, Inc. All Rights Reserved.
+
+#pragma once
+
+#include "GameplayTagContainer.h"
+#include "GameplayTagAssetInterface.h"
+#include "AttributeSet.h"
+#include "TimerManager.h"
+#include "GameplayEffect.generated.h"
+
+DECLARE_DELEGATE_OneParam(FOnGameplayAttributeEffectExecuted, struct FGameplayModifierEvaluatedData &);
+
+DECLARE_MULTICAST_DELEGATE_ThreeParams(FOnAttributeGameplayEffectSpecExected, const FGameplayAttribute &, const struct FGameplayEffectSpec &, struct FGameplayModifierEvaluatedData &);
+
+struct FActiveGameplayEffect;
+
+#define SKILL_SYSTEM_AGGREGATOR_DEBUG 1
+
+#if SKILL_SYSTEM_AGGREGATOR_DEBUG
+	#define SKILL_AGG_DEBUG( Format, ... ) *FString::Printf(Format, ##__VA_ARGS__)
+#else
+	#define SKILL_AGG_DEBUG( Format, ... ) NULL
+#endif
+
+class UGameplayEffect;
+class UAbilitySystemComponent;
+
+UENUM(BlueprintType)
+namespace EGameplayModOp
+{
+	enum Type
+	{		
+		// Numeric
+		Additive = 0		UMETA(DisplayName="Add"),
+		Multiplicitive		UMETA(DisplayName="Multiply"),
+		Division			UMETA(DisplayName="Divide"),
+
+		// Other
+		Override 			UMETA(DisplayName="Override"),	// This should always be the first non numeric ModOp
+		Callback			UMETA(DisplayName="Custom"),
+
+		// This must always be at the end
+		Max					UMETA(DisplayName="Invalid")
+	};
+}
+
+/**
+ * Tells us what thing a GameplayModifier modifies.
+ */
+UENUM(BlueprintType)
+namespace EGameplayMod
+{
+	enum Type
+	{
+		Attribute = 0,		// Modifies this Attributes
+		OutgoingGE,			// Modifies Outgoing Gameplay Effects (that modify this Attribute)
+		IncomingGE,			// Modifies Incoming Gameplay Effects (that modify this Attribute)
+		ActiveGE,			// Modifies currently active Gameplay Effects
+
+		// This must always be at the end
+		Max					UMETA(DisplayName="Invalid")
+	};
+}
+
+/**
+ * Tells us what thing a GameplayEffect provides immunity to. This must mirror the values in EGameplayMod
+ */
+UENUM(BlueprintType)
+namespace EGameplayImmunity
+{
+	enum Type
+	{
+		None = 0,			// Does not provide immunity
+		OutgoingGE,			// Provides immunity to outgoing GEs
+		IncomingGE,			// Provides immunity to incoming GEs
+		ActiveGE,			// Provides immunity from currently active GEs
+
+		// This must always be at the end
+		Max					UMETA(DisplayName="Invalid")
+	};
+}
+
+checkAtCompileTime((int32)EGameplayMod::OutgoingGE == (int32)EGameplayImmunity::OutgoingGE, "EGameplayMod::OutgoingGE and EGameplayImmunity::OutgoingGE must match. Did you forget to modify one of them?");
+checkAtCompileTime((int32)EGameplayMod::IncomingGE == (int32)EGameplayImmunity::IncomingGE, "EGameplayMod::IncomingGE and EGameplayImmunity::IncomingGE must match. Did you forget to modify one of them?");
+checkAtCompileTime((int32)EGameplayMod::ActiveGE == (int32)EGameplayImmunity::ActiveGE, "EGameplayMod::ActiveGE and EGameplayImmunity::ActiveGE must match. Did you forget to modify one of them?");
+checkAtCompileTime((int32)EGameplayMod::Max == (int32)EGameplayImmunity::Max, "EGameplayMod and EGameplayImmunity must cover the same cases. Did you forget to modify one of them?");
+
+/**
+ * Tells us what a GameplayEffect modifies when being applied to another GameplayEffect
+ */
+UENUM(BlueprintType)
+namespace EGameplayModEffect
+{
+	enum Type
+	{
+		Magnitude			= 0x01,		// Modifies magnitude of a GameplayEffect (Always default for Attribute mod)
+		Duration			= 0x02,		// Modifies duration of a GameplayEffect
+		ChanceApplyTarget	= 0x04,		// Modifies chance to apply GameplayEffect to target
+		ChanceExecuteEffect	= 0x08,		// Modifies chance to execute GameplayEffect on GameplayEffect
+		LinkedGameplayEffect= 0x10,		// Adds a linked GameplayEffect to a GameplayEffect
+
+		// This must always be at the end
+		All					= 0xFF		UMETA(DisplayName="Invalid")
+	};
+}
+
+/**
+ * Tells us how to handle copying gameplay effect when it is applied.
+ *	Default means to use context - e.g, OutgoingGE is are always snapshots, IncomingGE is always Link
+ *	AlwaysSnapshot vs AlwaysLink let mods themselves override
+ */
+UENUM(BlueprintType)
+namespace EGameplayEffectCopyPolicy
+{
+	enum Type
+	{
+		Default = 0			UMETA(DisplayName="Default"),
+		AlwaysSnapshot		UMETA(DisplayName="AlwaysSnapshot"),
+		AlwaysLink			UMETA(DisplayName="AlwaysLink")
+	};
+}
+
+UENUM(BlueprintType)
+namespace EGameplayEffectStackingPolicy
+{
+	enum Type
+	{
+		Unlimited = 0		UMETA(DisplayName = "NoRule"),
+		Highest				UMETA(DisplayName = "Strongest"),
+		Lowest				UMETA(DisplayName = "Weakest"),
+		Replaces			UMETA(DisplayName = "MostRecent"),
+		Callback			UMETA(DisplayName = "Custom"),
+
+		// This must always be at the end
+		Max					UMETA(DisplayName = "Invalid")
+	};
+}
+
+FString EGameplayModOpToString( int32 Type );
+
+FString EGameplayModToString( int32 Type );
+
+FString EGameplayModEffectToString( int32 Type );
+
+FString EGameplayEffectCopyPolicyToString(int32 Type);
+
+FString EGameplayEffectStackingPolicyToString(int32 Type);
+
+USTRUCT()
+struct FGameplayModifierCallbacks
+{
+	GENERATED_USTRUCT_BODY()
+
+	UPROPERTY(EditDefaultsOnly, Category = GameplayModifier)
+	TArray<TSubclassOf<class UGameplayEffectExtension> >	ExtensionClasses;
+};
+
+USTRUCT()
+struct FGameplayEffectStackingCallbacks
+{
+	GENERATED_USTRUCT_BODY()
+
+	UPROPERTY(EditDefaultsOnly, Category = GEStack)
+	TArray<TSubclassOf<class UGameplayEffectStackingExtension> >	ExtensionClasses;
+};
+
+/**
+* Defines how A GameplayEffect levels
+*	Normally, GameplayEffect levels are specified when they are created.
+*	They can also be tied to their instigators attribute.
+*		For example, a Damage apply GameplayEffect that 'levels' based on the PhysicalDamage attribute
+*/
+USTRUCT()
+struct GAMEPLAYABILITIES_API FGameplayEffectLevelDef
+{
+	GENERATED_USTRUCT_BODY()
+
+	/** When true, whatever creates of owns this will pass in a level. E.g, level is not intrinsic to this definition. */
+	UPROPERTY(EditDefaultsOnly, Category = GameplayEffectLevel)
+	bool InheritLevelFromOwner;
+
+	/** If set, the gameplay effects level will be tied to this attribute on the instigator */
+	UPROPERTY(EditDefaultsOnly, Category = GameplayEffectLevel, meta = (FilterMetaTag="HideFromLevelInfos"))
+	FGameplayAttribute	Attribute;
+
+	/** If true, take snapshot of attribute level when the gameplay effect is initialized. Otherwise, the level of the gameplay effect will update as the attribute it is tied to updates */
+	UPROPERTY(EditDefaultsOnly, Category = GameplayEffectLevel)
+	bool TakeSnapshotOnInit;
+};
+
+/**
+ * FGameplayModifierInfo
+ *	Tells us "Who/What we" modify
+ *	Does not tell us how exactly
+ *
+ */
+USTRUCT()
+struct GAMEPLAYABILITIES_API FGameplayModifierInfo
+{
+	GENERATED_USTRUCT_BODY()
+
+	FGameplayModifierInfo()
+		: ModifierType( EGameplayMod::Attribute )
+		, ModifierOp( EGameplayModOp::Additive )
+		, EffectType( EGameplayModEffect::Magnitude )
+		, TargetEffect( NULL )
+	{
+
+	}
+
+	UPROPERTY(EditDefaultsOnly, Category = GameplayModifier)
+	FScalableFloat Magnitude; // Not modified from defaults
+
+	/** What this modifies - Attribute, OutgoingGEs, IncomingGEs, ACtiveGEs. */
+	UPROPERTY(EditDefaultsOnly, Category=GameplayModifier)
+	TEnumAsByte<EGameplayMod::Type> ModifierType;
+
+	/** The Attribute we modify or the GE we modify modifies. */
+	UPROPERTY(EditDefaultsOnly, Category=GameplayModifier)
+	FGameplayAttribute Attribute;
+
+	/** The numeric operation of this modifier: Override, Add, Multiply  */
+	UPROPERTY(EditDefaultsOnly, Category=GameplayModifier)
+	TEnumAsByte<EGameplayModOp::Type> ModifierOp;
+
+	/** If we modify an effect, this is what we modify about it (Duration, Magnitude, etc) */
+	UPROPERTY(EditDefaultsOnly, Category=GameplayModifier)
+	TEnumAsByte<EGameplayModEffect::Type> EffectType;
+
+	/** If we are linking a gameplay effect to another effect, this is the effect to link */
+	UPROPERTY(EditDefaultsOnly, Category=GameplayModifier)
+	UGameplayEffect* TargetEffect;
+
+	// The thing I modify requires these tags
+	UPROPERTY(EditDefaultsOnly, Category=GameplayModifier)
+	FGameplayTagContainer RequiredTags;
+
+	// The thing I modify must not have any of these tags
+	UPROPERTY(EditDefaultsOnly, Category = GameplayModifier)
+	FGameplayTagContainer IgnoreTags;
+	
+	/** This modifiers tags. These tags are passed to any other modifiers that this modifies. */
+	UPROPERTY(EditDefaultsOnly, Category=GameplayModifier)
+	FGameplayTagContainer OwnedTags;
+
+	UPROPERTY(EditDefaultsOnly, BlueprintReadWrite, Category = GameplayEffect)
+	FGameplayEffectLevelDef	LevelInfo;
+
+	FString ToSimpleString() const
+	{
+		return FString::Printf(TEXT("%s %s %s BaseVaue: %s"), *EGameplayModToString(ModifierType), *EGameplayModOpToString(ModifierOp), *EGameplayModEffectToString(EffectType), *Magnitude.ToSimpleString());
+	}
+
+	UPROPERTY(EditDefaultsOnly, Category = GameplayModifier)
+	FGameplayModifierCallbacks	Callbacks;
+};
+
+/**
+ * FGameplayEffectCue
+ *	This is a cosmetic cue that can be tied to a UGameplayEffect. 
+ *  This is essentially a GameplayTag + a Min/Max level range that is used to map the level of a GameplayEffect to a normalized value used by the GameplayCue system.
+ */
+USTRUCT()
+struct FGameplayEffectCue
+{
+	GENERATED_USTRUCT_BODY()
+
+	FGameplayEffectCue()
+		: MinLevel(0.f)
+		, MaxLevel(0.f)
+	{
+	}
+
+	FGameplayEffectCue(const FGameplayTag& InTag, float InMinLevel, float InMaxLevel)
+		: MinLevel(InMinLevel)
+		, MaxLevel(InMaxLevel)
+	{
+		GameplayCueTags.AddTag(InTag);
+	}
+
+	UPROPERTY(EditDefaultsOnly, Category = GameplayCue)
+	float	MinLevel;
+
+	UPROPERTY(EditDefaultsOnly, Category = GameplayCue)
+	float	MaxLevel;
+
+	UPROPERTY(EditDefaultsOnly, Category = GameplayCue)
+	FGameplayTagContainer GameplayCueTags;
+
+	float NormalizeLevel(float InLevel)
+	{
+		float Range = MaxLevel - MinLevel;
+		if (Range <= KINDA_SMALL_NUMBER)
+		{
+			return 1.f;
+		}
+
+		return FMath::Clamp((InLevel - MinLevel) / Range, 0.f, 1.0f);
+	}
+};
+
+/**
+ * UGameplayEffect
+ *	The GameplayEffect definition. This is the data asset defined in the editor that drives everything.
+ */
+UCLASS(BlueprintType)
+class GAMEPLAYABILITIES_API UGameplayEffect : public UDataAsset, public IGameplayTagAssetInterface
+{
+	GENERATED_UCLASS_BODY()
+
+public:
+	/** Infinite duration */
+	static const float INFINITE_DURATION;
+
+	/** No duration; Time specifying instant application of an effect */
+	static const float INSTANT_APPLICATION;
+
+	/** Constant specifying that the combat effect has no period and doesn't check for over time application */
+	static const float NO_PERIOD;
+	
+	UPROPERTY(EditDefaultsOnly, Category=GameplayEffect)
+	FScalableFloat	Duration;
+
+	UPROPERTY(EditDefaultsOnly, Category=GameplayEffect)
+	FScalableFloat	Period;
+	
+	UPROPERTY(EditDefaultsOnly, BlueprintReadOnly, Category=GameplayEffect)
+	TArray<FGameplayModifierInfo> Modifiers;
+
+	UPROPERTY(EditDefaultsOnly, BlueprintReadWrite, Category=GameplayEffect)
+	FGameplayEffectLevelDef	LevelInfo;
+		
+	// "I can only be applied to targets that have these tags"
+	// "I can only exist on CE buckets on targets that have these tags":
+	
+	/** Container of gameplay tags that have to be present on the target for the effect to be applied */
+	UPROPERTY(EditDefaultsOnly, BlueprintReadOnly, Category = Application)
+	FGameplayTagContainer ApplicationRequiredTargetTags;
+
+	// "I can only be applied if my instigator has these tags"
+
+	/** Container of gameplay tags that have to be present on the instigator for the effect to be applied */
+	UPROPERTY(EditDefaultsOnly, BlueprintReadOnly, Category = Application)
+	FGameplayTagContainer ApplicationRequiredInstigatorTags;
+
+	UPROPERTY(EditDefaultsOnly, Category=Application, meta=(GameplayAttribute="True"))
+	FScalableFloat	ChanceToApplyToTarget;
+
+	UPROPERTY(EditDefaultsOnly, Category = Application, meta = (GameplayAttribute = "True"))
+	FScalableFloat	ChanceToExecuteOnGameplayEffect;
+
+	// other gameplay effects  that will be applied to the target of this effect
+	UPROPERTY(EditDefaultsOnly, BlueprintReadOnly, Category = GameplayEffect)
+	TArray<UGameplayEffect*> TargetEffects;
+
+	// removes or blocks gameplay effects that it applies to if the tags match
+	UPROPERTY(EditDefaultsOnly, BlueprintReadOnly, Category = GameplayEffect)
+	TEnumAsByte<EGameplayImmunity::Type> AppliesImmunityTo;
+
+	// Modify duration of CEs
+
+	// Modify MaxStacks of CEs (maybe... probably not)
+
+	// ------------------------------------------------
+	// Functions
+	// ------------------------------------------------
+
+	/**
+	 * Determines if the set of supplied gameplay tags are enough to satisfy the application tag requirements of the effect
+	 * 
+	 * @param InstigatorTags	Owned gameplay tags of the instigator applying the effect
+	 * @param TargetTags		Owned gameplay tags of the target about to be affected by the effect
+	 */
+	bool AreApplicationTagRequirementsSatisfied(const FGameplayTagContainer& InstigatorTags, const FGameplayTagContainer& TargetTags) const;
+
+	// ------------------------------------------------
+	// New Tagging functionality
+	// ------------------------------------------------
+
+	// "These are my tags"
+	UPROPERTY(EditDefaultsOnly, BlueprintReadOnly, Category=Tags)
+	FGameplayTagContainer GameplayEffectTags;
+
+	// "In order to affect another GE, they must have ALL of these tags"
+	UPROPERTY(EditDefaultsOnly, BlueprintReadOnly, Category=Tags)
+	FGameplayTagContainer GameplayEffectRequiredTags;
+	
+	// "In order to affect another GE, they must NOT have ANY of these tags"
+	UPROPERTY(EditDefaultsOnly, BlueprintReadOnly, Category=Tags)
+	FGameplayTagContainer GameplayEffectIgnoreTags;
+
+	bool AreGameplayEffectTagRequirementsSatisfied(const FGameplayTagContainer& Tags) const
+	{
+		bool bHasRequired = Tags.MatchesAll(GameplayEffectRequiredTags, true);
+		bool bHasIgnored = Tags.MatchesAny(GameplayEffectIgnoreTags, false);
+
+		return bHasRequired && !bHasIgnored;
+	}
+
+	/** Can this GameplayEffect modify the input parameter, based on tags  */
+	bool AreGameplayEffectTagRequirementsSatisfied(const UGameplayEffect *GameplayEffectToBeModified) const
+	{
+		return AreGameplayEffectTagRequirementsSatisfied(GameplayEffectToBeModified->GameplayEffectTags);
+	}
+
+	// ------------------------------------------------
+	// Gameplay tag interface
+	// ------------------------------------------------
+
+	/** Overridden to return requirements tags */
+	virtual void GetOwnedGameplayTags(FGameplayTagContainer& TagContainer) const OVERRIDE;
+
+	/** Get the "clear tags" for the effect */
+	virtual void GetClearGameplayTags(FGameplayTagContainer& TagContainer) const;
+
+	void ValidateGameplayEffect();
+
+	/** Container of "owned" gameplay tags that are applied to actor with this effect applied to them */
+	UPROPERTY(EditDefaultsOnly, BlueprintReadOnly, Category=Tags)
+	FGameplayTagContainer OwnedTagsContainer;
+
+	/** Container of gameplay tags to be cleared upon effect application; Any active effects with these tags that can be cleared, will be */
+	UPROPERTY(EditDefaultsOnly, BlueprintReadOnly, Category=Tags)
+	FGameplayTagContainer ClearTagsContainer;
+	
+	// Used to quickly tell if a GameplayEffect modifies another GameplayEffect (or a set of attributes)
+	bool ModifiesAnyProperties(EGameplayMod::Type ModType, const TSet<UProperty> & Properties)
+	{
+		return false;
+	}
+
+	virtual void PostLoad() OVERRIDE;
+
+	// -----------------------------------------------
+	
+	UPROPERTY(EditDefaultsOnly, BlueprintReadOnly, Category = Advanced)
+	TEnumAsByte<EGameplayEffectCopyPolicy::Type>	CopyPolicy;
+
+	// ----------------------------------------------
+
+	UPROPERTY(EditDefaultsOnly, BlueprintReadOnly, Category = Display)
+	TArray<FGameplayEffectCue>	GameplayCues;
+
+	/** Description of this combat effect. */
+	UPROPERTY(EditDefaultsOnly, BlueprintReadWrite, Category = Display)
+	FText Description;
+
+	// ----------------------------------------------
+
+	UPROPERTY(EditDefaultsOnly, BlueprintReadOnly, Category = Stacking)
+	TEnumAsByte<EGameplayEffectStackingPolicy::Type>	StackingPolicy;
+
+	UPROPERTY(EditDefaultsOnly, BlueprintReadOnly, Category = Stacking)
+	FName StackedAttribName;
+
+	UPROPERTY(EditDefaultsOnly, Category = Stacking)
+	TSubclassOf<class UGameplayEffectStackingExtension> StackingExtension;
+
+protected:
+	void ValidateStacking();
+};
+
+/**
+* This handle is required for things outside of FActiveGameplayEffectsContainer to refer to a specific active GameplayEffect
+*	For example if a skill needs to create an active effect and then destroy that specific effect that it created, it has to do so
+*	through a handle. a pointer or index into the active list is not sufficient.
+*/
+USTRUCT(BlueprintType)
+struct FActiveGameplayEffectHandle
+{
+	GENERATED_USTRUCT_BODY()
+
+	FActiveGameplayEffectHandle()
+		: Handle(INDEX_NONE)
+	{
+
+	}
+
+	FActiveGameplayEffectHandle(int32 InHandle)
+		: Handle(InHandle)
+	{
+
+	}
+
+	bool IsValid() const
+	{
+		return Handle != INDEX_NONE;
+	}
+
+	FActiveGameplayEffectHandle GetNextHandle()
+	{
+		return FActiveGameplayEffectHandle(Handle + 1);
+	}
+
+	bool operator==(const FActiveGameplayEffectHandle& Other) const
+	{
+		return Handle == Other.Handle;
+	}
+
+	bool operator!=(const FActiveGameplayEffectHandle& Other) const
+	{
+		return Handle != Other.Handle;
+	}
+
+	FString ToString() const
+	{
+		return FString::Printf(TEXT("%d"), Handle);
+	}
+
+private:
+
+	UPROPERTY()
+	int32 Handle;
+};
+
+/**
+ * FGameplayEffectLevelSpec
+ *	Level Specification. This can be a static, constant level specified on creation or it can be dynamically tied to a source's attribute value.
+ *	For example, a GameplayEffect could be made whose level is tied to its instigators PhysicalDamage or Intelligence attribute.
+ */
+
+struct FGameplayEffectLevelSpec
+{
+	static const float INVALID_LEVEL;
+
+	FGameplayEffectLevelSpec()
+		:ConstantLevel(INVALID_LEVEL)
+		,CachedLevel(INVALID_LEVEL)
+	{
+	}
+
+	FGameplayEffectLevelSpec(float InLevel, const FGameplayEffectLevelDef &Def, class AActor *InSource)
+		: ConstantLevel(InLevel)
+		, CachedLevel(InLevel)
+		, Source(InSource)
+	{
+		if (Def.Attribute.GetUProperty() != NULL)
+		{
+			Attribute = Def.Attribute;
+		}
+
+		if (Def.TakeSnapshotOnInit)
+		{
+			SnapshotLevel();
+		}
+	}
+
+	void ApplyNewDef(const FGameplayEffectLevelDef &Def, TSharedPtr<FGameplayEffectLevelSpec> &OutSharedPtr) const
+	{
+		if (Def.InheritLevelFromOwner)
+		{
+			return;
+		}
+
+		check(OutSharedPtr.IsValid());
+		if (Def.Attribute != OutSharedPtr->Attribute)
+		{
+			// In def levels off something different
+			// make a new level spec
+			OutSharedPtr = TSharedPtr<FGameplayEffectLevelSpec>(new FGameplayEffectLevelSpec(INVALID_LEVEL, Def, Source.Get()));
+		}
+	}
+
+	/** Dynamic simply means the level may change. It is not constant. */
+	bool IsDynamic() const
+	{
+		return ConstantLevel == INVALID_LEVEL && Attribute.GetUProperty() != NULL;
+	}
+
+	/** Valid means we have some meaningful data. If we have an INVALID_LEVEL constant value and are not tied to a dynamic property, then we are invalid. */
+	bool IsValid() const
+	{
+		return ConstantLevel != INVALID_LEVEL || Attribute.GetUProperty() != NULL;
+	}
+
+	float GetLevel() const;
+
+	void SnapshotLevel()
+	{
+		// This should snapshot the current level (if dynamic/delegate) and save off its value so that it doesn't change
+		ConstantLevel = GetLevel();
+		Source = NULL;
+	}
+
+	void RegisterLevelDependancy(TWeakPtr<struct FAggregator> OwningAggregator);
+
+	void PrintAll() const;
+
+	mutable float ConstantLevel;	// Final/constant level. Once this is set we are locked at the given level.
+	mutable float CachedLevel;		// Last read value. Needed in case we lose our source, we use the last known level.
+
+	TWeakObjectPtr<class AActor>	Source;
+
+	FGameplayAttribute	Attribute;
+};
+
+/**
+ * FGameplayEffectInstigatorContext
+ *	Data struct for an instigator. This is still being fleshed out. We will want to track actors but also be able to provide some level of tracking for actors that are destroyed.
+ *	We may need to store some positional information as well.
+ *
+ */
+USTRUCT()
+struct FGameplayEffectInstigatorContext
+{
+	GENERATED_USTRUCT_BODY()
+
+	FGameplayEffectInstigatorContext()
+		: Instigator(NULL)
+		, InstigatorAbilitySystemComponent(NULL)
+	{
+	}
+
+	void GetOwnedGameplayTags(OUT FGameplayTagContainer &TagContainer)
+	{
+		IGameplayTagAssetInterface* TagInterface = InterfaceCast<IGameplayTagAssetInterface>(Instigator);
+		if (TagInterface)
+		{
+			TagInterface->GetOwnedGameplayTags(TagContainer);
+		}
+	}
+
+	void AddInstigator(class AActor *InInstigator);
+
+	void AddHitResult(const FHitResult InHitResult);
+
+	FString ToString() const
+	{
+		return Instigator ? Instigator->GetName() : FString(TEXT("NONE"));
+	}
+
+	/** Should always return the original instigator that started the whole chain */
+	AActor * GetOriginalInstigator()
+	{
+		return Instigator;
+	}
+
+	UAbilitySystemComponent * GetOriginalInstigatorAbilitySystemComponent() const
+	{
+		return InstigatorAbilitySystemComponent;
+	}
+	
+	/** Instigator controller */
+	UPROPERTY()
+	class AActor* Instigator;
+
+	UPROPERTY(NotReplicated)
+	UAbilitySystemComponent *InstigatorAbilitySystemComponent;
+
+	/** Trace information - may be NULL in many cases */
+	TSharedPtr<FHitResult>	HitResult;
+
+	bool NetSerialize(FArchive& Ar, class UPackageMap* Map, bool& bOutSuccess);
+};
+
+template<>
+struct TStructOpsTypeTraits< FGameplayEffectInstigatorContext > : public TStructOpsTypeTraitsBase
+{
+	enum
+	{
+		WithNetSerializer = true,
+		WithCopy = true		// Necessary so that TSharedPtr<FGameplayAbilityTargetData> Data is copied around
+	};
+};
+
+
+/**
+ * FAggregatorRef
+ *	A reference to an FAggregator. The reference may be weak or hard, and this can be changed over the lifetime of the FAggregatorRef.
+ *	
+ *	There are cases where we want weak references in an aggregator chain. 
+ *		For example a RunSpeed buff, which when it is destroyed we want the RunSpeed attribute aggregator to recalculate the RunSpeed value.
+ *
+ *	There are cases where we want to make a copy of what we are weak referencing and make the reference a hard ref to that copy
+ *		For example, a DOT which is buffed is attached to a target. We want to make a copy of the DOT and its buff then give it to the target as a hard ref so that if 
+ *		the buff expires on the source, the applied DOT is still buffed.
+ *
+ */
+USTRUCT()
+struct FAggregatorRef
+{
+	GENERATED_USTRUCT_BODY()
+
+	friend struct FAggregatorRef;
+
+	FAggregatorRef()
+	{
+	}
+
+	FAggregatorRef(struct FAggregator *Src)
+		: SharedPtr(Src)
+	{		 
+		WeakPtr = SharedPtr;
+	}
+
+	FAggregatorRef(struct FAggregatorRef *Src)
+	{
+		SetSoftRef(Src);
+	}
+
+	void MakeHardRef()
+	{
+		check(WeakPtr.IsValid());
+		SharedPtr = WeakPtr.Pin();
+	}
+	void MakeSoftRef()
+	{
+		check(WeakPtr.IsValid());
+		SharedPtr.Reset();
+	}
+
+	void SetSoftRef(FAggregatorRef *Src)
+	{
+		check(!SharedPtr.IsValid());
+		WeakPtr = Src->SharedPtr;
+	}
+	
+	FAggregator * Get()
+	{
+		return WeakPtr.IsValid() ? WeakPtr.Pin().Get() : NULL;
+	}
+
+	const FAggregator * Get() const
+	{
+		return WeakPtr.IsValid() ? WeakPtr.Pin().Get() : NULL;
+	}
+
+	bool IsValid() const
+	{
+		return WeakPtr.IsValid();
+	}
+
+	/** Become a hard reference to a new copy of what we are reference  */
+	void MakeUnique();
+
+	/** Become a hard reference to a new copy of what we are reference AND make new copies/hard refs of the complete modifier chain in our FAggregator */
+	void MakeUniqueDeep();
+
+	bool NetSerialize(FArchive& Ar, class UPackageMap* Map, bool& bOutSuccess);	
+
+	FString ToString() const;
+	void PrintAll() const;
+
+private:
+
+	TSharedPtr<struct FAggregator>	SharedPtr;
+	TWeakPtr<struct FAggregator> WeakPtr;
+};
+
+template<>
+struct TStructOpsTypeTraits< FAggregatorRef > : public TStructOpsTypeTraitsBase
+{
+	enum
+	{
+		WithNetSerializer = true
+	};
+};
+
+/**
+ * GameplayModifierEvaluatedData
+ *	This is the data that FAggregator aggregates and turns into FGameplayModifierEvaluatedData.
+ *  It is distinct from FGameplayModifierEvaluatedData in that FGameplayModifierData ia level has not been applied to this data.
+ *  FGameplayModifierData::Magnitude is an FScalableFloat which describes a numeric value for a given level.
+ * 
+ */
+struct FGameplayModifierData
+{
+	FGameplayModifierData()
+	{
+		
+	}
+
+	FGameplayModifierData(const FGameplayModifierInfo &Info, const FGlobalCurveDataOverride *CurveData)
+	{
+		Magnitude = Info.Magnitude.MakeFinalizedCopy(CurveData);
+		Tags = Info.OwnedTags;
+
+		// Fixme: this is static data, should be a reference
+		RequireTags = Info.RequiredTags;
+		IgnoreTags = Info.IgnoreTags;
+
+		if (Info.Callbacks.ExtensionClasses.Num() > 0)
+		{
+			Callbacks = &Info.Callbacks;
+		}
+	}
+
+	FGameplayModifierData(FScalableFloat InMagnitude)
+	{
+		// Magnitude may scale based on our level
+		Magnitude = InMagnitude;
+		Callbacks = NULL;
+	}
+
+	FGameplayModifierData(float InMagnitude, const FGameplayModifierCallbacks * InCallbacks)
+	{
+		// Magnitude will be fixed at this value
+		Magnitude.SetValue(InMagnitude);
+		Callbacks = InCallbacks;
+	}
+
+	// That magnitude that we modify by
+	FScalableFloat Magnitude;
+
+	// The tags I have
+	FGameplayTagContainer Tags;
+	
+	FGameplayTagContainer RequireTags;
+	FGameplayTagContainer IgnoreTags;
+
+	// Callback information for custom logic pre/post evaluation
+	const FGameplayModifierCallbacks * Callbacks;
+
+	void PrintAll() const;
+};
+
+/**
+ * GameplayModifierEvaluatedData
+ *	This is the output from an FAggregator: a numeric value and a set of GameplayTags.
+ */
+struct FGameplayModifierEvaluatedData
+{
+	FGameplayModifierEvaluatedData()
+		: Magnitude(0.f)
+		, Callbacks(NULL)
+		, IsValid(false)
+	{
+	}
+
+	FGameplayModifierEvaluatedData(float InMagnitude, const FGameplayModifierCallbacks * InCallbacks = NULL, FActiveGameplayEffectHandle InHandle = FActiveGameplayEffectHandle(), const FGameplayTagContainer *InTags = NULL)
+		: Magnitude(InMagnitude)
+		, Callbacks(InCallbacks)
+		, Handle(InHandle)
+		, IsValid(true)
+	{
+		if (InTags)
+		{
+			Tags = *InTags;
+		}
+	}
+
+	float	Magnitude;
+	FGameplayTagContainer Tags;
+	const FGameplayModifierCallbacks * Callbacks;
+	FActiveGameplayEffectHandle	Handle;	// Handle of the active gameplay effect that originated us. Will be invalid in many cases
+	bool IsValid;
+
+	// Helper function for building up final values during an aggregation
+	void Aggregate(OUT FGameplayTagContainer &OutTags, OUT float &OutMagnitude, const float Bias=0.f) const;
+
+	void InvokePreExecute(struct FGameplayEffectModCallbackData &Data) const;
+	void InvokePostExecute(const struct FGameplayEffectModCallbackData &Data) const;
+
+	void PrintAll() const;
+};
+
+/**
+ * FAggregator - a data structure for aggregating stuff in GameplayEffects.
+ *	Aggregates a numeric value (float) and a set of gameplay tags. This could be further extended.
+ *
+ *	Aggregation is done with BaseData + Mods[].
+ *	-BaseData is simply the base data. We are initiliazed with base data and base data can be directly modified via ::ExecuteMod.
+ *	-Mods[] are lists of other FAggregators. That is, we have a list for each EGameplayModOp: Add, multiply, override.
+ *	-These lists contain FAggregatorRefs, which may be soft or hard refs to other FAggregators.
+ *	-::Evalate() takes our BaseData, and then crawls through ours Mods[] list and aggregates a final output (FGameplayModifierEvaluatedData)
+ *	-Results of ::Evaluate are cached in CachedData.
+ *	-FAggregator also keeps a list of weak ptrs to other FAggregators that are dependant on us. If we change, we let these aggregators know, so they can invalidate their cached data.
+ *
+ *
+ */
+struct FAggregator : public TSharedFromThis<FAggregator>
+{
+	DECLARE_DELEGATE_OneParam(FOnDirty, FAggregator*);
+
+	FAggregator();
+	FAggregator(const FGameplayModifierData &InBaseData, TSharedPtr<FGameplayEffectLevelSpec> LevelInfo, const TCHAR* InDebugString);
+	FAggregator(const FScalableFloat &InBaseMagnitude, TSharedPtr<FGameplayEffectLevelSpec> LevelInfo, const TCHAR* InDebugString);
+	FAggregator(const FGameplayModifierEvaluatedData &InEvalData, const TCHAR* InDebugString);
+	FAggregator(const FAggregator &In);
+	virtual ~FAggregator();
+
+	FAggregator & MarkDirty();
+	void ClearAllDependancies();
+
+	const FGameplayModifierEvaluatedData& Evaluate() const;
+
+	void PreEvaluate(struct FGameplayEffectModCallbackData &Data) const;
+	void PostEvaluate(const struct FGameplayEffectModCallbackData &Data) const;
+
+	void TakeSnapshotOfLevel();
+	void ApplyMod(EGameplayModOp::Type ModType, FAggregatorRef Ref, bool TakeSnapshot);
+	
+	void ExecuteModAggr(EGameplayModOp::Type ModType, FAggregatorRef Ref);
+	void ExecuteMod(EGameplayModOp::Type ModType, const FGameplayModifierEvaluatedData& EvaluatedData);
+
+	void AddDependantAggregator(TWeakPtr<FAggregator> InDependant);
+
+	void RegisterLevelDependancies();
+
+	TSharedPtr<FGameplayEffectLevelSpec> Level;
+	FActiveGameplayEffectHandle	ActiveHandle;	// Handle to owning active effect. Will be null in many cases.
+
+	FGameplayModifierData		BaseData;
+	TArray<FAggregatorRef>		Mods[EGameplayModOp::Max];
+
+	TArray<TWeakPtr<FAggregator > >	Dependants;
+
+	FOnDirty	OnDirty;
+
+	void PrintAll() const;
+	void RefreshDependencies();
+	void MakeUniqueDeep();
+
+	void SetFromNetSerialize(float NetSerialize);
+
+	// ----------------------------------------------------------------------------
+	// This is data only used in debugging/tracking where aggregator's came from
+	// ----------------------------------------------------------------------------
+#if SKILL_SYSTEM_AGGREGATOR_DEBUG
+	FString DebugString;
+	mutable int32 CopiesMade;
+
+	FString ToSimpleString() const
+	{
+		return FString::Printf(TEXT("0x%X %s. CacheValid: %d Mods: [%d/%d/%d]"), this, *DebugString, CachedData.IsValid, 
+			GetNumValidMods(EGameplayModOp::Override), GetNumValidMods(EGameplayModOp::Additive), GetNumValidMods(EGameplayModOp::Multiplicitive) );
+	}
+
+	struct FAllocationStats
+	{
+		FAllocationStats()
+		{
+			Reset();
+		}
+
+		void Reset()
+		{
+			FMemory::Memzero(this, sizeof(FAllocationStats));
+		}
+
+		int32 DefaultCStor;
+		int32 ModifierCStor;
+		int32 ScalableFloatCstor;
+		int32 FloatCstor;
+		int32 CopyCstor;
+
+		int32 DependantsUpdated;
+	};
+
+	static FAllocationStats AllocationStats;
+#else
+
+	FString ToSimpleString() const
+	{
+		return FString::Printf(TEXT("CacheValid: %d Mods: [%d/%d/%d]"), CachedData.IsValid, 
+			GetNumValidMods(EGameplayModOp::Override), GetNumValidMods(EGameplayModOp::Additive), GetNumValidMods(EGameplayModOp::Multiplicitive) );
+	}
+
+#endif
+	
+
+private:
+
+	int32 GetNumValidMods(EGameplayModOp::Type Type) const
+	{
+		int32 Num=0;
+		for (const FAggregatorRef &Agg : Mods[Type])
+		{
+			if (Agg.IsValid())
+			{
+				Num++;
+			}
+		}
+		return Num;
+	}
+
+	mutable FGameplayModifierEvaluatedData	CachedData;
+};
+
+/** 
+ * Qualification context for applying modifiers.
+ *	For example a Modifier may be setup in data to only apply to OutgoingGE mods.
+ *	The FModifierQualifier is the data structure to hold the 'what type of modifier are we applying' data.
+ *
+ *  This should ideally only hold data that is outside of FGameplayEffectSpec or FGameplayModifierSpec.
+ *	For example, specs know what they can and can't modify. We don't need to duplicate that there. 
+ *  FModifierQualifier is meant to hold the data that comes from the calling context that is not intrinsic 
+ *  to the existing data structures.
+ *
+ *	This class uses an optional initialization idiom such that you can do things like:
+ *		FModifierQualifier().Handle(InHandle).(InType).etc...
+ */
+USTRUCT(BlueprintType)
+struct FModifierQualifier
+{
+	GENERATED_USTRUCT_BODY()
+
+	FModifierQualifier()
+		: MyType(EGameplayMod::Max)
+		, MyPredictionKey(0)
+	{
+	}
+
+	// ----------------------------------
+
+	FModifierQualifier& Type(EGameplayMod::Type InType)
+	{
+		MyType = InType;
+		return *this;
+	}
+
+	EGameplayMod::Type Type() const
+	{
+		return MyType;
+	}
+
+	// ----------------------------------
+	// PredictionKey is used by networking/replication to specify that a GameplayEffect has been predictively created/added.
+
+	int32 PredictionKey() const
+	{
+		return MyPredictionKey;
+	}
+
+	FModifierQualifier& PredictionKey(int32 InPredictionKey)
+	{
+		MyPredictionKey = InPredictionKey;
+		return *this;
+	}
+
+	// ----------------------------------
+	// IgnoreHandle - ignore this handle completely. For example when executing an active gameplay effect, we never
+	// want to apply it to itself (either as IncomingGE or activeGE). It is ignored in all contexts.
+	//
+
+	FModifierQualifier& IgnoreHandle(FActiveGameplayEffectHandle InHandle)
+	{
+		MyIgnoreHandle = InHandle;
+		return *this;
+	}
+
+	FActiveGameplayEffectHandle IgnoreHandle() const
+	{
+		return MyIgnoreHandle;
+	}
+
+	// ----------------------------------
+	// ExclusiveTarget - sometimes we need to only apply a modifier to a specific active gameplay effect.
+	// We may only be able to use the handle if there are multiple instances of the same gameplay effect.
+	// This only applies in the context of applying/executing to a target. 'We only modify this active effect'.
+	// ExclusiveTarget is not checked in the context of applying outgoing/incoming GE modifiers to the spec.
+	//
+
+	FModifierQualifier& ExclusiveTarget(FActiveGameplayEffectHandle InHandle)
+	{
+		MyExclusiveTargetHandle = InHandle;
+		return *this;
+	}
+
+	FActiveGameplayEffectHandle ExclusiveTarget() const
+	{
+		return MyExclusiveTargetHandle;
+	}
+
+	// ----------------------------------
+
+	bool TestTarget(FActiveGameplayEffectHandle InHandle) const
+	{
+		if (MyIgnoreHandle.IsValid() && MyIgnoreHandle == InHandle)
+			return false;
+
+		if (MyExclusiveTargetHandle.IsValid() && MyExclusiveTargetHandle != InHandle)
+			return false;
+
+		return true;
+	}
+
+	FString ToString() const
+	{
+		return EGameplayModToString(MyType);
+	}
+
+private:
+	EGameplayMod::Type	MyType;
+	FActiveGameplayEffectHandle MyIgnoreHandle;		// Do not modify this handle
+	FActiveGameplayEffectHandle MyExclusiveTargetHandle;	// Only modify this handle
+	int32 MyPredictionKey;
+};
+
+/**
+ * Modifier Specification
+ *	-Const data (FGameplayModifierInfo) tells us what we modify, what we can modify
+ *	-Mutable Aggregated data tells us how we modify (magnitude).
+ *  
+ * Modifiers can be modified. A modifier spec holds these modifications along with a reference to the const data about the modifier.
+ * 
+ */
+struct FModifierSpec
+{
+	FModifierSpec(const FGameplayModifierInfo &InInfo, TSharedPtr<FGameplayEffectLevelSpec> InLevel, const FGlobalCurveDataOverride *CurveData, AActor *Owner = NULL, float Level = 0.f);
+
+	// Hard Ref to what we modify, this stuff is const and never changes
+	const FGameplayModifierInfo &Info;
+	
+	FAggregatorRef Aggregator;
+
+	TSharedPtr<FGameplayEffectSpec> TargetEffectSpec;
+
+	bool CanModifyInContext(const FModifierQualifier &QualifierContext) const;
+	// returns true if this GameplayEffect can modify Other, false otherwise
+	bool CanModifyModifier(FModifierSpec &Other, const FModifierQualifier &QualifierContext) const;
+	
+	void ApplyModTo(FModifierSpec &Other, bool TakeSnapshot) const;
+	void ExecuteModOn(FModifierSpec &Other) const;
+
+	FString ToSimpleString() const
+	{
+		return Info.ToSimpleString();
+	}
+
+	// Can this GameplayEffect modify the input parameter, based on tags
+	// Returns true if it can modify the input parameter, false otherwise
+	bool AreTagRequirementsSatisfied(const FModifierSpec &ModifierToBeModified) const;
+
+	void PrintAll() const;
+};
+
+/**
+ * GameplayEffect Specification. Tells us:
+ *	-What UGameplayEffect (const data)
+ *	-What Level
+ *  -Who instigated
+ *  
+ * FGameplayEffectSpec is modifiable. We start with initial conditions and modifications be applied to it. In this sense, it is stateful/mutable but it
+ * is still distinct from an FActiveGameplayEffect which in an applied instance of an FGameplayEffectSpec.
+ */
+USTRUCT()
+struct FGameplayEffectSpec
+{
+	GENERATED_USTRUCT_BODY()
+
+	FGameplayEffectSpec()
+		: ModifierLevel( TSharedPtr< FGameplayEffectLevelSpec >( new FGameplayEffectLevelSpec() ) )
+		, Duration(new FAggregator(FGameplayModifierEvaluatedData(0.f, NULL, FActiveGameplayEffectHandle()), SKILL_AGG_DEBUG(TEXT("Uninitialized Duration"))))
+		, Period(new FAggregator(FGameplayModifierEvaluatedData(0.f, NULL, FActiveGameplayEffectHandle()), SKILL_AGG_DEBUG(TEXT("Uninitialized Period"))))
+		, bTopOfStack(false)
+	{
+		// If we initialize a GameplayEffectSpec with no level object passed in.
+	}
+
+	FGameplayEffectSpec( const UGameplayEffect *InDef, AActor *Instigator, float Level, const FGlobalCurveDataOverride *CurveData );
+	
+	UPROPERTY()
+	const UGameplayEffect * Def;
+
+	// Replicated	
+	TSharedPtr< FGameplayEffectLevelSpec > ModifierLevel;
+	
+	// Replicated
+	UPROPERTY()
+	FGameplayEffectInstigatorContext InstigatorContext; // This tells us how we got here (who / what applied us)
+
+	float GetDuration() const;
+	float GetPeriod() const;
+	float GetChanceToApplyToTarget() const;
+	float GetChanceToExecuteOnGameplayEffect() const;
+	EGameplayEffectStackingPolicy::Type GetStackingType() const;
+	float GetMagnitude(const FGameplayAttribute &Attribute) const;
+
+	// other effects that need to be applied to the target if this effect is successful
+	TArray< TSharedRef< FGameplayEffectSpec > > TargetEffectSpecs;
+
+	UPROPERTY()
+	FAggregatorRef	Duration;
+
+	UPROPERTY()
+	FAggregatorRef	Period;
+
+	UPROPERTY()
+	FAggregatorRef	ChanceToApplyToTarget;
+
+	UPROPERTY()
+	FAggregatorRef	ChanceToExecuteOnGameplayEffect;
+
+	// This should only be true if this is a stacking effect and at the top of its stack
+	UPROPERTY()
+	bool bTopOfStack;
+
+	// The spec needs to own these FModifierSpecs so that other people can keep TSharedPtr to it.
+	// The stuff in this array is OWNED by this spec
+
+	TArray<FModifierSpec>	Modifiers;
+
+	void MakeUnique();
+
+	void InitModifiers(const FGlobalCurveDataOverride *CurveData, AActor *Owner, float Level);
+
+	// returns the number of modifiers applied to InSpec by the current GameplayEffect Spec
+	// returns -1 if the current GameplayEffectSpec prevents InSpec from being applied
+	int32 ApplyModifiersFrom(FGameplayEffectSpec &InSpec, const FModifierQualifier &QualifierContext);
+
+	int32 ExecuteModifiersFrom(const FGameplayEffectSpec &InSpec, const FModifierQualifier &QualifierContext);
+
+	bool ShouldApplyAsSnapshot(const FModifierQualifier &QualifierContext) const;
+
+	FString ToSimpleString() const
+	{
+		return FString::Printf(TEXT("%s"), *Def->GetName());
+	}
+
+	void PrintAll() const;
+
+	// Callbacks
+	FOnAttributeGameplayEffectSpecExected	OnExecute;
+};
+
+/**
+ * Active GameplayEffect instance
+ *	-What GameplayEffect Spec
+ *	-Start time
+ *  -When to execute next
+ *  -Replication callbacks
+ *
+ */
+USTRUCT()
+struct FActiveGameplayEffect : public FFastArraySerializerItem
+{
+	GENERATED_USTRUCT_BODY()
+
+	FActiveGameplayEffect()
+		: StartGameStateTime(0)
+		, StartWorldTime(0.f)
+		, PredictionKey(0)
+	{
+	}
+
+	FActiveGameplayEffect(FActiveGameplayEffectHandle InHandle, const FGameplayEffectSpec &InSpec, float CurrentWorldTime, int32 InStartGameStateTime, int32 InPredictionKey)
+		: Handle(InHandle)
+		, Spec(InSpec)
+		, StartGameStateTime(InStartGameStateTime)
+		, StartWorldTime(CurrentWorldTime)
+		, PredictionKey(InPredictionKey)
+	{
+		for (FModifierSpec &Mod : Spec.Modifiers)
+		{
+			Mod.Aggregator.Get()->ActiveHandle = InHandle;
+		}
+	}
+
+	FActiveGameplayEffectHandle Handle;
+
+	UPROPERTY()
+	FGameplayEffectSpec Spec;
+
+	/** Game time this started */
+	UPROPERTY()
+	int32 StartGameStateTime;
+
+	UPROPERTY(NotReplicated)
+	float StartWorldTime;
+
+	FTimerHandle PeriodHandle;
+
+	FTimerHandle DurationHandle;
+
+	int32 PredictionKey;
+
+	float GetTimeRemaining(float WorldTime)
+	{
+		float Duration = GetDuration();		
+		return (Duration == UGameplayEffect::INFINITE_DURATION ? -1.f : Duration - (WorldTime - StartWorldTime));
+	}
+	
+	float GetDuration() const
+	{
+		return Spec.GetDuration();
+	}
+
+	float GetPeriod() const
+	{
+		return Spec.GetPeriod();
+	}
+
+	bool CanBeStacked(const FActiveGameplayEffect& Other) const;
+
+	void PrintAll() const;
+
+	void PreReplicatedRemove(const struct FActiveGameplayEffectsContainer &InArray);
+	void PostReplicatedAdd(const struct FActiveGameplayEffectsContainer &InArray);
+	void PostReplicatedChange(const struct FActiveGameplayEffectsContainer &InArray) { }
+
+	bool operator==(const FActiveGameplayEffect& Other)
+	{
+		return Handle == Other.Handle;
+	}
+};
+
+/** Generic querying data structure for active GameplayEffects. Lets us ask things like:
+ *		-Give me duration/magnitude of active gameplay effects with these tags
+ *		-
+ */
+USTRUCT()
+struct FActiveGameplayEffectQuery
+{
+	GENERATED_USTRUCT_BODY()
+
+	FActiveGameplayEffectQuery()
+		: TagContainer(NULL)
+	{
+	}
+
+	FActiveGameplayEffectQuery(const FGameplayTagContainer * InTagContainer)
+		: TagContainer(InTagContainer)
+	{
+	}
+
+	const FGameplayTagContainer *	TagContainer;
+};
+
+/**
+ * Active GameplayEffects Container
+ *	-Bucket of ActiveGameplayEffects
+ *	-Needed for FFastArraySerialization
+ *  
+ * This should only be used by UAbilitySystemComponent. All of this could just live in UAbilitySystemComponent except that we need a distinct USTRUCT to implement FFastArraySerializer.
+ *
+ */
+USTRUCT()
+struct FActiveGameplayEffectsContainer : public FFastArraySerializer
+{
+	GENERATED_USTRUCT_BODY();
+
+	FActiveGameplayEffectsContainer() : bNeedToRecalculateStacks(false) {};
+
+	UPROPERTY()
+	TArray< FActiveGameplayEffect >	GameplayEffects;
+	
+	FActiveGameplayEffect & CreateNewActiveGameplayEffect(const FGameplayEffectSpec &Spec, int32 InPredictionKey);
+
+	// returns true if none of the active effects provide immunity to Spec
+	// returns false if one (or more) of the active effects provides immunity to Spec
+	bool ApplyActiveEffectsTo(OUT FGameplayEffectSpec &Spec, const FModifierQualifier &QualifierContext);
+
+	void ApplySpecToActiveEffectsAndAttributes(FGameplayEffectSpec &Spec, const FModifierQualifier &QualifierContext);
+		
+	void ExecuteActiveEffectsFrom(const FGameplayEffectSpec &Spec, const FModifierQualifier &QualifierContext);
+	
+	void ExecutePeriodicGameplayEffect(FActiveGameplayEffectHandle Handle);	// This should not be outward facing to the skill system API, should only be called by the owning AbilitySystemComponent
+
+	void AddDependancyToAttribute(FGameplayAttribute Attribute, const TWeakPtr<FAggregator> InDependant);
+
+	bool RemoveActiveGameplayEffect(FActiveGameplayEffectHandle Handle);
+	
+	float GetGameplayEffectDuration(FActiveGameplayEffectHandle Handle) const;
+
+	float GetGameplayEffectMagnitude(FActiveGameplayEffectHandle Handle, FGameplayAttribute Attribute) const;
+
+	// returns true if the handle points to an effect in this container that is not a stacking effect or an effect in this container that does stack and is applied by the current stacking rules
+	// returns false if the handle points to an effect that is not in this container or is not applied because of the current stacking rules
+	bool IsGameplayEffectActive(FActiveGameplayEffectHandle Handle) const;
+
+	void PrintAllGameplayEffects() const;
+
+	int32 GetNumGameplayEffects() const
+	{
+		return GameplayEffects.Num();
+	}
+
+	float GetGameplayEffectMagnitudeByTag(FActiveGameplayEffectHandle Handle, const FGameplayTag& InTag) const;
+
+	void OnPropertyAggregatorDirty(FAggregator* Aggregator, FGameplayAttribute Attribute);
+
+	UAbilitySystemComponent *Owner;
+
+	void CheckDuration(FActiveGameplayEffectHandle Handle);
+
+	void StacksNeedToRecalculate();
+
+	// recalculates all of the stacks in the current container
+	void RecalculateStacking();
+
+	bool NetDeltaSerialize(FNetDeltaSerializeInfo & DeltaParms)
+	{
+		return FastArrayDeltaSerialize<FActiveGameplayEffect>(GameplayEffects, DeltaParms, *this);
+	}
+
+	void PreDestroy();
+
+	bool bNeedToRecalculateStacks;
+
+	bool HasAnyTags(FGameplayTagContainer &Tags);
+
+	bool HasAllTags(FGameplayTagContainer &Tags);
+
+	bool CanApplyAttributeModifiers(const UGameplayEffect *GameplayEffect, float Level, AActor *Instigator);
+	
+	TArray<float> GetActiveEffectsTimeRemaining(const FActiveGameplayEffectQuery Query) const;
+
+	int32 GetGameStateTime() const;
+	float GetWorldTime() const;
+
+	void OnDurationAggregatorDirty(FAggregator* Aggregator, UAbilitySystemComponent* Owner, FActiveGameplayEffectHandle Handle);
+
+private:
+
+	FTimerHandle StackHandle;
+
+	bool InternalRemoveActiveGameplayEffect(int32 Idx);
+
+	FAggregatorRef & FindOrCreateAttributeAggregator(FGameplayAttribute Attribute);
+
+	FActiveGameplayEffectHandle LastAssignedHandle;
+
+	TMap<FGameplayAttribute, FAggregatorRef> OngoingPropertyEffects;
+
+	bool IsNetAuthority() const;
+};
+
+template<>
+struct TStructOpsTypeTraits< FActiveGameplayEffectsContainer > : public TStructOpsTypeTraitsBase
+{
+	enum
+	{
+		WithNetDeltaSerializer = true,
+	};
+};
diff --git a/Engine/Source/Runtime/GameplayAbilities/Public/GameplayEffectExtension.h b/Engine/Source/Runtime/GameplayAbilities/Public/GameplayEffectExtension.h
new file mode 100644
--- /dev/null
+++ b/Engine/Source/Runtime/GameplayAbilities/Public/GameplayEffectExtension.h
@@ -0,0 +1,45 @@
+// Copyright 1998-2013 Epic Games, Inc. All Rights Reserved.
+
+#pragma once
+
+#include "GameplayTagContainer.h"
+#include "GameplayTagAssetInterface.h"
+#include "AttributeSet.h"
+#include "GameplayEffect.h"
+#include "GameplayEffectExtension.generated.h"
+
+struct FGameplayEffectModCallbackData
+{
+	FGameplayEffectModCallbackData(const FGameplayEffectSpec & InEffectSpec, const FModifierSpec & InModifierSpec, FGameplayModifierEvaluatedData & InEvaluatedData, UAbilitySystemComponent & InTarget)
+		: EffectSpec(InEffectSpec)
+		, ModifierSpec(InModifierSpec)
+		, EvaluatedData(InEvaluatedData)
+		, Target(InTarget)
+	{
+
+	}
+
+	const struct FGameplayEffectSpec &		EffectSpec;		// The spec that the mod came from
+	const struct FModifierSpec &			ModifierSpec;	// The mod we are going to apply
+	struct FGameplayModifierEvaluatedData & EvaluatedData;	// The 'flat'/computed data to be applied to the target
+
+	class UAbilitySystemComponent &Target;		// Target we intend to apply to
+};
+
+UCLASS(BlueprintType)
+class GAMEPLAYABILITIES_API UGameplayEffectExtension: public UObject
+{
+	GENERATED_UCLASS_BODY()
+
+public:
+
+	virtual void PreGameplayEffectExecute(const FGameplayModifierEvaluatedData &SelfData, FGameplayEffectModCallbackData &Data) const
+	{
+
+	}
+
+	virtual void PostGameplayEffectExecute(const FGameplayModifierEvaluatedData &SelfData, const FGameplayEffectModCallbackData &Data) const
+	{
+
+	}
+};
\ No newline at end of file
diff --git a/Engine/Source/Runtime/GameplayAbilities/Public/GameplayEffectExtension_LifestealTest.h b/Engine/Source/Runtime/GameplayAbilities/Public/GameplayEffectExtension_LifestealTest.h
new file mode 100644
--- /dev/null
+++ b/Engine/Source/Runtime/GameplayAbilities/Public/GameplayEffectExtension_LifestealTest.h
@@ -0,0 +1,24 @@
+// Copyright 1998-2013 Epic Games, Inc. All Rights Reserved.
+
+#pragma once
+
+#include "GameplayEffectExtension.h"
+#include "GameplayTagContainer.h"
+#include "GameplayTagAssetInterface.h"
+#include "AttributeSet.h"
+#include "GameplayEffectExtension_LifestealTest.generated.h"
+
+UCLASS(BlueprintType)
+class GAMEPLAYABILITIES_API UGameplayEffectExtension_LifestealTest : public UGameplayEffectExtension
+{
+	GENERATED_UCLASS_BODY()
+
+public:
+
+	void PreGameplayEffectExecute(const FGameplayModifierEvaluatedData &SelfData, FGameplayEffectModCallbackData &Data) const OVERRIDE;
+	void PostGameplayEffectExecute(const FGameplayModifierEvaluatedData &SelfData, const FGameplayEffectModCallbackData &Data) const OVERRIDE;
+
+	/** The GameplayEffect to apply when restoring health to the instigator */
+	UPROPERTY(EditDefaultsOnly, Category = Lifesteal)
+	UGameplayEffect *	HealthRestoreGameplayEffect;
+};
\ No newline at end of file
diff --git a/Engine/Source/Runtime/GameplayAbilities/Public/GameplayEffectExtension_ShieldTest.h b/Engine/Source/Runtime/GameplayAbilities/Public/GameplayEffectExtension_ShieldTest.h
new file mode 100644
--- /dev/null
+++ b/Engine/Source/Runtime/GameplayAbilities/Public/GameplayEffectExtension_ShieldTest.h
@@ -0,0 +1,23 @@
+// Copyright 1998-2013 Epic Games, Inc. All Rights Reserved.
+
+#pragma once
+
+#include "GameplayEffectExtension.h"
+#include "GameplayTagContainer.h"
+#include "GameplayTagAssetInterface.h"
+#include "AttributeSet.h"
+#include "GameplayEffectExtension_ShieldTest.generated.h"
+
+UCLASS(BlueprintType)
+class GAMEPLAYABILITIES_API UGameplayEffectExtension_ShieldTest : public UGameplayEffectExtension
+{
+	GENERATED_UCLASS_BODY()
+
+public:
+
+	void PreGameplayEffectExecute(const FGameplayModifierEvaluatedData &SelfData, FGameplayEffectModCallbackData &Data) const OVERRIDE;
+	void PostGameplayEffectExecute(const FGameplayModifierEvaluatedData &SelfData, const FGameplayEffectModCallbackData &Data) const OVERRIDE;
+
+	UPROPERTY(EditDefaultsOnly, Category = Lifesteal)
+	UGameplayEffect * ShieldRemoveGameplayEffect;
+};
\ No newline at end of file
diff --git a/Engine/Source/Runtime/GameplayAbilities/Public/GameplayEffectStackingExtension.h b/Engine/Source/Runtime/GameplayAbilities/Public/GameplayEffectStackingExtension.h
new file mode 100644
--- /dev/null
+++ b/Engine/Source/Runtime/GameplayAbilities/Public/GameplayEffectStackingExtension.h
@@ -0,0 +1,85 @@
+// Copyright 1998-2014 Epic Games, Inc. All Rights Reserved.
+
+#pragma once
+
+#include "GameplayTagContainer.h"
+#include "GameplayTagAssetInterface.h"
+#include "AttributeSet.h"
+#include "GameplayEffect.h"
+#include "GameplayEffectStackingExtension.generated.h"
+
+struct FGameplayEffectStackingCallbackData
+{
+	FGameplayEffectStackingCallbackData(const FGameplayEffectSpec & InEffectSpec, const FModifierSpec & InModifierSpec, FGameplayModifierEvaluatedData & InEvaluatedData, UAbilitySystemComponent & InTarget)
+		: EffectSpec(InEffectSpec)
+		, ModifierSpec(InModifierSpec)
+		, EvaluatedData(InEvaluatedData)
+		, Target(InTarget)
+	{
+
+	}
+
+	const struct FGameplayEffectSpec &		EffectSpec;		// The spec that the mod came from
+	const struct FModifierSpec &			ModifierSpec;	// The mod we are going to apply
+	struct FGameplayModifierEvaluatedData & EvaluatedData;	// The 'flat'/computed data to be applied to the target
+
+	class UAbilitySystemComponent &Target;		// Target we intend to apply to
+};
+
+
+USTRUCT(BlueprintType)
+struct FGameplayEffectStackingHandle
+{
+	GENERATED_USTRUCT_BODY()
+
+	FGameplayEffectStackingHandle()
+	: Handle(INDEX_NONE)
+	{
+		// do nothing
+	}
+
+	FGameplayEffectStackingHandle(uint32 InHandle)
+		: Handle(InHandle)
+	{
+		// do nothing
+	}
+
+	bool IsValid() const
+	{
+		return Handle != INDEX_NONE;
+	}
+
+	bool operator==(const FGameplayEffectStackingHandle& Other) const
+	{
+		return Handle == Other.Handle;
+	}
+
+	bool operator!=(const FGameplayEffectStackingHandle& Other) const
+	{
+		return Handle != Other.Handle;
+	}
+
+	FString ToString() const
+	{
+		return FString::Printf(TEXT("%d"), Handle);
+	}
+
+private:
+
+	UPROPERTY()
+	uint32 Handle;
+};
+
+UCLASS(BlueprintType)
+class GAMEPLAYABILITIES_API UGameplayEffectStackingExtension: public UObject
+{
+	GENERATED_UCLASS_BODY()
+
+public:
+
+	static FGameplayEffectSpec* StackCustomEffects(TArray<FGameplayEffectSpec*>& CustomGameplayEffects);
+
+	virtual void CalculateStack(TArray<FActiveGameplayEffect*>& CustomGameplayEffects, FActiveGameplayEffectsContainer& Container, FActiveGameplayEffect& CurrentEffect) {}
+
+	FGameplayEffectStackingHandle Handle;
+};
diff --git a/Engine/Source/Runtime/GameplayAbilities/Public/GameplayEffectStackingExtension_CappedNumberTest.h b/Engine/Source/Runtime/GameplayAbilities/Public/GameplayEffectStackingExtension_CappedNumberTest.h
new file mode 100644
--- /dev/null
+++ b/Engine/Source/Runtime/GameplayAbilities/Public/GameplayEffectStackingExtension_CappedNumberTest.h
@@ -0,0 +1,19 @@
+// Copyright 1998-2014 Epic Games, Inc. All Rights Reserved.
+
+#pragma once
+
+#include "GameplayEffectStackingExtension.h"
+#include "GameplayTagContainer.h"
+#include "GameplayTagAssetInterface.h"
+#include "AttributeSet.h"
+#include "GameplayEffectStackingExtension_CappedNumberTest.generated.h"
+
+UCLASS(BlueprintType)
+class GAMEPLAYABILITIES_API UGameplayEffectStackingExtension_CappedNumberTest : public UGameplayEffectStackingExtension
+{
+	GENERATED_UCLASS_BODY()
+
+public:
+
+	void CalculateStack(TArray<FActiveGameplayEffect*>& CustomGameplayEffects, FActiveGameplayEffectsContainer& Container, FActiveGameplayEffect& CurrentEffect) OVERRIDE;
+};
\ No newline at end of file
diff --git a/Engine/Source/Runtime/GameplayAbilities/Public/GameplayEffectStackingExtension_DiminishingReturnsTest.h b/Engine/Source/Runtime/GameplayAbilities/Public/GameplayEffectStackingExtension_DiminishingReturnsTest.h
new file mode 100644
--- /dev/null
+++ b/Engine/Source/Runtime/GameplayAbilities/Public/GameplayEffectStackingExtension_DiminishingReturnsTest.h
@@ -0,0 +1,19 @@
+// Copyright 1998-2014 Epic Games, Inc. All Rights Reserved.
+
+#pragma once
+
+#include "GameplayEffectStackingExtension.h"
+#include "GameplayTagContainer.h"
+#include "GameplayTagAssetInterface.h"
+#include "AttributeSet.h"
+#include "GameplayEffectStackingExtension_DiminishingReturnsTest.generated.h"
+
+UCLASS(BlueprintType)
+class GAMEPLAYABILITIES_API UGameplayEffectStackingExtension_DiminishingReturnsTest : public UGameplayEffectStackingExtension
+{
+	GENERATED_UCLASS_BODY()
+
+public:
+
+	void CalculateStack(TArray<FActiveGameplayEffect*>& CustomGameplayEffects, FActiveGameplayEffectsContainer& Container, FActiveGameplayEffect& CurrentEffect) OVERRIDE;
+};
\ No newline at end of file
