Avoid issue of passing GameplayAbilityActorInfo* around as delegate payload parameter.
More comments.

diff --git a/Engine/Source/Runtime/GameplayAbilities/Private/Abilities/GameplayAbility_Montage.cpp b/Engine/Source/Runtime/GameplayAbilities/Private/Abilities/GameplayAbility_Montage.cpp
--- a/Engine/Source/Runtime/GameplayAbilities/Private/Abilities/GameplayAbility_Montage.cpp
+++ b/Engine/Source/Runtime/GameplayAbilities/Private/Abilities/GameplayAbility_Montage.cpp
@@ -42,7 +42,7 @@ void UGameplayAbility_Montage::ActivateAbility(const FGameplayAbilityActorInfo*
 		float const Duration = ActorInfo->AnimInstance->Montage_Play(MontageToPlay, PlayRate);
 
 		FOnMontageEnded EndDelegate;
-		EndDelegate.BindUObject(this, &UGameplayAbility_Montage::OnMontageEnded, ActorInfo, AppliedEffects);
+		EndDelegate.BindUObject(this, &UGameplayAbility_Montage::OnMontageEnded, ActorInfo->AbilitySystemComponent, AppliedEffects);
 		ActorInfo->AnimInstance->Montage_SetEndDelegate(EndDelegate);
 
 		if (SectionName != NAME_None)
@@ -52,15 +52,14 @@ void UGameplayAbility_Montage::ActivateAbility(const FGameplayAbilityActorInfo*
 	}
 }
 
-void UGameplayAbility_Montage::OnMontageEnded(UAnimMontage* Montage, bool bInterrupted, const FGameplayAbilityActorInfo* ActorInfo, TArray<FActiveGameplayEffectHandle> AppliedEffects)
+void UGameplayAbility_Montage::OnMontageEnded(UAnimMontage* Montage, bool bInterrupted, TWeakObjectPtr<UAbilitySystemComponent> AbilitySystemComponent, TArray<FActiveGameplayEffectHandle> AppliedEffects)
 {
 	// Remove any GameplayEffects that we applied
-	UAbilitySystemComponent *OwnerComponent = ActorInfo->AbilitySystemComponent.Get();
-	if (OwnerComponent)
+	if (AbilitySystemComponent.IsValid())
 	{
 		for (FActiveGameplayEffectHandle Handle : AppliedEffects)
 		{
-			OwnerComponent->RemoveActiveGameplayEffect(Handle);
+			AbilitySystemComponent->RemoveActiveGameplayEffect(Handle);
 		}
 	}
 }
\ No newline at end of file
diff --git a/Engine/Source/Runtime/GameplayAbilities/Public/Abilities/GameplayAbility.h b/Engine/Source/Runtime/GameplayAbilities/Public/Abilities/GameplayAbility.h
--- a/Engine/Source/Runtime/GameplayAbilities/Public/Abilities/GameplayAbility.h
+++ b/Engine/Source/Runtime/GameplayAbilities/Public/Abilities/GameplayAbility.h
@@ -8,12 +8,59 @@
 #include "GameplayAbilityTypes.h"
 #include "GameplayAbility.generated.h"
 
-
 /**
-* UGameplayAbility
-*	The GameplayAbility baseclass. This base class is a data asset and can be used to create non-instanced abilities.
-*	Instanced abilities should derive off UGameplyAbility_Instanced.
-*/
+ * UGameplayAbility
+ *	
+ *	Abilities define custom gameplay logic that can be activated or triggered.
+ *	
+ *	The main features provided by the AbilitySystem for GameplayAbilities are: 
+ *		-CanUse functionality:
+ *			-Cooldowns
+ *			-Resources (mana, stamina, etc)
+ *			-etc
+ *			
+ *		-Replication support
+ *			-Client/Server communication for ability activation
+ *			-Client prediction for ability activation
+ *			
+ *		-Instancing support
+ *			-Abilities can be non-instanced (default)
+ *			-Instanced per owner
+ *			-Instanced per execution
+ *			
+ *		-Basic, extendable support for:
+ *			-Input binding
+ *			-'Giving' abilities (that can be used) to actors
+ *	
+ *	
+ *	
+ *	UGameplayAbility is the base class. It is non instanced, non replicating, and not blueprintable. 
+ *	This is essentially the 'cheapest' type of ability with the least amount of overhead.
+ *	
+ *	The intention is for programmers to create these non instanced abilities in C++. Designers can then
+ *	extend them as data assets (E.g., they can change default properties, they cannot implement blueprint graphs).
+ *	
+ *	See GameplayAbility_Montage for example.
+ *		-Plays a montage and applies a GameplayEffect to its target while the montage is playing.
+ *		-When finished, removes GameplayEffect.
+ *		
+ *	
+ *	Note on replication support:
+ *		-Non instanced abilities have limited replication support. 
+ *			-Cannot have state (obviously) so no replicated properties
+ *			-RPCs on the ability class are not possible either.
+ *			
+ *			-However: generic RPC functionality can be achieved through the UAbilitySystemAttribute.
+ *				-E.g.: ServerTryActivateAbility(class UGameplayAbility* AbilityToActivate, int32 PredictionKey)
+ *				
+ *	A lot is possible with non instanced abilities but care must be taken.
+ *	
+ *	
+ *	To support state or event replication, an ability must be instanced. This can be done with the UGameplayAbility_Instanced class.
+ *	
+ *
+ *	
+ */
 UCLASS()
 class GAMEPLAYABILITIES_API UGameplayAbility : public UDataAsset
 {
diff --git a/Engine/Source/Runtime/GameplayAbilities/Public/Abilities/GameplayAbility_Instanced.h b/Engine/Source/Runtime/GameplayAbilities/Public/Abilities/GameplayAbility_Instanced.h
--- a/Engine/Source/Runtime/GameplayAbilities/Public/Abilities/GameplayAbility_Instanced.h
+++ b/Engine/Source/Runtime/GameplayAbilities/Public/Abilities/GameplayAbility_Instanced.h
@@ -8,9 +8,13 @@ class UAnimInstance;
 class UAbilitySystemComponent;
 
 /**
-* UGameplayEffect
-*	The GameplayEffect definition. This is the data asset defined in the editor that drives everything.
-*/
+ * UGameplayAbility_Instanced
+ *	
+ *	This is the base class for all instanced abilities. It is fully blueprintable.
+ *	
+ *
+ *
+ */
 UCLASS(Blueprintable)
 class GAMEPLAYABILITIES_API UGameplayAbility_Instanced : public UGameplayAbility
 {
diff --git a/Engine/Source/Runtime/GameplayAbilities/Public/Abilities/GameplayAbility_Montage.h b/Engine/Source/Runtime/GameplayAbilities/Public/Abilities/GameplayAbility_Montage.h
--- a/Engine/Source/Runtime/GameplayAbilities/Public/Abilities/GameplayAbility_Montage.h
+++ b/Engine/Source/Runtime/GameplayAbilities/Public/Abilities/GameplayAbility_Montage.h
@@ -30,5 +30,5 @@ public:
 	UPROPERTY(EditDefaultsOnly, Category = MontageAbility)
 	TArray<const UGameplayEffect*>	GameplayEffectsWhileAnimating;
 
-	void OnMontageEnded(UAnimMontage* Montage, bool bInterrupted, const FGameplayAbilityActorInfo* OwnerInfo, TArray<struct FActiveGameplayEffectHandle>	AppliedEffects);
+	void OnMontageEnded(UAnimMontage* Montage, bool bInterrupted, TWeakObjectPtr<UAbilitySystemComponent> AbilitySystemComponent, TArray<struct FActiveGameplayEffectHandle>	AppliedEffects);
 };
\ No newline at end of file
diff --git a/Engine/Source/Runtime/GameplayAbilities/Public/AbilitySystemComponent.h b/Engine/Source/Runtime/GameplayAbilities/Public/AbilitySystemComponent.h
--- a/Engine/Source/Runtime/GameplayAbilities/Public/AbilitySystemComponent.h
+++ b/Engine/Source/Runtime/GameplayAbilities/Public/AbilitySystemComponent.h
@@ -22,8 +22,35 @@ struct GAMEPLAYABILITIES_API FAttributeDefaults
 };
 
 /** 
+ *	UAbilitySystemComponent	
+ *
+ *	A component to easily interface with the 3 aspects of the AbilitySystem:
+ *		-GameplayAbilities
+ *		-GameplayEffects
+ *		-GameplayAttributes
+ *		
+ *	This component will make life easier for interfacing with these subsystems, but is not completely required. The main functions are:
+ *	
+ *	GameplayAbilities:
+ *		-Provides a way to give/assign abilities that can be used (by a player or AI for example)
+ *		-Provides managment of instanced abilities (something must hold onto them)
+ *		-Provides replication functionality
+ *			-Ability state must always be replicated on the UGameplayAbility itself, but UAbilitySystemComponent can provide RPC replication
+ *			for non-instanced gameplay abilities. (Explained more in GameplayAbility.h).
+ *			
+ *	GameplayEffects:
+ *		-Provides an FActiveGameplayEffectsContainer for holding active GameplayEffects
+ *		-Provides methods for apply GameplayEffect to a target or to self
+ *		-Provides wrappers for querying information in FActiveGameplayEffectsContainers (duration, magnitude, etc)
+ *		-Provides methods for clearing/remove GameplayEffects
+ *		
+ *	GameplayAttributes
+ *		-Provides methods for allocating and initializing attribute sets
+ *		-Provides methods for getting ATtributeSets
+ *  
  * 
  */
+
 UCLASS(ClassGroup=AbilitySystem, hidecategories=(Object,LOD,Lighting,Transform,Sockets,TextureStreaming), editinlinenew, meta=(BlueprintSpawnableComponent))
 class GAMEPLAYABILITIES_API UAbilitySystemComponent : public UActorComponent
 {
