Refer to active abilities by TWeakObjectPtr instead of index in callbacks.

diff --git a/Engine/Source/Runtime/GameplayAbilities/Private/AbilitySystemComponent_Abilities.cpp b/Engine/Source/Runtime/GameplayAbilities/Private/AbilitySystemComponent_Abilities.cpp
--- a/Engine/Source/Runtime/GameplayAbilities/Private/AbilitySystemComponent_Abilities.cpp
+++ b/Engine/Source/Runtime/GameplayAbilities/Private/AbilitySystemComponent_Abilities.cpp
@@ -78,16 +78,16 @@ UGameplayAbility * UAbilitySystemComponent::CreateNewInstanceOfAbility(UGameplay
 	return AbilityInstance;
 }
 
-bool UAbilitySystemComponent::ActivateAbility(int32 AbilityIdx)
+bool UAbilitySystemComponent::ActivateAbility(TWeakObjectPtr<UGameplayAbility> Ability)
 {
 	check(AbilityActorInfo.IsValid());
 
 	AActor * OwnerActor = GetOwner();
 	check(OwnerActor);
 
-	if (ActivatableAbilities.IsValidIndex(AbilityIdx))
+	if (Ability.IsValid())
 	{
-		ActivatableAbilities[AbilityIdx]->InputPressed(0, AbilityActorInfo.Get());
+		Ability.Get()->InputPressed(0, AbilityActorInfo.Get());
 	}
 
 	return false;
@@ -98,9 +98,9 @@ void UAbilitySystemComponent::CancelAbilitiesWithTags(const FGameplayTagContaine
 	/**
 	 *	FIXME
 	 *
-	 *	Right now we are cancelling all activatable abilites that match Tags. This includes abilities that might not have been activated in the first place!
+	 *	Right now we are canceling all activatable abilities that match Tags. This includes abilities that might not have been activated in the first place!
 	 *	For instanced-per-actor abilities this is fine. They could check if they were activated/still activating.
-	 *	For non-instanced abilities it is ambigous. We have no way to know 'how many' non-instanced abilities are in flight.
+	 *	For non-instanced abilities it is ambiguous. We have no way to know 'how many' non-instanced abilities are in flight.
 	 *	Likewise for instanced-per-execution abilities. Though they are present in Replicated/NonReplicatedInstancedAbilities list
 	 */
 
diff --git a/Engine/Source/Runtime/GameplayAbilities/Private/GameplayAbilitySet.cpp b/Engine/Source/Runtime/GameplayAbilities/Private/GameplayAbilitySet.cpp
--- a/Engine/Source/Runtime/GameplayAbilities/Private/GameplayAbilitySet.cpp
+++ b/Engine/Source/Runtime/GameplayAbilities/Private/GameplayAbilitySet.cpp
@@ -11,7 +11,7 @@
 //
 // --------------------------------------------------------------------------------------------------------------------------------------------------------
 
-void InputPressed(int32 AbilityIdx, int32 InputID, TWeakObjectPtr<AActor> ActorOwner)
+void InputPressed(TWeakObjectPtr<UGameplayAbility> Ability, int32 InputID, TWeakObjectPtr<AActor> ActorOwner)
 {
 	FGameplayAbilityActorInfo	ActorInfo;
 	ActorInfo.InitFromActor(ActorOwner.Get());
@@ -19,18 +19,14 @@ void InputPressed(int32 AbilityIdx, int32 InputID, TWeakObjectPtr<AActor> ActorO
 	UAbilitySystemComponent * AbilitySystemComponent = ActorInfo.AbilitySystemComponent.Get();
 	if (AbilitySystemComponent)
 	{
-		if (AbilitySystemComponent->ActivatableAbilities.IsValidIndex(AbilityIdx))
+		if (Ability.IsValid())
 		{
-			UGameplayAbility * Ability = AbilitySystemComponent->ActivatableAbilities[AbilityIdx];
-			if (Ability)
-			{
-				Ability->InputPressed(InputID, AbilitySystemComponent->AbilityActorInfo.Get());
-			}
+			Ability.Get()->InputPressed(InputID, AbilitySystemComponent->AbilityActorInfo.Get());
 		}
 	}
 }
 
-void InputReleased(int32 AbilityIdx, int32 InputID, TWeakObjectPtr<AActor> ActorOwner)
+void InputReleased(TWeakObjectPtr<UGameplayAbility> Ability, int32 InputID, TWeakObjectPtr<AActor> ActorOwner)
 {
 	FGameplayAbilityActorInfo	ActorInfo;
 	ActorInfo.InitFromActor(ActorOwner.Get());
@@ -38,13 +34,9 @@ void InputReleased(int32 AbilityIdx, int32 InputID, TWeakObjectPtr<AActor> Actor
 	UAbilitySystemComponent * AbilitySystemComponent = ActorInfo.AbilitySystemComponent.Get();
 	if (AbilitySystemComponent)
 	{
-		if (AbilitySystemComponent->ActivatableAbilities.IsValidIndex(AbilityIdx))
+		if (Ability.IsValid())
 		{
-			UGameplayAbility * Ability = AbilitySystemComponent->ActivatableAbilities[AbilityIdx];
-			if (Ability)
-			{
-				Ability->InputReleased(InputID, AbilitySystemComponent->AbilityActorInfo.Get());
-			}
+			Ability.Get()->InputReleased(InputID, AbilitySystemComponent->AbilityActorInfo.Get());
 		}
 	}
 }
@@ -94,14 +86,14 @@ void UGameplayAbilitySet::BindInputComponentToAbilities(UInputComponent *InputCo
 			// Pressed event
 			{
 				FInputActionBinding AB(FName(*CommandPair.CommandString), IE_Pressed);
-				AB.ActionDelegate.GetDelegateForManualSet().BindStatic(&InputPressed, idx, CommandPair.InputID, TWeakObjectPtr<AActor>(OwnerActor));
+				AB.ActionDelegate.GetDelegateForManualSet().BindStatic(&InputPressed, TWeakObjectPtr<UGameplayAbility>(Ability), CommandPair.InputID, TWeakObjectPtr<AActor>(OwnerActor));
 				InputComponent->AddActionBinding(AB);
 			}
 
 			// Released event
 			{
 				FInputActionBinding AB(FName(*CommandPair.CommandString), IE_Released);
-				AB.ActionDelegate.GetDelegateForManualSet().BindStatic(&InputReleased, idx, CommandPair.InputID, TWeakObjectPtr<AActor>(OwnerActor));
+				AB.ActionDelegate.GetDelegateForManualSet().BindStatic(&InputReleased, TWeakObjectPtr<UGameplayAbility>(Ability), CommandPair.InputID, TWeakObjectPtr<AActor>(OwnerActor));
 				InputComponent->AddActionBinding(AB);
 			}
 		}
diff --git a/Engine/Source/Runtime/GameplayAbilities/Public/AbilitySystemComponent.h b/Engine/Source/Runtime/GameplayAbilities/Public/AbilitySystemComponent.h
--- a/Engine/Source/Runtime/GameplayAbilities/Public/AbilitySystemComponent.h
+++ b/Engine/Source/Runtime/GameplayAbilities/Public/AbilitySystemComponent.h
@@ -261,7 +261,7 @@ class GAMEPLAYABILITIES_API UAbilitySystemComponent : public UActorComponent
 	 *	GameplayAbilities
 	 *	
 	 *	The role of the AbilitySystemComponent wrt Abilities is to provide:
-	 *		-Management of abiltiy instances (whether per actor or per execution instance).
+	 *		-Management of ability instances (whether per actor or per execution instance).
 	 *			-Someone *has* to keep track of these instances.
 	 *			-Non instanced abilities *could* be executed without any ability stuff in AbilitySystemComponent.
 	 *				They should be able to operate on an GameplayAbilityActorInfo + GameplayAbility.
@@ -276,11 +276,11 @@ class GAMEPLAYABILITIES_API UAbilitySystemComponent : public UActorComponent
 
 	void CancelAbilitiesWithTags(const FGameplayTagContainer Tags, const FGameplayAbilityActorInfo* ActorInfo, const FGameplayAbilityActivationInfo ActivationInfo, UGameplayAbility* Ignore);
 	
-	/** References to non-replicating abilities that we instanced. We need to keep these refernces to avoid GC */
+	/** References to non-replicating abilities that we instanced. We need to keep these references to avoid GC */
 	UPROPERTY()
 	TArray<class UGameplayAbility *>	NonReplicatedInstancedAbilities;
 
-	/** References to replicating abilities that we instanced. We need to keep these refernces to avoid GC */
+	/** References to replicating abilities that we instanced. We need to keep these references to avoid GC */
 	UPROPERTY(Replicated)
 	TArray<class UGameplayAbility *>	ReplicatedInstancedAbilities;
 
@@ -289,7 +289,7 @@ class GAMEPLAYABILITIES_API UAbilitySystemComponent : public UActorComponent
 	 *		-This will include CDOs for non instanced abilities and per-execution instanced abilities. 
 	 *		-Actor-instanced abilities will be the actual instance (not CDO)
 	 *		
-	 *	This array is not vital for things to work. It is a convenience thing for 'giving abilties to the actor'. But abilities could also work on things
+	 *	This array is not vital for things to work. It is a convenience thing for 'giving abilities to the actor'. But abilities could also work on things
 	 *	without an AbilitySystemComponent. For example an ability could be written to execute on a StaticMeshActor. As long as the ability doesn't require 
 	 *	instancing or anything else that the AbilitySystemComponent would provide, then it doesn't need the component to function.
 	 */ 
@@ -309,10 +309,9 @@ class GAMEPLAYABILITIES_API UAbilitySystemComponent : public UActorComponent
 	// ----------------------------------------------------------------------------------------------------------------
 
 	/** This is probably temp for testing. We want to think a bit more on the API for outside stuff activating abilities */
-	UFUNCTION(BlueprintCallable, Category = "Abilities")
-	bool	ActivateAbility(int32 AbilityIdx);
+	bool	ActivateAbility(TWeakObjectPtr<UGameplayAbility> Ability);
 
-	/** There needs to be a concept of an animating ability. Only one may exist at a time. New requestrs can be queued up, overridden, or ignored. */
+	/** There needs to be a concept of an animating ability. Only one may exist at a time. New requests can be queued up, overridden, or ignored. */
 	UPROPERTY()
 	class UGameplayAbility *	AnimatingAbility;
 
@@ -353,7 +352,7 @@ private:
 
 	const FGlobalCurveDataOverride * GetCurveDataOverride() const
 	{
-		// only return data if we have overrides. NULL if we dont.
+		// only return data if we have overrides. NULL if we don't.
 		return (GlobalCurveDataOverride.Overrides.Num() > 0 ? &GlobalCurveDataOverride : NULL);
 	}
 
