remove pragma once in cpp to fix NonUnityTestCompile_OnMac

diff --git a/Engine/Source/Runtime/GameplayAbilities/Private/GameplayCueView.cpp b/Engine/Source/Runtime/GameplayAbilities/Private/GameplayCueView.cpp
--- a/Engine/Source/Runtime/GameplayAbilities/Private/GameplayCueView.cpp
+++ b/Engine/Source/Runtime/GameplayAbilities/Private/GameplayCueView.cpp
@@ -1,7 +1,5 @@
 // Copyright 1998-2014 Epic Games, Inc. All Rights Reserved.
 
-#pragma once
-
 #include "AbilitySystemPrivatePCH.h"
 #include "GameplayEffect.h"
 #include "GameplayCueView.h"
