GameplayAbility updates, mostly work on base targeting system

diff --git a/Engine/Source/Editor/GameplayAbilitiesEditor/Classes/K2Node_LatentAbilityCall.h b/Engine/Source/Editor/GameplayAbilitiesEditor/Classes/K2Node_LatentAbilityCall.h
new file mode 100644
--- /dev/null
+++ b/Engine/Source/Editor/GameplayAbilitiesEditor/Classes/K2Node_LatentAbilityCall.h
@@ -0,0 +1,14 @@
+// Copyright 1998-2014 Epic Games, Inc. All Rights Reserved.
+
+#pragma once
+#include "K2Node_LatentAbilityCall.generated.h"
+
+UCLASS()
+class UK2Node_LatentAbilityCall : public UK2Node_BaseAsyncTask
+{
+	GENERATED_UCLASS_BODY()
+
+	// UEdGraphNode interface
+	virtual void GetMenuEntries(FGraphContextMenuBuilder& ContextMenuBuilder) const OVERRIDE;
+	// End of UEdGraphNode interface
+};
diff --git a/Engine/Source/Editor/GameplayAbilitiesEditor/Private/K2Node_LatentAbilityCall.cpp b/Engine/Source/Editor/GameplayAbilitiesEditor/Private/K2Node_LatentAbilityCall.cpp
new file mode 100644
--- /dev/null
+++ b/Engine/Source/Editor/GameplayAbilitiesEditor/Private/K2Node_LatentAbilityCall.cpp
@@ -0,0 +1,55 @@
+// Copyright 1998-2014 Epic Games, Inc. All Rights Reserved.
+
+
+#include "AbilitySystemEditorPrivatePCH.h"
+#include "AbilityTask.h"
+#include "K2Node_LatentAbilityCall.h"
+
+#define LOCTEXT_NAMESPACE "K2Node"
+
+UK2Node_LatentAbilityCall::UK2Node_LatentAbilityCall(const class FPostConstructInitializeProperties& PCIP)
+	: Super(PCIP)
+{
+	ProxyActivateFunctionName = GET_FUNCTION_NAME_CHECKED(UAbilityTask, Activate);
+}
+
+void UK2Node_LatentAbilityCall::GetMenuEntries(FGraphContextMenuBuilder& ContextMenuBuilder) const
+{
+	const UEdGraphSchema_K2* K2Schema = GetDefault<UEdGraphSchema_K2>();
+	EGraphType GraphType = K2Schema->GetGraphType(ContextMenuBuilder.CurrentGraph);
+	const bool bAllowLatentFuncs = (GraphType == GT_Ubergraph || GraphType == GT_Macro);
+
+	if (bAllowLatentFuncs)
+	{
+		for (TObjectIterator<UClass> ClassIt; ClassIt; ++ClassIt)
+		{
+			UClass* TestClass = *ClassIt;
+			if (TestClass->IsChildOf(UAbilityTask::StaticClass()) && !TestClass->HasAnyClassFlags(CLASS_Abstract | CLASS_Deprecated))
+			{
+				for (TFieldIterator<UFunction> FuncIt(TestClass, EFieldIteratorFlags::ExcludeSuper); FuncIt; ++FuncIt)
+				{
+					// See if the function is a static factory method for online proxies
+					UFunction* CandidateFunction = *FuncIt;
+
+					UObjectProperty* ReturnProperty = Cast<UObjectProperty>(CandidateFunction->GetReturnProperty());
+					const bool bValidReturnType = (ReturnProperty != nullptr) && (ReturnProperty->PropertyClass != nullptr) && (ReturnProperty->PropertyClass->IsChildOf(UAbilityTask::StaticClass()));
+
+					if (CandidateFunction->HasAllFunctionFlags(FUNC_Static) && bValidReturnType)
+					{
+						// Create a node template for this factory method
+						UK2Node_LatentAbilityCall* NodeTemplate = NewObject<UK2Node_LatentAbilityCall>(ContextMenuBuilder.OwnerOfTemporaries);
+						NodeTemplate->ProxyFactoryFunctionName = CandidateFunction->GetFName();
+						NodeTemplate->ProxyFactoryClass = TestClass;
+						NodeTemplate->ProxyClass = ReturnProperty->PropertyClass;
+
+						CreateDefaultMenuEntry(NodeTemplate, ContextMenuBuilder);
+					}
+				}
+			}
+		}
+	}
+}
+
+#undef LOCTEXT_NAMESPACE
+
+
diff --git a/Engine/Source/Runtime/GameplayAbilities/Private/Abilities/GameplayAbility.cpp b/Engine/Source/Runtime/GameplayAbilities/Private/Abilities/GameplayAbility.cpp
--- a/Engine/Source/Runtime/GameplayAbilities/Private/Abilities/GameplayAbility.cpp
+++ b/Engine/Source/Runtime/GameplayAbilities/Private/Abilities/GameplayAbility.cpp
@@ -40,7 +40,20 @@ bool UGameplayAbility::IsSupportedForNetworking() const
 
 bool UGameplayAbility::CanActivateAbility(const FGameplayAbilityActorInfo* ActorInfo) const
 {
-	// Check cooldowns and resource consumption immediately
+	if (ActorInfo->AbilitySystemComponent->GetUserAbilityActivationInhibited())
+	{
+		/**
+		 *	Input is inhibited (UI is pulled up, another ability may be blocking all other input, etc).
+		 *	When we get into triggered abilities, we may need to better differentiate between CanActviate and CanUserActivate or something.
+		 *	E.g., we would want LMB/RMB to be inhibited while the user is in the menu UI, but we wouldn't want to prevent a 'buff when I am low health'
+		 *	ability to not trigger.
+		 *	
+		 *	Basically: CanActivateAbility is only used by user activated abilities now. If triggered abilities need to check costs/cooldowns, then we may
+		 *	want to split this function up and change the calling API to distinguish between 'can I initiate an ability activation' and 'can this ability be activated'.
+		 */ 
+		false;
+	}
+	
 	if (!CheckCooldown(ActorInfo))
 	{
 		return false;
@@ -171,6 +184,10 @@ void UGameplayAbility::ActivateAbility(const FGameplayAbilityActorInfo* ActorInf
 void UGameplayAbility::PreActivate(const FGameplayAbilityActorInfo* ActorInfo, const FGameplayAbilityActivationInfo ActivationInfo)
 {
 	ActorInfo->AbilitySystemComponent->CancelAbilitiesWithTags(CancelAbilitiesWithTag, ActorInfo, ActivationInfo, this);
+
+	// Become the 'Targeting Ability'
+	// These may need to be more robust - does every ability that activates become the targeting ability, or only certain ones?
+	ActorInfo->AbilitySystemComponent->SetTargetAbility(this);
 }
 
 void UGameplayAbility::CallActivateAbility(const FGameplayAbilityActorInfo* ActorInfo, const FGameplayAbilityActivationInfo ActivationInfo)
@@ -262,13 +279,26 @@ float UGameplayAbility::GetCooldownTimeRemaining(const FGameplayAbilityActorInfo
 		if (Durations.Num() > 0)
 		{
 			Durations.Sort();
-			return Durations[0];
+			return Durations[Durations.Num()-1];
 		}
 	}
 
 	return 0.f;
 }
 
+const FGameplayAbilityActorInfo* UGameplayAbility::GetCurrentActorInfo()
+{
+	ABILITY_LOG(Fatal, TEXT("GetCurrentActorInfo called on base UGameplayAbility class"));
+	return NULL;
+}
+
+FGameplayAbilityActivationInfo UGameplayAbility::GetCurrentActivationInfo()
+{
+	ABILITY_LOG(Fatal, TEXT("GetCurrentActivationInfo called on base UGameplayAbility class"));
+
+	return FGameplayAbilityActivationInfo();
+}
+
 // --------------------------------------------------------------------
 
 void UGameplayAbility::MontageBranchPoint_AbilityDecisionStop(const FGameplayAbilityActorInfo* ActorInfo) const
@@ -298,6 +328,28 @@ void FGameplayAbilityActorInfo::InitFromActor(AActor *InActor)
 {
 	Actor = InActor;
 
+	UStruct * Struct = FGameplayAbilityActorInfo::StaticStruct();
+	UScriptStruct * ScriptStruct = FGameplayAbilityActorInfo::StaticStruct();
+
+	// Look for a player controller or pawn in the owner chain.
+	AActor *TestActor = InActor;
+	while(TestActor)
+	{
+		if (APlayerController * CastPC = Cast<APlayerController>(TestActor))
+		{
+			PlayerController = CastPC;
+			break;
+		}
+
+		if (APawn * Pawn = Cast<APawn>(TestActor))
+		{
+			PlayerController = Cast<APlayerController>(Pawn->GetController());
+			break;
+		}
+
+		TestActor = TestActor->GetOwner();
+	}
+
 	// Grab Components that we care about
 	USkeletalMeshComponent * SkelMeshComponent = InActor->FindComponentByClass<USkeletalMeshComponent>();
 	if (SkelMeshComponent)
@@ -310,6 +362,16 @@ void FGameplayAbilityActorInfo::InitFromActor(AActor *InActor)
 	MovementComponent = InActor->FindComponentByClass<UMovementComponent>();
 }
 
+bool FGameplayAbilityActorInfo::IsLocallyControlled() const
+{
+	if (PlayerController.IsValid())
+	{
+		return PlayerController->IsLocalController();
+	}
+
+	return false;
+}
+
 void FGameplayAbilityActivationInfo::GeneratePredictionKey(UAbilitySystemComponent * Component) const
 {
 	check(Component);
@@ -366,3 +428,72 @@ FString FGameplayAbilityTargetData::ToString() const
 {
 	return TEXT("BASE CLASS");
 }
+
+
+bool FGameplayAbilityTargetDataHandle::NetSerialize(FArchive& Ar, class UPackageMap* Map, bool& bOutSuccess)
+{
+
+	UScriptStruct * ScriptStruct = Data.IsValid() ? Data->GetScriptStruct() : NULL;
+	Ar << ScriptStruct;
+
+	if (ScriptStruct)
+	{
+		if (Ar.IsLoading())
+		{
+			// For now, just always reset/reallocate the data when loading.
+			// Longer term if we want to generalize this and use it for property replication, we should support
+			// only reallocating when necessary
+			check(!Data.IsValid());
+
+			FGameplayAbilityTargetData * NewData = (FGameplayAbilityTargetData*)FMemory::Malloc(ScriptStruct->GetCppStructOps()->GetSize());
+			ScriptStruct->InitializeScriptStruct(NewData);
+
+			Data = TSharedPtr<FGameplayAbilityTargetData>(NewData);
+		}
+
+		void* ContainerPtr = Data.Get();
+
+		if (ScriptStruct->StructFlags & STRUCT_NetSerializeNative)
+		{
+			ScriptStruct->GetCppStructOps()->NetSerialize(Ar, Map, bOutSuccess, Data.Get());
+		}
+		else
+		{
+			// This won't work since UStructProperty::NetSerializeItem is deprecrated.
+			//	1) we have to manually crawl through the topmost struct's fields since we don't have a UStructProperty for it (just the UScriptProperty)
+			//	2) if there are any UStructProperties in the topmost struct's fields, we will assert in UStructProperty::NetSerializeItem.
+			
+			ABILITY_LOG(Fatal, TEXT("FGameplayAbilityTargetDataHandle::NetSerialize called on data struct %s without a native NetSerilaize"), *ScriptStruct->GetName());
+
+			for (TFieldIterator<UProperty> It(ScriptStruct); It; ++It)
+			{
+				if (It->PropertyFlags & CPF_RepSkip)
+				{
+					continue;
+				}
+
+				void * PropertyData = It->ContainerPtrToValuePtr<void*>(ContainerPtr);
+
+				It->NetSerializeItem(Ar, Map, PropertyData);
+			}
+		}
+	}
+	
+
+	ABILITY_LOG(Warning, TEXT("FGameplayAbilityTargetDataHandle Serialized: %s"), ScriptStruct ? *ScriptStruct->GetName() : TEXT("NULL") );
+
+	bOutSuccess = true;
+	return true;
+}
+
+bool FGameplayAbilityTargetData_SingleTargetHit::NetSerialize(FArchive& Ar, class UPackageMap* Map, bool& bOutSuccess)
+{
+	Ar << HitResult.Actor;
+
+	HitResult.Location.NetSerialize(Ar, Map, bOutSuccess);
+	HitResult.Normal.NetSerialize(Ar, Map, bOutSuccess);
+
+	bOutSuccess = true;
+	return true;
+}
+
diff --git a/Engine/Source/Runtime/GameplayAbilities/Private/Abilities/GameplayAbilitySet.cpp b/Engine/Source/Runtime/GameplayAbilities/Private/Abilities/GameplayAbilitySet.cpp
new file mode 100644
--- /dev/null
+++ b/Engine/Source/Runtime/GameplayAbilities/Private/Abilities/GameplayAbilitySet.cpp
@@ -0,0 +1,105 @@
+// Copyright 1998-2014 Epic Games, Inc. All Rights Reserved.
+
+#include "AbilitySystemPrivatePCH.h"
+#include "AbilitySystemComponent.h"
+#include "Abilities/GameplayAbility.h"
+#include "GameplayAbilitySet.h"
+
+// --------------------------------------------------------------------------------------------------------------------------------------------------------
+//
+//	UGameplayAbilitySet
+//
+// --------------------------------------------------------------------------------------------------------------------------------------------------------
+
+void InputPressed(TWeakObjectPtr<UGameplayAbility> Ability, int32 InputID, TWeakObjectPtr<AActor> ActorOwner)
+{
+	FGameplayAbilityActorInfo	ActorInfo;
+	ActorInfo.InitFromActor(ActorOwner.Get());
+
+	UAbilitySystemComponent * AbilitySystemComponent = ActorInfo.AbilitySystemComponent.Get();
+	if (AbilitySystemComponent)
+	{
+		// FIXME: We have to reinit the ability component actor info to get a valid player controller
+		AbilitySystemComponent->InitAbilityActorInfo();
+
+		if (Ability.IsValid())
+		{
+			Ability.Get()->InputPressed(InputID, AbilitySystemComponent->AbilityActorInfo.Get());
+		}
+	}
+}
+
+void InputReleased(TWeakObjectPtr<UGameplayAbility> Ability, int32 InputID, TWeakObjectPtr<AActor> ActorOwner)
+{
+	FGameplayAbilityActorInfo	ActorInfo;
+	ActorInfo.InitFromActor(ActorOwner.Get());
+
+	UAbilitySystemComponent * AbilitySystemComponent = ActorInfo.AbilitySystemComponent.Get();
+	if (AbilitySystemComponent)
+	{
+		if (Ability.IsValid())
+		{
+			Ability.Get()->InputReleased(InputID, AbilitySystemComponent->AbilityActorInfo.Get());
+		}
+	}
+}
+
+// --------------------------------------------------------------------------------------------------------------------------------------------------------
+//
+//	UGameplayAbilitySet
+//
+// --------------------------------------------------------------------------------------------------------------------------------------------------------
+
+
+UGameplayAbilitySet::UGameplayAbilitySet(const class FPostConstructInitializeProperties& PCIP)
+	: Super(PCIP)
+{
+
+}
+
+void UGameplayAbilitySet::InitializeAbilities(UAbilitySystemComponent *AbilitySystemComponent) const
+{
+	for (const FGameplayAbilityBindInfo &BindInfo : this->Abilities)
+	{
+		UGameplayAbility * Ability = BindInfo.GameplayAbilityClass ? BindInfo.GameplayAbilityClass->GetDefaultObject<UGameplayAbility>() : BindInfo.GameplayAbilityInstance;
+		if (Ability)
+		{
+			if (Ability->GetInstancingPolicy() == EGameplayAbilityInstancingPolicy::InstancedPerActor)
+			{
+				Ability = AbilitySystemComponent->CreateNewInstanceOfAbility(Ability);
+			}
+
+			AbilitySystemComponent->ActivatableAbilities.Add(Ability);
+		}
+	}
+}
+
+void UGameplayAbilitySet::BindInputComponentToAbilities(UInputComponent *InputComponent, UAbilitySystemComponent *AbilitySystemComponent) const
+{
+	AActor * OwnerActor = AbilitySystemComponent->GetOwner();
+	check(OwnerActor);
+
+	for (int32 idx = 0; idx < this->Abilities.Num(); ++idx)
+	{
+		const FGameplayAbilityBindInfo & BindInfo = this->Abilities[idx];
+		UGameplayAbility * Ability = AbilitySystemComponent->ActivatableAbilities[idx];	
+		
+		for (const FGameplayAbilityBindInfoCommandIDPair &CommandPair : BindInfo.CommandBindings)
+		{
+			// Pressed event
+			{
+				FInputActionBinding AB(FName(*CommandPair.CommandString), IE_Pressed);
+				AB.ActionDelegate.GetDelegateForManualSet().BindStatic(&InputPressed, TWeakObjectPtr<UGameplayAbility>(Ability), CommandPair.InputID, TWeakObjectPtr<AActor>(OwnerActor));
+				InputComponent->AddActionBinding(AB);
+			}
+
+			// Released event
+			{
+				FInputActionBinding AB(FName(*CommandPair.CommandString), IE_Released);
+				AB.ActionDelegate.GetDelegateForManualSet().BindStatic(&InputReleased, TWeakObjectPtr<UGameplayAbility>(Ability), CommandPair.InputID, TWeakObjectPtr<AActor>(OwnerActor));
+				InputComponent->AddActionBinding(AB);
+			}
+		}
+	}
+}
+
diff --git a/Engine/Source/Runtime/GameplayAbilities/Private/Abilities/GameplayAbilityTargetActor.cpp b/Engine/Source/Runtime/GameplayAbilities/Private/Abilities/GameplayAbilityTargetActor.cpp
new file mode 100644
--- /dev/null
+++ b/Engine/Source/Runtime/GameplayAbilities/Private/Abilities/GameplayAbilityTargetActor.cpp
@@ -0,0 +1,16 @@
+// Copyright 1998-2013 Epic Games, Inc. All Rights Reserved.
+
+#include "AbilitySystemPrivatePCH.h"
+#include "GameplayAbilityTargetActor.h"
+
+// --------------------------------------------------------------------------------------------------------------------------------------------------------
+//
+//	AGameplayAbilityTargetActor
+//
+// --------------------------------------------------------------------------------------------------------------------------------------------------------
+
+AGameplayAbilityTargetActor::AGameplayAbilityTargetActor(const class FPostConstructInitializeProperties& PCIP)
+	: Super(PCIP)
+{
+	StaticTargetFunction = false;
+}
\ No newline at end of file
diff --git a/Engine/Source/Runtime/GameplayAbilities/Private/Abilities/GameplayAbilityTargetActor_SingleLineTrace.cpp b/Engine/Source/Runtime/GameplayAbilities/Private/Abilities/GameplayAbilityTargetActor_SingleLineTrace.cpp
new file mode 100644
--- /dev/null
+++ b/Engine/Source/Runtime/GameplayAbilities/Private/Abilities/GameplayAbilityTargetActor_SingleLineTrace.cpp
@@ -0,0 +1,99 @@
+// Copyright 1998-2014 Epic Games, Inc. All Rights Reserved.
+
+#include "AbilitySystemPrivatePCH.h"
+#include "AbilitySystemComponent.h"
+#include "GameplayAbilityTargetActor_SingleLineTrace.h"
+#include "Engine/World.h"
+
+// --------------------------------------------------------------------------------------------------------------------------------------------------------
+//
+//	AGameplayAbilityTargetActor_SingleLineTrace
+//
+// --------------------------------------------------------------------------------------------------------------------------------------------------------
+
+AGameplayAbilityTargetActor_SingleLineTrace::AGameplayAbilityTargetActor_SingleLineTrace(const class FPostConstructInitializeProperties& PCIP)
+	: Super(PCIP)
+{
+	PrimaryActorTick.bCanEverTick = true;
+	PrimaryActorTick.TickGroup = TG_PrePhysics;
+	StaticTargetFunction = false;
+	bDebug = false;
+}
+
+FGameplayAbilityTargetDataHandle AGameplayAbilityTargetActor_SingleLineTrace::StaticGetTargetData(UWorld * World, const FGameplayAbilityActorInfo* ActorInfo, FGameplayAbilityActivationInfo ActivationInfo)
+{
+	// very temp - do a mostly hardcoded trace from the source actor
+
+	AActor * SourceActor = ActorInfo->Actor.Get();
+	check(SourceActor);
+
+	static const FName LineTraceSingleName(TEXT("AGameplayAbilityTargetActor_SingleLineTrace"));
+	bool bTraceComplex = false;
+	TArray<AActor*> ActorsToIgnore;
+
+	ActorsToIgnore.Add(SourceActor);
+
+	FCollisionQueryParams Params(LineTraceSingleName, bTraceComplex);
+	Params.bReturnPhysicalMaterial = true;
+	Params.bTraceAsyncScene = true;
+	Params.AddIgnoredActors(ActorsToIgnore);
+
+	FVector Start = SourceActor->GetActorLocation();
+	FVector End = Start + (SourceActor->GetActorForwardVector() * 3000.f);
+
+	// ------------------------------------------------------
+
+	FGameplayAbilityTargetData_SingleTargetHit* ReturnData = new FGameplayAbilityTargetData_SingleTargetHit();
+
+	SourceActor->GetWorld()->LineTraceSingle(ReturnData->HitResult, Start, End, ECC_WorldStatic, Params);
+
+	if (bDebug)
+	{
+		DrawDebugLine(World, Start, End, FLinearColor::Green, false);
+		DrawDebugSphere(World, ReturnData->HitResult.Location, 16, 10, FLinearColor::Green, false);
+	}
+
+	return FGameplayAbilityTargetDataHandle(ReturnData);
+}
+
+void AGameplayAbilityTargetActor_SingleLineTrace::StartTargeting(UGameplayAbility* InAbility)
+{
+	Ability = InAbility;
+
+	UAbilitySystemComponent* ASC = Ability->GetCurrentActorInfo()->AbilitySystemComponent.Get();
+	if (ASC)
+	{
+		ASC->ConfirmCallbacks.AddDynamic(this, &AGameplayAbilityTargetActor_SingleLineTrace::Confirm);
+		ASC->CancelCallbacks.AddDynamic(this, &AGameplayAbilityTargetActor_SingleLineTrace::Cancel);
+	}
+	
+	bDebug = true;
+}
+
+void AGameplayAbilityTargetActor_SingleLineTrace::Confirm()
+{
+	if (Ability.IsValid())
+	{
+		bDebug = false;
+		FGameplayAbilityTargetDataHandle Handle = StaticGetTargetData(Ability->GetWorld(), Ability->GetCurrentActorInfo(), Ability->GetCurrentActivationInfo());
+		TargetDataReadyDelegate.Broadcast(Handle);
+	}
+
+	Destroy();
+}
+
+void AGameplayAbilityTargetActor_SingleLineTrace::Cancel()
+{
+	Destroy();
+}
+
+void AGameplayAbilityTargetActor_SingleLineTrace::Tick(float DeltaSeconds)
+{
+	Super::Tick(DeltaSeconds);
+
+	/** Temp: Do a trace wiuth bDebug=true to draw a cheap "preview" */
+	if (Ability.IsValid())
+	{
+		StaticGetTargetData(Ability->GetWorld(), Ability->GetCurrentActorInfo(), Ability->GetCurrentActivationInfo());
+	}
+}
\ No newline at end of file
diff --git a/Engine/Source/Runtime/GameplayAbilities/Private/Abilities/GameplayAbility_Instanced.cpp b/Engine/Source/Runtime/GameplayAbilities/Private/Abilities/GameplayAbility_Instanced.cpp
--- a/Engine/Source/Runtime/GameplayAbilities/Private/Abilities/GameplayAbility_Instanced.cpp
+++ b/Engine/Source/Runtime/GameplayAbilities/Private/Abilities/GameplayAbility_Instanced.cpp
@@ -48,7 +48,7 @@ void UGameplayAbility_Instanced::CommitExecute(const FGameplayAbilityActorInfo*
 
 	Super::CommitExecute(ActorInfo, ActivationInfo);
 
-	check(!HasAnyFlags(RF_ClassDefaultObject));
+	ensure(!HasAnyFlags(RF_ClassDefaultObject));
 
 	// Call into blueprint (fixme: cache this)
 	static FName FuncName = FName(TEXT("K2_CommitAbility"));
diff --git a/Engine/Source/Runtime/GameplayAbilities/Private/Abilities/Tasks/AbilityTask.cpp b/Engine/Source/Runtime/GameplayAbilities/Private/Abilities/Tasks/AbilityTask.cpp
--- a/Engine/Source/Runtime/GameplayAbilities/Private/Abilities/Tasks/AbilityTask.cpp
+++ b/Engine/Source/Runtime/GameplayAbilities/Private/Abilities/Tasks/AbilityTask.cpp
@@ -6,4 +6,8 @@ UAbilityTask::UAbilityTask(const class FPostConstructInitializeProperties& PCIP)
 	: Super(PCIP)
 {
 
-}
\ No newline at end of file
+}
+
+void UAbilityTask::Activate()
+{
+}
diff --git a/Engine/Source/Runtime/GameplayAbilities/Private/Abilities/Tasks/AbilityTask_WaitConfirmCancel.cpp b/Engine/Source/Runtime/GameplayAbilities/Private/Abilities/Tasks/AbilityTask_WaitConfirmCancel.cpp
new file mode 100644
--- /dev/null
+++ b/Engine/Source/Runtime/GameplayAbilities/Private/Abilities/Tasks/AbilityTask_WaitConfirmCancel.cpp
@@ -0,0 +1,67 @@
+
+#include "AbilitySystemPrivatePCH.h"
+#include "Abilities/Tasks/AbilityTask_WaitConfirmCancel.h"
+
+#include "AbilitySystemComponent.h"
+
+#include "Abilities/GameplayAbility.h"
+#include "Abilities/GameplayAbility_Instanced.h"
+
+UAbilityTask_WaitConfirmCancel::UAbilityTask_WaitConfirmCancel(const class FPostConstructInitializeProperties& PCIP)
+	: Super(PCIP)
+{
+	ASC = NULL;
+}
+
+void UAbilityTask_WaitConfirmCancel::OnConfirmCallback()
+{
+	OnConfirm.Broadcast(true);
+	ASC->ConsumeAbilityConfirm();
+}
+
+void UAbilityTask_WaitConfirmCancel::OnCancelCallback()
+{
+	OnCancel.Broadcast(false);
+	ASC->ConsumeAbilityConfirm();
+}
+
+UAbilityTask_WaitConfirmCancel* UAbilityTask_WaitConfirmCancel::WaitConfirmCancel(UObject* WorldContextObject)
+{
+	check(WorldContextObject);
+	UGameplayAbility_Instanced* ThisAbility = CastChecked<UGameplayAbility_Instanced>(WorldContextObject);
+	if (ThisAbility)
+	{
+		UAbilityTask_WaitConfirmCancel * MyObj = NULL;
+		MyObj = NewObject<UAbilityTask_WaitConfirmCancel>();
+		MyObj->Ability = ThisAbility;
+
+		return MyObj;
+	}
+	return NULL;
+}
+
+void UAbilityTask_WaitConfirmCancel::Activate()
+{
+	if (Ability.IsValid())
+	{
+		const FGameplayAbilityActorInfo* Info = Ability.Get()->GetCurrentActorInfo();
+
+		ASC = Info->AbilitySystemComponent.Get();
+
+		if (!Info->IsLocallyControlled())
+		{
+			if (ASC->ReplicatedConfirmAbility)
+			{
+				// It could have already been confirmed (we may have been behind a client predicting this ability, so it could have got here
+				// before we actually executed this task)
+				OnConfirmCallback();
+				ASC->ConsumeAbilityConfirm();
+				return;
+			}
+		}
+		
+		// We have to wait for the callback from the AbilitySystemComponent. Which may be instigated locally or through network replication
+		ASC->ConfirmCallbacks.AddDynamic(this, &UAbilityTask_WaitConfirmCancel::OnConfirmCallback);
+		ASC->CancelCallbacks.AddDynamic(this, &UAbilityTask_WaitConfirmCancel::OnCancelCallback);
+	}
+}
\ No newline at end of file
diff --git a/Engine/Source/Runtime/GameplayAbilities/Private/Abilities/Tasks/AbilityTask_WaitTargetData.cpp b/Engine/Source/Runtime/GameplayAbilities/Private/Abilities/Tasks/AbilityTask_WaitTargetData.cpp
new file mode 100644
--- /dev/null
+++ b/Engine/Source/Runtime/GameplayAbilities/Private/Abilities/Tasks/AbilityTask_WaitTargetData.cpp
@@ -0,0 +1,80 @@
+
+#include "AbilitySystemPrivatePCH.h"
+#include "GameplayAbilityTargetActor.h"
+#include "AbilitySystemComponent.h"
+#include "Abilities/Tasks/AbilityTask_WaitTargetData.h"
+
+UAbilityTask_WaitTargetData::UAbilityTask_WaitTargetData(const class FPostConstructInitializeProperties& PCIP)
+	: Super(PCIP)
+{
+	ASC = NULL;
+}
+
+UAbilityTask_WaitTargetData* UAbilityTask_WaitTargetData::WaitTargetData(UObject* WorldContextObject, TSubclassOf<AGameplayAbilityTargetActor> InTargetClass)
+{
+	UAbilityTask_WaitTargetData * MyObj = NewObject<UAbilityTask_WaitTargetData>();
+	UGameplayAbility * ThisAbility = CastChecked<UGameplayAbility>(WorldContextObject);
+	MyObj->Ability = ThisAbility;
+	MyObj->ASC = ThisAbility->GetCurrentActorInfo()->AbilitySystemComponent.Get();
+	MyObj->TargetClass = InTargetClass;
+	return MyObj;	
+}
+
+void UAbilityTask_WaitTargetData::Activate()
+{
+	if (Ability.IsValid())
+	{
+		
+		if (!Ability.Get()->GetCurrentActorInfo()->IsLocallyControlled())
+		{
+			// If not locally controlled (server for remote client), see if TargetData was already sent
+			// else register callback for when it does get here
+
+			if (ASC->ReplicatedTargetData.IsValid())
+			{
+				ValidData.Broadcast(ASC->ReplicatedTargetData);
+				ASC->ConsumeAbilityTargetData();
+				return;
+			}
+			else
+			{
+				ASC->ReplicatedTargetDataDelegate.AddUObject(this, &UAbilityTask_WaitTargetData::OnTargetDataReplicatedCallback);
+				return;
+			}
+		}
+		else
+		{
+			// Locally controlled - spawn the targeting actor.
+			AGameplayAbilityTargetActor* CDO = CastChecked<AGameplayAbilityTargetActor>(TargetClass->GetDefaultObject());
+			if (CDO->StaticTargetFunction)
+			{
+				// This is just a static function that should instantly give us back target data
+				FGameplayAbilityTargetDataHandle Data = CDO->StaticGetTargetData(Ability->GetWorld(), Ability->GetCurrentActorInfo(), Ability->GetCurrentActivationInfo());
+				OnTargetDataReadyCallback(Data);
+			}
+			else
+			{
+				// This is a latent thing, probably with visuals and other gameplay related stuff. Spawn an actor locally and it will tell us when it has TargetData.
+				AGameplayAbilityTargetActor* SpawnedActor = CastChecked<AGameplayAbilityTargetActor>(Ability->GetWorld()->SpawnActor(TargetClass));
+				SpawnedActor->TargetDataReadyDelegate.AddUObject(this, &UAbilityTask_WaitTargetData::OnTargetDataReadyCallback);
+
+				SpawnedActor->StartTargeting(Ability.Get());
+			}
+		}
+	}
+}
+
+void UAbilityTask_WaitTargetData::OnTargetDataReplicatedCallback(FGameplayAbilityTargetDataHandle Data)
+{
+	check(ASC);
+	ValidData.Broadcast(Data);
+	ASC->ConsumeAbilityTargetData();
+}
+
+void UAbilityTask_WaitTargetData::OnTargetDataReadyCallback(FGameplayAbilityTargetDataHandle Data)
+{
+	check(ASC);
+	ASC->ServerSetReplicatedTargetData(Data);
+	ValidData.Broadcast(Data);
+	ASC->ConsumeAbilityTargetData();
+}
\ No newline at end of file
diff --git a/Engine/Source/Runtime/GameplayAbilities/Private/AbilitySystemBlueprintLibrary.cpp b/Engine/Source/Runtime/GameplayAbilities/Private/AbilitySystemBlueprintLibrary.cpp
--- a/Engine/Source/Runtime/GameplayAbilities/Private/AbilitySystemBlueprintLibrary.cpp
+++ b/Engine/Source/Runtime/GameplayAbilities/Private/AbilitySystemBlueprintLibrary.cpp
@@ -2,11 +2,16 @@
 
 #include "AbilitySystemPrivatePCH.h"
 #include "AbilitySystemBlueprintLibrary.h"
+
+#include "AbilitySystemComponent.h"
+
 #include "Abilities/GameplayAbility.h"
 #include "Abilities/GameplayAbility_Instanced.h"
+
 #include "Abilities/Tasks/AbilityTask_PlayMontageAndWait.h"
 #include "Abilities/Tasks/AbilityTask_WaitMovementModeChange.h"
 #include "Abilities/Tasks/AbilityTask_WaitOverlap.h"
+#include "Abilities/Tasks/AbilityTask_WaitConfirmCancel.h"
 #include "LatentActions.h"
 
 UAbilitySystemBlueprintLibrary::UAbilitySystemBlueprintLibrary(const class FPostConstructInitializeProperties& PCIP)
@@ -14,6 +19,52 @@ UAbilitySystemBlueprintLibrary::UAbilitySystemBlueprintLibrary(const class FPost
 {
 }
 
+UAbilitySystemComponent* UAbilitySystemBlueprintLibrary::GetAbilitySystemComponent(AActor *Actor)
+{
+	/**
+	 *	Fixme: This is supposed to be a small optimization but here we are going through module code, then a global uobject, then finally a virtual function - yuck!
+	 *		-think of a better way to have a global, overridable per project function. Thats it!
+	 */
+	return UAbilitySystemGlobals::Get().GetAbilitySystemComponentFromActor(Actor);
+}
+
+void UAbilitySystemBlueprintLibrary::ApplyGameplayEffectToTargetData(FGameplayAbilityTargetDataHandle Target, UGameplayEffect *GameplayEffect, const FGameplayAbilityActorInfo InstigatorInfo)
+{
+	if (Target.Data.IsValid())
+	{
+		Target.Data->ApplyGameplayEffect(GameplayEffect, InstigatorInfo);
+	}
+}
+
+FGameplayAbilityTargetDataHandle UAbilitySystemBlueprintLibrary::AbilityTargetDataFromHitResult(FHitResult HitResult)
+{
+	// Construct TargetData
+	FGameplayAbilityTargetData_SingleTargetHit * TargetData = new FGameplayAbilityTargetData_SingleTargetHit(HitResult);
+
+	// Give it a handle and return
+	FGameplayAbilityTargetDataHandle	Handle;
+	Handle.Data = TSharedPtr<FGameplayAbilityTargetData>(TargetData);
+
+	return Handle;
+}
+
+FHitResult UAbilitySystemBlueprintLibrary::GetHitResultFromTargetData(FGameplayAbilityTargetDataHandle TargetData)
+{
+	FGameplayAbilityTargetData * Data = TargetData.Data.Get();
+	if (Data)
+	{
+		const FHitResult * HitResultPtr = Data->GetHitResult();
+		if (HitResultPtr)
+		{
+			return *HitResultPtr;
+		}
+	}
+
+	return FHitResult();
+}
+
+// -------------------------------------------------------------------------------------
+
 UAbilityTask_PlayMontageAndWait* UAbilitySystemBlueprintLibrary::CreatePlayMontageAndWaitProxy(class UObject* WorldContextObject, UAnimMontage *MontageToPlay)
 {
 	check(WorldContextObject);
@@ -85,7 +136,7 @@ UAbilityTask_WaitOverlap* UAbilitySystemBlueprintLibrary::CreateWaitOverlap(clas
 		MyObj = NewObject<UAbilityTask_WaitOverlap>();
 
 
-		// TEMP
+		// TEMP - we are just using root component's collision. A real system will need more data to specify which component to use
 		UPrimitiveComponent * PrimComponent = Cast<UPrimitiveComponent>(ActorOwner->GetRootComponent());
 		if (!PrimComponent)
 		{			
@@ -100,22 +151,4 @@ UAbilityTask_WaitOverlap* UAbilitySystemBlueprintLibrary::CreateWaitOverlap(clas
 	}
 
 	return NULL;
-}
-
-/**
- *	Fixme: This is supposed to be a small optimization but here we are going through module code, then a global uobject, then finally a virtual function - yuck!
- *		-think of a better way to have a global, overridable per project function. Thats it!
- */
-
-UAbilitySystemComponent* UAbilitySystemBlueprintLibrary::GetAbilitySystemComponent(AActor *Actor)
-{
-	return UAbilitySystemGlobals::Get().GetAbilitySystemComponentFromActor(Actor);
-}
-
-void UAbilitySystemBlueprintLibrary::ApplyGameplayEffectToTargetData(FGameplayAbilityTargetDataHandle Target, UGameplayEffect *GameplayEffect, const FGameplayAbilityActorInfo InstigatorInfo)
-{
-	if (Target.Data.IsValid())
-	{
-		Target.Data->ApplyGameplayEffect(GameplayEffect, InstigatorInfo);
-	}
 }
\ No newline at end of file
diff --git a/Engine/Source/Runtime/GameplayAbilities/Private/AbilitySystemComponent.cpp b/Engine/Source/Runtime/GameplayAbilities/Private/AbilitySystemComponent.cpp
--- a/Engine/Source/Runtime/GameplayAbilities/Private/AbilitySystemComponent.cpp
+++ b/Engine/Source/Runtime/GameplayAbilities/Private/AbilitySystemComponent.cpp
@@ -43,6 +43,8 @@ UAbilitySystemComponent::UAbilitySystemComponent(const class FPostConstructIniti
 	SetIsReplicated(true);
 	ReplicatedPredictionKey = 0;
 	LocalPredictionKey = 0;
+	UserAbilityActivationInhibited = false;
+	TargetingAbility = NULL;
 }
 
 UAbilitySystemComponent::~UAbilitySystemComponent()
diff --git a/Engine/Source/Runtime/GameplayAbilities/Private/AbilitySystemComponent_Abilities.cpp b/Engine/Source/Runtime/GameplayAbilities/Private/AbilitySystemComponent_Abilities.cpp
--- a/Engine/Source/Runtime/GameplayAbilities/Private/AbilitySystemComponent_Abilities.cpp
+++ b/Engine/Source/Runtime/GameplayAbilities/Private/AbilitySystemComponent_Abilities.cpp
@@ -134,7 +134,7 @@ void UAbilitySystemComponent::CancelAbilitiesWithTags(const FGameplayTagContaine
 #if !(UE_BUILD_SHIPPING || UE_BUILD_TEST)
 int32 DenyClientActivation = 0;
 static FAutoConsoleVariableRef CVarDenyClientActivation(
-	TEXT("AbilitySystem.DenyClientActivations"),
+TEXT("AbilitySystem.DenyClientActivations"),
 	DenyClientActivation,
 	TEXT("Make server deny the next X ability activations from clients. For testing misprediction."),
 	ECVF_Default
@@ -222,9 +222,12 @@ void UAbilitySystemComponent::ClientActivateAbilitySucceed_Implementation(UGamep
 	AActor * OwnerActor = GetOwner();
 	check(OwnerActor);
 
-	FGameplayAbilityActivationInfo	ActivationInfo(EGameplayAbilityActivationMode::Confirmed, PredictionKey);
-
-	AbilityToActivate->CallActivateAbility(AbilityActorInfo.Get(), ActivationInfo);
+	// Fixme: We need a better way to link up/reconcile preditive replicated abilities.
+	if (AbilityToActivate->GetInstancingPolicy() != EGameplayAbilityInstancingPolicy::InstancedPerExecution)
+	{
+		FGameplayAbilityActivationInfo	ActivationInfo(EGameplayAbilityActivationMode::Confirmed, PredictionKey);
+		AbilityToActivate->CallActivateAbility(AbilityActorInfo.Get(), ActivationInfo);
+	}
 }
 
 void UAbilitySystemComponent::MontageBranchPoint_AbilityDecisionStop()
@@ -252,5 +255,107 @@ void UAbilitySystemComponent::MontageBranchPoint_AbilityDecisionStart()
 	}
 }
 
+bool UAbilitySystemComponent::GetUserAbilityActivationInhibited() const
+{
+	return UserAbilityActivationInhibited;
+}
+
+void UAbilitySystemComponent::SetUserAbilityActivationInhibited(bool NewInhibit)
+{
+	if(AbilityActorInfo->IsLocallyControlled())
+	{
+		if (NewInhibit && UserAbilityActivationInhibited)
+		{
+			// This could cause problems if two sources try to inhibit ability activation, it is not clear when the ability should be uninhibited
+			ABILITY_LOG(Warning, TEXT("Call to SetUserAbilityActivationInhibited(true) when UserAbilityActivationInhibited was already true"));
+		}
+
+		UserAbilityActivationInhibited = NewInhibit;
+	}
+}
+
+// --------------------------------------------------------------------------
+
+void UAbilitySystemComponent::BindToInputComponent(UInputComponent *InputComponent)
+{
+	static const FName ConfirmBindName(TEXT("AbilityConfirm"));
+	static const FName CancelBindName(TEXT("AbilityCancel"));
+
+	// Pressed event
+	{
+		FInputActionBinding AB(ConfirmBindName, IE_Pressed);
+		AB.ActionDelegate.GetDelegateForManualSet().BindUObject(this, &UAbilitySystemComponent::InputConfirm);
+		InputComponent->AddActionBinding(AB);
+	}
+
+	// 
+	{
+		FInputActionBinding AB(CancelBindName, IE_Pressed);
+		AB.ActionDelegate.GetDelegateForManualSet().BindUObject(this, &UAbilitySystemComponent::InputCancel);
+		InputComponent->AddActionBinding(AB);
+	}
+}
+
+void UAbilitySystemComponent::InputConfirm()
+{
+	if (GetOwnerRole() != ROLE_Authority)
+	{
+		// Tell the server we confirmed input
+		// FIXME TODO: we don't need to do this everytime! just when there is an activae ability waiting for this!
+		// 
+		ServerSetReplicatedConfirm(true);
+	}
+	
+	ConfirmCallbacks.Broadcast();
+}
+
+void UAbilitySystemComponent::InputCancel()
+{
+	CancelCallbacks.Broadcast();
+}
+
+// --------------------------------------------------------------------------
+
+void UAbilitySystemComponent::ServerSetReplicatedConfirm_Implementation(bool Confirmed)
+{
+	ReplicatedConfirmAbility  = Confirmed;
+	ConfirmCallbacks.Broadcast();
+}
+
+bool UAbilitySystemComponent::ServerSetReplicatedConfirm_Validate(bool Confirmed)
+{
+	return true;
+}
+
+void UAbilitySystemComponent::ServerSetReplicatedTargetData_Implementation(FGameplayAbilityTargetDataHandle Confirmed)
+{
+	ReplicatedTargetData = Confirmed;
+	ReplicatedTargetDataDelegate.Broadcast(ReplicatedTargetData);
+}
+
+bool UAbilitySystemComponent::ServerSetReplicatedTargetData_Validate(FGameplayAbilityTargetDataHandle Confirmed)
+{
+	return true;
+}
+
+void UAbilitySystemComponent::SetTargetAbility(UGameplayAbility* NewTargetingAbility)
+{
+	TargetingAbility = NewTargetingAbility;
+
+	//ReplicatedConfirmAbility = false;
+	//ReplicatedTargetData.Clear();
+}
+
+void UAbilitySystemComponent::ConsumeAbilityConfirm()
+{
+	ReplicatedConfirmAbility = false;
+}
+
+void UAbilitySystemComponent::ConsumeAbilityTargetData()
+{
+	ReplicatedTargetData.Clear();
+}
+
+
 #undef LOCTEXT_NAMESPACE
 
diff --git a/Engine/Source/Runtime/GameplayAbilities/Private/AbilitySystemGlobals.cpp b/Engine/Source/Runtime/GameplayAbilities/Private/AbilitySystemGlobals.cpp
--- a/Engine/Source/Runtime/GameplayAbilities/Private/AbilitySystemGlobals.cpp
+++ b/Engine/Source/Runtime/GameplayAbilities/Private/AbilitySystemGlobals.cpp
@@ -2,7 +2,7 @@
 
 #include "AbilitySystemPrivatePCH.h"
 #include "AbilitySystemComponent.h"
-#include "GameplayAbilityTypes.h"
+#include "Abilities/GameplayAbilityTypes.h"
 
 UAbilitySystemGlobals::UAbilitySystemGlobals(const class FPostConstructInitializeProperties& PCIP)
 : Super(PCIP)
diff --git a/Engine/Source/Runtime/GameplayAbilities/Private/AbilitySystemPrivatePCH.h b/Engine/Source/Runtime/GameplayAbilities/Private/AbilitySystemPrivatePCH.h
--- a/Engine/Source/Runtime/GameplayAbilities/Private/AbilitySystemPrivatePCH.h
+++ b/Engine/Source/Runtime/GameplayAbilities/Private/AbilitySystemPrivatePCH.h
@@ -14,7 +14,7 @@
 #include "Engine.h"
 #include "Engine/DataTable.h"
 
-#include "GameplayAbilityTypes.h"
+#include "Abilities/GameplayAbilityTypes.h"
 
 #include "AbilitySystemStats.h"
 #include "GameplayAbilitiesModule.h"
diff --git a/Engine/Source/Runtime/GameplayAbilities/Private/GameplayAbilitySet.cpp b/Engine/Source/Runtime/GameplayAbilities/Private/GameplayAbilitySet.cpp
deleted file mode 100644
diff --git a/Engine/Source/Runtime/GameplayAbilities/Public/Abilities/GameplayAbility.h b/Engine/Source/Runtime/GameplayAbilities/Public/Abilities/GameplayAbility.h
--- a/Engine/Source/Runtime/GameplayAbilities/Public/Abilities/GameplayAbility.h
+++ b/Engine/Source/Runtime/GameplayAbilities/Public/Abilities/GameplayAbility.h
@@ -5,7 +5,7 @@
 #include "GameplayTagContainer.h"
 #include "GameplayTagAssetInterface.h"
 #include "Runtime/Engine/Classes/Animation/AnimInstance.h"
-#include "GameplayAbilityTypes.h"
+#include "Abilities/GameplayAbilityTypes.h"
 #include "GameplayAbility.generated.h"
 
 /**
@@ -61,6 +61,32 @@
  *
  *	
  */
+
+
+
+/***
+ 
+
+
+This is the basic flow of function calls in a networked environment:
+-Client sees if he thinks he can activate the ability.
+-Tells server to try and activate ability, predicts the ability immediately
+-Server runs through Try/Can/Call, tells client if it passed or failed.
+
+
+																			 [If predictive] **Instance Ability** -> CallPredictiveActivateAbility() 
+			  (Optional)                                                     &&
+Client:		InputPressed()  ->  TryActivateAbility() -> CanActivateAbility() -> ServerTryActivateAbility()            --                         --                      --                          --                         --                ClientActivateAbilitySucceed() -> CallActivateAbility()
+																										  \\                                                                                                                                     //
+Server:		     --                      --                    --                          --                ServerTryActivateAbility()	-> CanActivateAbility() -> **Instance Ability** -> CallActivateAbility() -> ClientActivateAbilitySucceed()
+																																								||
+																																								-> ClientActivateAbilityFailed()
+																																																\\
+																																											Client:               ClientActivateAbilityFailed() [Remove and prediction work]
+
+***/
+
+
 UCLASS()
 class GAMEPLAYABILITIES_API UGameplayAbility : public UDataAsset
 {
@@ -188,6 +214,12 @@ public:
 
 	// ----------------------------------------------------------------------------------------------------------------
 
+	/** Gets the current actor info bound to this ability - can only be called on instanced abilities. */
+	virtual const FGameplayAbilityActorInfo* GetCurrentActorInfo();
+
+	/** Gets the current activation info bound to this ability - can only be called on instanced abilities. */
+	virtual FGameplayAbilityActivationInfo GetCurrentActivationInfo();
+
 	virtual UWorld* GetWorld() const OVERRIDE
 	{
 		return GetOuter()->GetWorld();
@@ -252,4 +284,4 @@ protected:
 	bool Prediction;
 
 	friend class UAbilitySystemComponent;
-};
\ No newline at end of file
+};
diff --git a/Engine/Source/Runtime/GameplayAbilities/Public/Abilities/GameplayAbilityTargetActor.h b/Engine/Source/Runtime/GameplayAbilities/Public/Abilities/GameplayAbilityTargetActor.h
new file mode 100644
--- /dev/null
+++ b/Engine/Source/Runtime/GameplayAbilities/Public/Abilities/GameplayAbilityTargetActor.h
@@ -0,0 +1,30 @@
+// Copyright 1998-2014 Epic Games, Inc. All Rights Reserved.
+#pragma once
+
+#include "GameplayAbility.h"
+#include "GameplayAbilityTargetActor.generated.h"
+
+UCLASS(Blueprintable, abstract)
+class GAMEPLAYABILITIES_API AGameplayAbilityTargetActor : public AActor
+{
+	GENERATED_UCLASS_BODY()
+
+public:
+
+	/** Native classes can set this to call the non-blueprintable  */
+	UPROPERTY()
+	bool	StaticTargetFunction;
+
+	virtual FGameplayAbilityTargetDataHandle StaticGetTargetData(UWorld * World, const FGameplayAbilityActorInfo* ActorInfo, FGameplayAbilityActivationInfo ActivationInfo) { return FGameplayAbilityTargetDataHandle(); }
+
+	/** Otherwise, the actor will be instantiated and call the blueprintable function GetTargetData */
+	UFUNCTION(BlueprintImplementableEvent, Category = "Targeting")
+	FGameplayAbilityTargetDataHandle StaticGetTargetData(UWorld * World, FGameplayAbilityActorInfo ActorInfo, FGameplayAbilityActivationInfo ActivationInfo);
+
+	// temp
+	virtual void StartTargeting(UGameplayAbility* Ability) { };
+
+	// ------------------------------
+	
+	FAbilityTargetData	TargetDataReadyDelegate;
+};
\ No newline at end of file
diff --git a/Engine/Source/Runtime/GameplayAbilities/Public/Abilities/GameplayAbilityTargetActor_SingleLineTrace.h b/Engine/Source/Runtime/GameplayAbilities/Public/Abilities/GameplayAbilityTargetActor_SingleLineTrace.h
new file mode 100644
--- /dev/null
+++ b/Engine/Source/Runtime/GameplayAbilities/Public/Abilities/GameplayAbilityTargetActor_SingleLineTrace.h
@@ -0,0 +1,34 @@
+// Copyright 1998-2014 Epic Games, Inc. All Rights Reserved.
+
+#pragma once
+
+#include "GameplayAbilityTargetActor.h"
+#include "GameplayAbilityTargetActor_SingleLineTrace.generated.h"
+
+UCLASS(Blueprintable)
+class GAMEPLAYABILITIES_API AGameplayAbilityTargetActor_SingleLineTrace : public AGameplayAbilityTargetActor
+{
+	GENERATED_UCLASS_BODY()
+
+public:
+	
+	virtual FGameplayAbilityTargetDataHandle StaticGetTargetData(UWorld * World, const FGameplayAbilityActorInfo* ActorInfo, FGameplayAbilityActivationInfo ActivationInfo) override;
+
+	virtual void StartTargeting(UGameplayAbility* Ability);
+
+	//temp
+
+	void DoIt(TWeakObjectPtr<UGameplayAbility> Ability);
+
+	UFUNCTION()
+	void Confirm();
+
+	UFUNCTION()
+	void Cancel();
+
+	virtual void Tick(float DeltaSeconds) override;
+	
+	TWeakObjectPtr<UGameplayAbility> Ability;
+
+	bool bDebug;
+};
\ No newline at end of file
diff --git a/Engine/Source/Runtime/GameplayAbilities/Public/Abilities/GameplayAbilityTypes.h b/Engine/Source/Runtime/GameplayAbilities/Public/Abilities/GameplayAbilityTypes.h
new file mode 100644
--- /dev/null
+++ b/Engine/Source/Runtime/GameplayAbilities/Public/Abilities/GameplayAbilityTypes.h
@@ -0,0 +1,418 @@
+// Copyright 1998-2014 Epic Games, Inc. All Rights Reserved.
+
+#pragma once
+
+#include "GameplayAbilityTypes.generated.h"
+
+class UGameplayEffect;
+class UAnimInstance;
+class UAbilitySystemComponent;
+
+/** 
+ *
+ *
+ */
+
+UENUM(BlueprintType)
+namespace EGameplayAbilityInstancingPolicy
+{
+	/**
+	*	How the ability is instanced when executed. This limits what an ability can do in its implementation. For example, a NonInstanced
+	*	Ability cannot have state. It is probably unsafe for an InstancedPerActor ability to have latent actions, etc.
+	*/
+
+	enum Type
+	{
+		NonInstanced,			// This ability is never instanced. Anything that executes the ability is operating on the CDO.
+		InstancedPerActor,		// Each actor gets their own instance of this ability. State can be saved, replication is possible.
+		InstancedPerExecution,	// We instance this ability each time it is executed. Replication possible but not recommended.
+	};
+}
+
+UENUM(BlueprintType)
+namespace EGameplayAbilityNetExecutionPolicy
+{
+	/**
+	*	How does an ability execute on the network. Does a client "ask and predict", "ask and wait", "don't ask"
+	*/
+
+	enum Type
+	{
+		Predictive		UMETA(DisplayName = "Predictive"),	// Part of this ability runs predictively on the client.
+		Server			UMETA(DisplayName = "Server"),		// This ability must be OK'd by the server before doing anything on a client.
+		Client			UMETA(DisplayName = "Client"),		// This ability runs as long the client says it does.
+	};
+}
+
+UENUM(BlueprintType)
+namespace EGameplayAbilityReplicationPolicy
+{
+	/**
+	*	How an ability replicates state/events to everyone on the network.
+	*/
+
+	enum Type
+	{
+		ReplicateNone		UMETA(DisplayName = "Replicate None"),	// We don't replicate the instance of the ability to anyone.
+		ReplicateAll		UMETA(DisplayName = "Replicate All"),	// We replicate the instance of the ability to everyone (even simulating clients).
+		ReplicateOwner		UMETA(DisplayName = "Replicate Owner"),	// Only replicate the instance of the ability to the owner.
+	};
+}
+
+USTRUCT(BlueprintType)
+struct FGameplayAbilityHandle
+{
+	GENERATED_USTRUCT_BODY()
+
+	FGameplayAbilityHandle()
+	: Handle(INDEX_NONE)
+	{
+
+	}
+
+	FGameplayAbilityHandle(int32 InHandle)
+		: Handle(InHandle)
+	{
+
+	}
+
+	bool IsValid() const
+	{
+		return Handle != INDEX_NONE;
+	}
+
+	FGameplayAbilityHandle GetNextHandle()
+	{
+		return FGameplayAbilityHandle(Handle + 1);
+	}
+
+	bool operator==(const FGameplayAbilityHandle& Other) const
+	{
+		return Handle == Other.Handle;
+	}
+
+	bool operator!=(const FGameplayAbilityHandle& Other) const
+	{
+		return Handle != Other.Handle;
+	}
+
+	FString ToString() const
+	{
+		return FString::Printf(TEXT("%d"), Handle);
+	}
+
+private:
+
+	UPROPERTY()
+	int32 Handle;
+};
+
+/**
+*	Not implemented yet, but we will need something like this to track ability + levels if we choose to support ability leveling in the base classes.
+*/
+
+USTRUCT()
+struct FGameplayAbilitySpec
+{
+	GENERATED_USTRUCT_BODY()
+
+	UPROPERTY()
+	int32	Level;
+
+	UPROPERTY()
+	class UGameplayAbility * Ability;
+
+	UPROPERTY()
+	FGameplayAbilityHandle	Handle;
+};
+
+
+UENUM(BlueprintType)
+namespace EGameplayAbilityActivationMode
+{
+	enum Type
+	{
+		Authority,			// We are the authority activating this ability
+		NonAuthority,		// We are not the authority but aren't predicting yet. This is a mostly invalid state to be in.
+
+		Predicting,			// We are predicting the activation of this ability
+		Confirmed,			// We are not the authority, but the authority has confirmed this activation
+	};
+}
+
+/**
+*	FGameplayAbilityActorInfo
+*
+*	Cached data associated with an Actor using an Ability.
+*		-Initialized from an AActor* in InitFromActor
+*		-Abilities use this to know what to actor upon. E.g., instead of being coupled to a specific actor class.
+*		-These are generally passed around as pointers to support polymorphism.
+*		-Projects can override UAbilitySystemGlobals::AllocAbilityActorInfo to override the default struct type that is created.
+*
+*/
+USTRUCT(BlueprintType)
+struct GAMEPLAYABILITIES_API FGameplayAbilityActorInfo
+{
+	GENERATED_USTRUCT_BODY()
+
+	virtual ~FGameplayAbilityActorInfo() {}
+
+	UPROPERTY(BlueprintReadOnly, Category = "ActorInfo")
+	TWeakObjectPtr<AActor>	Actor;
+
+	/** PlayerController associated with this actor. This will often be null! */
+	UPROPERTY(BlueprintReadOnly, Category = "ActorInfo")
+	TWeakObjectPtr<APlayerController>	PlayerController;
+
+	UPROPERTY(BlueprintReadOnly, Category = "ActorInfo")
+	TWeakObjectPtr<UAnimInstance>	AnimInstance;
+
+	UPROPERTY(BlueprintReadOnly, Category = "ActorInfo")
+	TWeakObjectPtr<UAbilitySystemComponent>	AbilitySystemComponent;
+
+	UPROPERTY(BlueprintReadOnly, Category = "ActorInfo")
+	TWeakObjectPtr<UMovementComponent>	MovementComponent;
+
+	bool IsLocallyControlled() const;
+
+	virtual void InitFromActor(AActor *Actor);
+};
+
+/**
+*	FGameplayAbilityActivationInfo
+*
+*	Data tied to a specific activation of an ability.
+*		-Tell us whether we are the authority, if we are predicting, confirmed, etc.
+*		-Holds PredictionKey
+*		-Generally not meant to be subclassed in projects.
+*		-Passed around by value since the struct is small.
+*/
+
+USTRUCT(BlueprintType)
+struct GAMEPLAYABILITIES_API FGameplayAbilityActivationInfo
+{
+	GENERATED_USTRUCT_BODY()
+
+	FGameplayAbilityActivationInfo()
+	: ActivationMode(EGameplayAbilityActivationMode::Authority),
+	PredictionKey(0)
+	{
+
+	}
+
+	FGameplayAbilityActivationInfo(AActor * InActor, int32 InPredictionKey)
+		: PredictionKey(InPredictionKey)
+	{
+		// On Init, we are either Authority or NonAuthority. We haven't been given a PredictionKey and we haven't been confirmed.
+		// NonAuthority essentially means 'I'm not sure what how I'm going to do this yet'.
+		ActivationMode = (InActor->Role == ROLE_Authority ? EGameplayAbilityActivationMode::Authority : EGameplayAbilityActivationMode::NonAuthority);
+	}
+
+	FGameplayAbilityActivationInfo(EGameplayAbilityActivationMode::Type InType, int32 InPredictionKey)
+		: ActivationMode(InType)
+		, PredictionKey(InPredictionKey)
+	{
+	}
+
+
+	void GeneratePredictionKey(UAbilitySystemComponent * Component) const;
+
+	void SetPredictionKey(int32 InPredictionKey);
+
+	void SetActivationConfirmed();
+
+	UPROPERTY(BlueprintReadOnly, Category = "ActorInfo")
+	mutable TEnumAsByte<EGameplayAbilityActivationMode::Type>	ActivationMode;
+
+	UPROPERTY()
+	mutable int32 PredictionKey;
+};
+
+
+// ----------------------------------------------------
+
+/**
+*	A generic structure for targeting data. We want generic functions to produce this data and other generic
+*	functions to consume this data.
+*
+*	We expect this to be able to hold specific actors/object reference and also generic location/direction/origin
+*	information.
+*
+*	Some example producers:
+*		-Overlap/Hit collision event generates TargetData about who was hit in a melee attack
+*		-A mouse input causes a hit trace and the actor infront of the crosshair is turned into TargetData
+*		-A mouse input causes TargetData to be generated from the owner's crosshair view origin/direction
+*		-An AOE/aura pulses and all actors in a radius around the instigator are added to TargetData
+*		-Panzer Dragoon style 'painting' targeting mode
+*		-MMORPG style ground AOE targeting style (potentially both a location on the ground and actors that were targeted)
+*
+*	Some example consumers:
+*		-Apply a GameplayEffect to all actors in TargetData
+*		-Find closest actor from all in TargetData
+*		-Call some function on all actors in TargetData
+*		-Filter or merge TargetDatas
+*		-Spawn a new actor at a TargetData location
+*
+*
+*
+*	Maybe it is better to distinguish between actor list targeting vs positional targeting data?
+*		-AOE/aura type of targeting data blurs the line
+*
+*
+*/
+
+USTRUCT()
+struct GAMEPLAYABILITIES_API FGameplayAbilityTargetData
+{
+	GENERATED_USTRUCT_BODY()
+
+	virtual ~FGameplayAbilityTargetData() { }
+
+	void ApplyGameplayEffect(UGameplayEffect *GameplayEffect, const FGameplayAbilityActorInfo InstigatorInfo);
+
+	virtual TArray<AActor*>	GetActors() const
+	{
+		return TArray<AActor*>();
+	}
+
+	virtual bool HasHitResult() const
+	{
+		return false;
+	}
+
+	virtual const FHitResult * GetHitResult() const
+	{
+		return NULL;
+	}
+
+	virtual UScriptStruct * GetScriptStruct()
+	{
+		return FGameplayAbilityTargetData::StaticStruct();
+	}
+
+	virtual FString ToString() const;
+};
+
+USTRUCT()
+struct GAMEPLAYABILITIES_API FGameplayAbilityTargetData_SingleTargetHit : public FGameplayAbilityTargetData
+{
+	GENERATED_USTRUCT_BODY()
+
+	FGameplayAbilityTargetData_SingleTargetHit()
+	{ }
+
+	FGameplayAbilityTargetData_SingleTargetHit(const FHitResult InHitResult)
+	: HitResult(InHitResult)
+	{ }
+
+	virtual TArray<AActor*>	GetActors() const
+	{
+		TArray<AActor*>	Actors;
+		if (HitResult.Actor.IsValid())
+		{
+			Actors.Push(HitResult.Actor.Get());
+		}
+		return Actors;
+	}
+
+	virtual bool HasHitResult() const
+	{
+		return true;
+	}
+
+	virtual const FHitResult * GetHitResult() const
+	{
+		return &HitResult;
+	}
+
+	UPROPERTY()
+	FHitResult	HitResult;
+
+	bool NetSerialize(FArchive& Ar, class UPackageMap* Map, bool& bOutSuccess);
+
+	virtual UScriptStruct * GetScriptStruct()
+	{
+		return FGameplayAbilityTargetData_SingleTargetHit::StaticStruct();
+	}
+};
+
+template<>
+struct TStructOpsTypeTraits<FGameplayAbilityTargetData_SingleTargetHit> : public TStructOpsTypeTraitsBase
+{
+	enum
+	{
+		WithNetSerializer = true	// Fow now this is REQUIRED for FGameplayAbilityTargetDataHandle net serialization to work
+	};
+};
+
+
+/**
+*	Handle for Targeting Data. This servers two main purposes:
+*		-Avoid us having to copy around the full targeting data structure in Blueprints
+*		-Allows us to leverage polymorphism in the target data structure
+*		-Allows us to implement NetSerialize and replicate by value between clients/server
+*
+*		-Avoid using UObjects could be used to give us polymorphism and by reference passing in blueprints.
+*		-However we would still be screwed when it came to replication
+*
+*		-Replication by value
+*		-Pass by reference in blueprints
+*		-Polymophism in TargetData structure
+*
+*/
+
+USTRUCT(BlueprintType)
+struct GAMEPLAYABILITIES_API FGameplayAbilityTargetDataHandle
+{
+	GENERATED_USTRUCT_BODY()
+
+	FGameplayAbilityTargetDataHandle() { }
+	FGameplayAbilityTargetDataHandle(FGameplayAbilityTargetData* DataPtr)
+		: Data(DataPtr)
+	{
+		
+	}
+
+	TSharedPtr<FGameplayAbilityTargetData>	Data;
+
+	void Clear()
+	{
+		Data.Reset();
+	}
+
+	bool IsValid() const
+	{
+		return Data.IsValid();
+	}
+
+	bool NetSerialize(FArchive& Ar, class UPackageMap* Map, bool& bOutSuccess);
+
+	/** Comparison operator */
+	bool operator==(FGameplayAbilityTargetDataHandle const& Other) const
+	{
+		// Both invalid structs or both valid and Pointer compare (???) // deep comparison equality
+		bool bBothValid = IsValid() && Other.IsValid();
+		bool bBothInvalid = !IsValid() && !Other.IsValid();
+		return (bBothInvalid || (bBothValid && (Data.Get() == Other.Data.Get())));
+	}
+
+	/** Comparison operator */
+	bool operator!=(FGameplayAbilityTargetDataHandle const& Other) const
+	{
+		return !(FGameplayAbilityTargetDataHandle::operator==(Other));
+	}
+};
+
+template<>
+struct TStructOpsTypeTraits<FGameplayAbilityTargetDataHandle> : public TStructOpsTypeTraitsBase
+{
+	enum
+	{
+		WithCopy = true,		// Necessary so that TSharedPtr<FGameplayAbilityTargetData> Data is copied around
+		WithNetSerializer = true,
+		WithIdenticalViaEquality = true,
+	};
+};
+
+/** Generic callback for returning when target data is available */
+DECLARE_MULTICAST_DELEGATE_OneParam(FAbilityTargetData, FGameplayAbilityTargetDataHandle);
\ No newline at end of file
diff --git a/Engine/Source/Runtime/GameplayAbilities/Public/Abilities/GameplayAbility_Instanced.h b/Engine/Source/Runtime/GameplayAbilities/Public/Abilities/GameplayAbility_Instanced.h
--- a/Engine/Source/Runtime/GameplayAbilities/Public/Abilities/GameplayAbility_Instanced.h
+++ b/Engine/Source/Runtime/GameplayAbilities/Public/Abilities/GameplayAbility_Instanced.h
@@ -64,11 +64,9 @@ public:
 
 	UFUNCTION(BlueprintCallable, Category = Ability)
 	virtual void K2_EndAbility();
-	
 
 	// -------------------------------------------------
 
-
 	virtual EGameplayAbilityInstancingPolicy::Type GetInstancingPolicy() const OVERRIDE
 	{
 		return InstancedPerExecution ? EGameplayAbilityInstancingPolicy::InstancedPerExecution : EGameplayAbilityInstancingPolicy::InstancedPerActor;
@@ -92,20 +90,30 @@ public:
 
 	// -------------------------------------------------
 
+	const FGameplayAbilityActorInfo* GetCurrentActorInfo() override
+	{
+		return CurrentActorInfo;
+	}
+
+	FGameplayAbilityActivationInfo GetCurrentActivationInfo() override
+	{
+		return CurrentActivationInfo;
+	}
+
+protected:
+
+	virtual void EndAbility(const FGameplayAbilityActorInfo* ActorInfo) OVERRIDE;
 
+	/** This is information specific to this instance of the ability. E.g, whether it is predicting, authorting, confirmed, etc. */
+	UPROPERTY(BlueprintReadOnly, Category = Ability)
+	FGameplayAbilityActivationInfo	CurrentActivationInfo;
+
+private:
 	/** This is shared, cached information about the thing using us
 	 *		E.g, Actor*, MovementComponent*, AnimInstance, etc.
 	 *		This is hopefully allocated once per actor and shared by many abilities.
 	 *		The actual struct may be overridden per game to include game specific data.
 	 *		(E.g, child classes may want to cast to FMyGameAbilityActorInfo)
 	 */
-	mutable const FGameplayAbilityActorInfo * CurrentActorInfo;
-
-	/** This is information specific to this instance of the ability. E.g, whether it is predicting, authorting, confirmed, etc. */
-	UPROPERTY(BlueprintReadOnly, Category=Ability)
-	FGameplayAbilityActivationInfo	CurrentActivationInfo;
-
-protected:
-
-	virtual void EndAbility(const FGameplayAbilityActorInfo* ActorInfo) OVERRIDE;
+	mutable const FGameplayAbilityActorInfo* CurrentActorInfo;
 };
\ No newline at end of file
diff --git a/Engine/Source/Runtime/GameplayAbilities/Public/Abilities/GameplayAbility_Montage.h b/Engine/Source/Runtime/GameplayAbilities/Public/Abilities/GameplayAbility_Montage.h
--- a/Engine/Source/Runtime/GameplayAbilities/Public/Abilities/GameplayAbility_Montage.h
+++ b/Engine/Source/Runtime/GameplayAbilities/Public/Abilities/GameplayAbility_Montage.h
@@ -2,7 +2,7 @@
 
 #pragma once
 
-#include "Abilities/GameplayAbility.h"
+#include "GameplayAbility.h"
 #include "GameplayAbility_Montage.generated.h"
 
 /**
diff --git a/Engine/Source/Runtime/GameplayAbilities/Public/Abilities/Tasks/AbilityTask.h b/Engine/Source/Runtime/GameplayAbilities/Public/Abilities/Tasks/AbilityTask.h
--- a/Engine/Source/Runtime/GameplayAbilities/Public/Abilities/Tasks/AbilityTask.h
+++ b/Engine/Source/Runtime/GameplayAbilities/Public/Abilities/Tasks/AbilityTask.h
@@ -1,12 +1,16 @@
 // Copyright 1998-2014 Epic Games, Inc. All Rights Reserved.
 #pragma once
+#include "GameplayAbility_Instanced.h"
 #include "AbilityTask.Generated.h"
 
 UCLASS(MinimalAPI)
 class UAbilityTask : public UObject
 {
-	GENERATED_UCLASS_BODY()	
+	GENERATED_UCLASS_BODY()
 
-public:
+	// Called to trigger the actual task once the delegates have been set up
+	UFUNCTION(BlueprintCallable, meta = (BlueprintInternalUseOnly = "true"), Category = "Abilities")
+	virtual void Activate();
 	
+	TWeakObjectPtr<UGameplayAbility> Ability;
 };
\ No newline at end of file
diff --git a/Engine/Source/Runtime/GameplayAbilities/Public/Abilities/Tasks/AbilityTask_WaitOverlap.h b/Engine/Source/Runtime/GameplayAbilities/Public/Abilities/Tasks/AbilityTask_WaitOverlap.h
--- a/Engine/Source/Runtime/GameplayAbilities/Public/Abilities/Tasks/AbilityTask_WaitOverlap.h
+++ b/Engine/Source/Runtime/GameplayAbilities/Public/Abilities/Tasks/AbilityTask_WaitOverlap.h
@@ -1,7 +1,7 @@
 // Copyright 1998-2014 Epic Games, Inc. All Rights Reserved.
 #pragma once
 #include "AbilityTask.h"
-#include "GameplayAbilityTypes.h"
+#include "Abilities/GameplayAbilityTypes.h"
 #include "AbilityTask_WaitOverlap.generated.h"
 
 DECLARE_DYNAMIC_MULTICAST_DELEGATE_OneParam(FWaitOverlapDelegate, FGameplayAbilityTargetDataHandle, TargetData);
diff --git a/Engine/Source/Runtime/GameplayAbilities/Public/AbilitySystemBlueprintLibrary.h b/Engine/Source/Runtime/GameplayAbilities/Public/AbilitySystemBlueprintLibrary.h
--- a/Engine/Source/Runtime/GameplayAbilities/Public/AbilitySystemBlueprintLibrary.h
+++ b/Engine/Source/Runtime/GameplayAbilities/Public/AbilitySystemBlueprintLibrary.h
@@ -3,13 +3,13 @@
 #pragma once
 
 #include "Abilities/GameplayAbility.h"
-#include "GameplayAbilityTypes.h"
 #include "AbilitySystemBlueprintLibrary.generated.h"
 
 class UAnimMontage;
 class UAbilityTask_PlayMontageAndWait;
 class UAbilityTask_WaitMovementModeChange;
 class UAbilityTask_WaitOverlap;
+class UAbilityTask_WaitConfirmCancel;
 
 UCLASS(MinimalAPI)
 class UAbilitySystemBlueprintLibrary : public UBlueprintFunctionLibrary
@@ -22,6 +22,12 @@ class UAbilitySystemBlueprintLibrary : public UBlueprintFunctionLibrary
 	UFUNCTION(BlueprintCallable, Category=Ability)
 	static void ApplyGameplayEffectToTargetData(FGameplayAbilityTargetDataHandle Target, UGameplayEffect *GameplayEffect, const FGameplayAbilityActorInfo InstigatorInfo);
 
+	UFUNCTION(BlueprintCallable, Category = Ability)
+	static FGameplayAbilityTargetDataHandle	AbilityTargetDataFromHitResult(FHitResult HitResult);
+
+	UFUNCTION(BlueprintCallable, Category = Ability)
+	static FHitResult GetHitResultFromTargetData(FGameplayAbilityTargetDataHandle HitResult);
+
 	// -------------------------------------------------------------------------------
 
 	UFUNCTION(BlueprintCallable, meta=(HidePin="WorldContextObject", DefaultToSelf="WorldContextObject", BlueprintInternalUseOnly = "TRUE"))
diff --git a/Engine/Source/Runtime/GameplayAbilities/Public/AbilitySystemComponent.h b/Engine/Source/Runtime/GameplayAbilities/Public/AbilitySystemComponent.h
--- a/Engine/Source/Runtime/GameplayAbilities/Public/AbilitySystemComponent.h
+++ b/Engine/Source/Runtime/GameplayAbilities/Public/AbilitySystemComponent.h
@@ -2,13 +2,17 @@
 #pragma once
 
 #include "GameplayEffect.h"
-#include "GameplayAbilityTypes.h"
+#include "Abilities/GameplayAbilityTypes.h"
 #include "AbilitySystemComponent.generated.h"
 
 GAMEPLAYABILITIES_API DECLARE_LOG_CATEGORY_EXTERN(LogAbilitySystemComponent, Log, All);
 
+/** Used for cleaning up predicted data on network clients */
 DECLARE_MULTICAST_DELEGATE(FAbilitySystemComponentPredictionKeyClear);
 
+/** Used to register callbacks to confirm/cancel input */
+DECLARE_DYNAMIC_MULTICAST_DELEGATE(FAbilityConfirmOrCancel);
+
 USTRUCT()
 struct GAMEPLAYABILITIES_API FAttributeDefaults
 {
@@ -156,21 +160,15 @@ class GAMEPLAYABILITIES_API UAbilitySystemComponent : public UActorComponent
 	UFUNCTION()
 	void OnRep_PredictionKey();
 
-	TArray<TPair<int32, FAbilitySystemComponentPredictionKeyClear> >	PredictionDelegates;
-	
+	TArray<TPair<int32, FAbilitySystemComponentPredictionKeyClear> > PredictionDelegates;
 
 	FAbilitySystemComponentPredictionKeyClear &	GetPredictionKeyDelegate(int32 PredictionKey);
-
-	/*
-	virtual void OnSubobjectCreatedFromReplication(UObject *NewSubobject) OVERRIDE;
 	
-	virtual void OnSubobjectDestroyFromReplication(UObject *NewSubobject) OVERRIDE;
-	*/
 
 	// ----------------------------------------------------------------------------------------------------------------
 	//
-	//	GameplayEffects
-	//	(maybe should go in a different component?)
+	//	GameplayEffects	
+	//	
 	// ----------------------------------------------------------------------------------------------------------------
 
 	// --------------------------------------------
@@ -272,17 +270,17 @@ class GAMEPLAYABILITIES_API UAbilitySystemComponent : public UActorComponent
 	 *	
 	 */
 
-	class UGameplayAbility * CreateNewInstanceOfAbility(UGameplayAbility* Ability);
+	UGameplayAbility* CreateNewInstanceOfAbility(UGameplayAbility* Ability);
 
 	void CancelAbilitiesWithTags(const FGameplayTagContainer Tags, const FGameplayAbilityActorInfo* ActorInfo, const FGameplayAbilityActivationInfo ActivationInfo, UGameplayAbility* Ignore);
 	
 	/** References to non-replicating abilities that we instanced. We need to keep these references to avoid GC */
 	UPROPERTY()
-	TArray<class UGameplayAbility *>	NonReplicatedInstancedAbilities;
+	TArray<UGameplayAbility*>	NonReplicatedInstancedAbilities;
 
 	/** References to replicating abilities that we instanced. We need to keep these references to avoid GC */
 	UPROPERTY(Replicated)
-	TArray<class UGameplayAbility *>	ReplicatedInstancedAbilities;
+	TArray<UGameplayAbility*>	ReplicatedInstancedAbilities;
 
 	/**
 	 *	The abilities we can activate. 
@@ -294,10 +292,10 @@ class GAMEPLAYABILITIES_API UAbilitySystemComponent : public UActorComponent
 	 *	instancing or anything else that the AbilitySystemComponent would provide, then it doesn't need the component to function.
 	 */ 
 	UPROPERTY(Replicated, BlueprintReadOnly, Category = "Abilities")
-	TArray<class UGameplayAbility *>	ActivatableAbilities;
+	TArray<UGameplayAbility*>	ActivatableAbilities;
 	
 
-	UFUNCTION(Server, WithValidation, Reliable)
+	UFUNCTION(Server, reliable, WithValidation)
 	void	ServerTryActivateAbility(class UGameplayAbility* AbilityToActivate, int32 PredictionKey);
 
 	UFUNCTION(Client, Reliable)
@@ -308,18 +306,41 @@ class GAMEPLAYABILITIES_API UAbilitySystemComponent : public UActorComponent
 
 	// ----------------------------------------------------------------------------------------------------------------
 
-	/** This is probably temp for testing. We want to think a bit more on the API for outside stuff activating abilities */
+	// This is meant to be used to inhibit activating an ability from an input perspective. (E.g., the menu is pulled up, another game mechanism is consuming all input, etc)
+	// This should only be called on locally owned players.
+	// This should not be used to game mechanics like silences or disables. Those should be done through gameplay effects.
+
+	UFUNCTION(BlueprintCallable, Category="Abilities")
+	bool	GetUserAbilityActivationInhibited() const;
+	
+	/** Disable or Enable a local user from being able to activate abilities. This should only be used for input/UI etc related inhibition. Do not use for game mechanics. */
+	UFUNCTION(BlueprintCallable, Category="Abilities")
+	void	SetUserAbilityActivationInhibited(bool NewInhibit);
+
+	bool	UserAbilityActivationInhibited;
+
+	// ----------------------------------------------------------------------------------------------------------------
+
+	virtual void BindToInputComponent(UInputComponent *InputComponent);
+
+	void InputConfirm();
+	void InputCancel();
+
+	FAbilityConfirmOrCancel	ConfirmCallbacks;
+	FAbilityConfirmOrCancel	CancelCallbacks;
+
+	// ----------------------------------------------------------------------------------------------------------------
+
+	/** This is temp for testing. We want to think a bit more on the API for outside stuff activating abilities */
 	bool	ActivateAbility(TWeakObjectPtr<UGameplayAbility> Ability);
 
 	/** There needs to be a concept of an animating ability. Only one may exist at a time. New requests can be queued up, overridden, or ignored. */
 	UPROPERTY()
-	class UGameplayAbility *	AnimatingAbility;
+	UGameplayAbility*	AnimatingAbility;
 
-	void	MontageBranchPoint_AbilityDecisionStop();
+	void MontageBranchPoint_AbilityDecisionStop();
 
-	void	MontageBranchPoint_AbilityDecisionStart();
-
-	
+	void MontageBranchPoint_AbilityDecisionStart();
 
 	// -----------------------------------------------------------------------------
 
@@ -328,6 +349,34 @@ class GAMEPLAYABILITIES_API UAbilitySystemComponent : public UActorComponent
 
 	void InitAbilityActorInfo();
 
+	// -----------------------------------------------------------------------------
+
+	/**
+	 *	While these appear to be state, this is actually synchronization events w/ some payload data
+	 */
+
+	UFUNCTION(Server, reliable, WithValidation)
+	void ServerSetReplicatedConfirm(bool Confirmed);
+
+	UFUNCTION(Server, reliable, WithValidation)
+	void ServerSetReplicatedTargetData(FGameplayAbilityTargetDataHandle ReplicatedTargetData);
+
+	void SetTargetAbility(UGameplayAbility* NewTargetingAbility);
+
+	void ConsumeAbilityConfirm();
+
+	void ConsumeAbilityTargetData();
+
+	/** This ability has a 'lock' on replicated targeting data. Only 1 targeting ability can be active at once */
+	UPROPERTY()
+	UGameplayAbility* TargetingAbility;
+
+	bool ReplicatedConfirmAbility;
+
+	FGameplayAbilityTargetDataHandle ReplicatedTargetData;
+
+	FAbilityTargetData	ReplicatedTargetDataDelegate;
+
 private:
 
 	// --------------------------------------------
diff --git a/Engine/Source/Runtime/GameplayAbilities/Public/GameplayAbilityTypes.h b/Engine/Source/Runtime/GameplayAbilities/Public/GameplayAbilityTypes.h
deleted file mode 100644
