Missed files

diff --git a/Engine/Source/Runtime/GameplayAbilities/Public/Abilities/Tasks/AbilityTask_WaitConfirmCancel.h b/Engine/Source/Runtime/GameplayAbilities/Public/Abilities/Tasks/AbilityTask_WaitConfirmCancel.h
new file mode 100644
--- /dev/null
+++ b/Engine/Source/Runtime/GameplayAbilities/Public/Abilities/Tasks/AbilityTask_WaitConfirmCancel.h
@@ -0,0 +1,37 @@
+// Copyright 1998-2014 Epic Games, Inc. All Rights Reserved.
+#pragma once
+
+#include "AbilityTask.h"
+#include "GameplayAbility_Instanced.h"
+#include "AbilityTask_WaitConfirmCancel.generated.h"
+
+DECLARE_DYNAMIC_MULTICAST_DELEGATE_OneParam(FWaitConfirmCancelDelegate, bool, Test);
+
+UCLASS(MinimalAPI)
+class UAbilityTask_WaitConfirmCancel : public UAbilityTask
+{
+	GENERATED_UCLASS_BODY()
+
+	UPROPERTY(BlueprintAssignable)
+	FWaitConfirmCancelDelegate	OnConfirm;
+
+	UPROPERTY(BlueprintAssignable)
+	FWaitConfirmCancelDelegate	OnCancel;
+	
+	UFUNCTION()
+	void OnConfirmCallback();
+
+	UFUNCTION()
+	void OnCancelCallback();
+
+	UFUNCTION(BlueprintCallable, meta=(HidePin = "WorldContextObject", DefaultToSelf = "WorldContextObject", BlueprintInternalUseOnly = "true", FriendlyName="Wait for Confirm Input"), Category="Abilities")
+	static UAbilityTask_WaitConfirmCancel* WaitConfirmCancel(UObject* WorldContextObject);	
+
+	virtual void Activate() OVERRIDE;
+
+	UPROPERTY()
+	UAbilitySystemComponent* ASC;
+
+public:
+	
+};
\ No newline at end of file
diff --git a/Engine/Source/Runtime/GameplayAbilities/Public/Abilities/Tasks/AbilityTask_WaitTargetData.h b/Engine/Source/Runtime/GameplayAbilities/Public/Abilities/Tasks/AbilityTask_WaitTargetData.h
new file mode 100644
--- /dev/null
+++ b/Engine/Source/Runtime/GameplayAbilities/Public/Abilities/Tasks/AbilityTask_WaitTargetData.h
@@ -0,0 +1,38 @@
+// Copyright 1998-2014 Epic Games, Inc. All Rights Reserved.
+#pragma once
+
+#include "AbilityTask.h"
+#include "AbilityTask_WaitTargetData.generated.h"
+
+class AGameplayAbilityTargetActor;
+
+DECLARE_DYNAMIC_MULTICAST_DELEGATE_OneParam(FWaitTargetDataDelegate, FGameplayAbilityTargetDataHandle, Data);
+
+UCLASS(MinimalAPI)
+class UAbilityTask_WaitTargetData: public UAbilityTask
+{
+	GENERATED_UCLASS_BODY()
+
+	UPROPERTY(BlueprintAssignable)
+	FWaitTargetDataDelegate	ValidData;
+
+	UPROPERTY(BlueprintAssignable)
+	FWaitTargetDataDelegate	Error;
+
+	UFUNCTION()
+	void OnTargetDataReplicatedCallback(FGameplayAbilityTargetDataHandle Data);
+
+	UFUNCTION()
+	void OnTargetDataReadyCallback(FGameplayAbilityTargetDataHandle Data);
+	
+	
+	UFUNCTION(BlueprintCallable, meta=(HidePin = "WorldContextObject", DefaultToSelf = "WorldContextObject", BlueprintInternalUseOnly = "true"), Category = "Abilities")
+	static UAbilityTask_WaitTargetData* WaitTargetData(UObject* WorldContextObject, TSubclassOf<AGameplayAbilityTargetActor> TargetClass);
+
+	virtual void Activate() OVERRIDE;
+
+	TSubclassOf<AGameplayAbilityTargetActor> TargetClass;
+
+	UPROPERTY()
+	UAbilitySystemComponent* ASC;
+};
\ No newline at end of file
