remove GameplayAbility_Instanced class - will use one base class now

diff --git a/Engine/Source/Runtime/GameplayAbilities/Private/Abilities/GameplayAbility.cpp b/Engine/Source/Runtime/GameplayAbilities/Private/Abilities/GameplayAbility.cpp
--- a/Engine/Source/Runtime/GameplayAbilities/Private/Abilities/GameplayAbility.cpp
+++ b/Engine/Source/Runtime/GameplayAbilities/Private/Abilities/GameplayAbility.cpp
@@ -18,6 +18,22 @@ UGameplayAbility::UGameplayAbility(const class FPostConstructInitializePropertie
 	CostGameplayEffect = NULL;
 	CooldownGameplayEffect = NULL;
 
+	
+	{
+		static FName FuncName = FName(TEXT("K2_CanActivateAbility"));
+		UFunction* CanActivateFunction = GetClass()->FindFunctionByName(FuncName);
+		HasBlueprintCanUse = CanActivateFunction && CanActivateFunction->GetOuter()->IsA(UBlueprintGeneratedClass::StaticClass());
+	}
+	{
+		static FName FuncName = FName(TEXT("K2_ActivateAbility"));
+		UFunction* ActivateFunction = GetClass()->FindFunctionByName(FuncName);
+		HasBlueprintActivate = ActivateFunction && ActivateFunction->GetOuter()->IsA(UBlueprintGeneratedClass::StaticClass());
+	}
+	{
+		static FName FuncName = FName(TEXT("K2_PredictiveActivateAbility"));
+		UFunction* PredictiveActivateFunction = GetClass()->FindFunctionByName(FuncName);
+		HasBlueprintPredictiveActivate = PredictiveActivateFunction && PredictiveActivateFunction->GetOuter()->IsA(UBlueprintGeneratedClass::StaticClass());
+	}
 }
 
 bool UGameplayAbility::IsSupportedForNetworking() const
@@ -40,6 +56,8 @@ bool UGameplayAbility::IsSupportedForNetworking() const
 
 bool UGameplayAbility::CanActivateAbility(const FGameplayAbilityActorInfo* ActorInfo) const
 {
+	SetCurrentActorInfo(ActorInfo);
+
 	if (ActorInfo->AbilitySystemComponent->GetUserAbilityActivationInhibited())
 	{
 		/**
@@ -64,6 +82,15 @@ bool UGameplayAbility::CanActivateAbility(const FGameplayAbilityActorInfo* Actor
 		return false;
 	}
 
+	if (HasBlueprintCanUse)
+	{
+		if (K2_CanActivateAbility(*ActorInfo) == false)
+		{
+			ABILITY_LOG(Log, TEXT("CanActivateAbility %s failed, blueprint refused"), *GetName());
+			return false;
+		}
+	}
+
 	return true;
 }
 
@@ -75,6 +102,10 @@ bool UGameplayAbility::CommitAbility(const FGameplayAbilityActorInfo* ActorInfo,
 	}
 
 	CommitExecute(ActorInfo, ActivationInfo);
+
+	// Fixme: Should we always call this or only if it is implemented? A noop may not hurt but could be bad for perf (storing a HasBlueprintCommit per instance isn't good either)
+	K2_CommitExecute();
+
 	return true;
 }
 
@@ -168,16 +199,34 @@ bool UGameplayAbility::TryActivateAbility(const FGameplayAbilityActorInfo* Actor
 
 void UGameplayAbility::EndAbility(const FGameplayAbilityActorInfo* ActorInfo)
 {
+	if (InstancingPolicy == EGameplayAbilityInstancingPolicy::InstancedPerExecution)
+	{
+		MarkPendingKill();
+	}
 
+	// Remove from owning AbilitySystemComponent?
+	// generic way of releasing all callbacks!
 }
 
 void UGameplayAbility::ActivateAbility(const FGameplayAbilityActorInfo* ActorInfo, const FGameplayAbilityActivationInfo ActivationInfo)
 {
-	// Child classes may want to do stuff here...
+	SetCurrentInfo(ActorInfo, ActivationInfo);
 
-	if (CommitAbility(ActorInfo, ActivationInfo))		// ..then commit the ability...
+	if (HasBlueprintActivate)
 	{
-		//	Then do more stuff...
+		// A Blueprinted ActivateAbility function must call CommitAbility somewhere in its execution chain.
+		K2_ActivateAbility();
+	}
+	else
+	{
+		// Native child classes may want to override ActivateAbility and do something like this:
+
+		// Do stuff...
+
+		if (CommitAbility(ActorInfo, ActivationInfo))		// ..then commit the ability...
+		{			
+			//	Then do more stuff...
+		}
 	}
 }
 
@@ -204,6 +253,14 @@ void UGameplayAbility::CallPredictiveActivateAbility(const FGameplayAbilityActor
 
 void UGameplayAbility::PredictiveActivateAbility(const FGameplayAbilityActorInfo* ActorInfo, const FGameplayAbilityActivationInfo ActivationInfo)
 {
+	SetCurrentInfo(ActorInfo, ActivationInfo);
+	
+	check(ActivationInfo.ActivationMode == EGameplayAbilityActivationMode::Predicting);
+
+	if (HasBlueprintPredictiveActivate)
+	{
+		K2_PredictiveActivateAbility();
+	}
 	
 }
 
@@ -286,17 +343,51 @@ float UGameplayAbility::GetCooldownTimeRemaining(const FGameplayAbilityActorInfo
 	return 0.f;
 }
 
-const FGameplayAbilityActorInfo* UGameplayAbility::GetCurrentActorInfo()
+FGameplayAbilityActorInfo UGameplayAbility::GetActorInfo()
+{
+	check(CurrentActorInfo);
+	return *CurrentActorInfo;
+}
+
+/** Fixme: Naming is confusing here */
+
+bool UGameplayAbility::K2_CommitAbility()
+{
+	check(CurrentActorInfo);
+	return CommitAbility(CurrentActorInfo, CurrentActivationInfo);
+}
+
+void UGameplayAbility::K2_EndAbility()
 {
-	ABILITY_LOG(Fatal, TEXT("GetCurrentActorInfo called on base UGameplayAbility class"));
-	return NULL;
+	check(CurrentActorInfo);
+	EndAbility(CurrentActorInfo);
 }
 
-FGameplayAbilityActivationInfo UGameplayAbility::GetCurrentActivationInfo()
+
+int32 UGameplayAbility::GetFunctionCallspace(UFunction* Function, void* Parameters, FFrame* Stack)
+{
+	if (HasAnyFlags(RF_ClassDefaultObject))
+	{
+		return FunctionCallspace::Local;
+	}
+	check(GetOuter() != NULL);
+	return GetOuter()->GetFunctionCallspace(Function, Parameters, Stack);
+}
+
+bool UGameplayAbility::CallRemoteFunction(UFunction* Function, void* Parameters, FOutParmRec* OutParms, FFrame* Stack)
 {
-	ABILITY_LOG(Fatal, TEXT("GetCurrentActivationInfo called on base UGameplayAbility class"));
+	check(!HasAnyFlags(RF_ClassDefaultObject));
+	check(GetOuter() != NULL);
 
-	return FGameplayAbilityActivationInfo();
+	AActor* Owner = CastChecked<AActor>(GetOuter());
+	UNetDriver* NetDriver = Owner->GetNetDriver();
+	if (NetDriver)
+	{
+		NetDriver->ProcessRemoteFunction(Owner, Function, Parameters, OutParms, Stack, this);
+		return true;
+	}
+
+	return false;
 }
 
 // --------------------------------------------------------------------
@@ -319,7 +410,10 @@ void UGameplayAbility::ClientActivateAbilitySucceed_Implementation(int32 Predict
 
 void UGameplayAbility::ClientActivateAbilitySucceed_Internal(int32 PredictionKey)
 {
-	ABILITY_LOG(Fatal, TEXT("ClientActivateAbilitySucceed_Internal called on base class ability %s"), *GetName());
+	check(CurrentActorInfo);
+	CurrentActivationInfo.SetActivationConfirmed();
+
+	CallActivateAbility(CurrentActorInfo, CurrentActivationInfo);
 }
 
 //----------------------------------------------------------------------
diff --git a/Engine/Source/Runtime/GameplayAbilities/Private/Abilities/GameplayAbility_Instanced.cpp b/Engine/Source/Runtime/GameplayAbilities/Private/Abilities/GameplayAbility_Instanced.cpp
deleted file mode 100644
diff --git a/Engine/Source/Runtime/GameplayAbilities/Private/Abilities/Tasks/AbilityTask_WaitConfirmCancel.cpp b/Engine/Source/Runtime/GameplayAbilities/Private/Abilities/Tasks/AbilityTask_WaitConfirmCancel.cpp
--- a/Engine/Source/Runtime/GameplayAbilities/Private/Abilities/Tasks/AbilityTask_WaitConfirmCancel.cpp
+++ b/Engine/Source/Runtime/GameplayAbilities/Private/Abilities/Tasks/AbilityTask_WaitConfirmCancel.cpp
@@ -5,7 +5,6 @@
 #include "AbilitySystemComponent.h"
 
 #include "Abilities/GameplayAbility.h"
-#include "Abilities/GameplayAbility_Instanced.h"
 
 UAbilityTask_WaitConfirmCancel::UAbilityTask_WaitConfirmCancel(const class FPostConstructInitializeProperties& PCIP)
 	: Super(PCIP)
@@ -28,7 +27,7 @@ void UAbilityTask_WaitConfirmCancel::OnCancelCallback()
 UAbilityTask_WaitConfirmCancel* UAbilityTask_WaitConfirmCancel::WaitConfirmCancel(UObject* WorldContextObject)
 {
 	check(WorldContextObject);
-	UGameplayAbility_Instanced* ThisAbility = CastChecked<UGameplayAbility_Instanced>(WorldContextObject);
+	UGameplayAbility* ThisAbility = CastChecked<UGameplayAbility>(WorldContextObject);
 	if (ThisAbility)
 	{
 		UAbilityTask_WaitConfirmCancel * MyObj = NULL;
diff --git a/Engine/Source/Runtime/GameplayAbilities/Private/AbilitySystemBlueprintLibrary.cpp b/Engine/Source/Runtime/GameplayAbilities/Private/AbilitySystemBlueprintLibrary.cpp
--- a/Engine/Source/Runtime/GameplayAbilities/Private/AbilitySystemBlueprintLibrary.cpp
+++ b/Engine/Source/Runtime/GameplayAbilities/Private/AbilitySystemBlueprintLibrary.cpp
@@ -2,12 +2,8 @@
 
 #include "AbilitySystemPrivatePCH.h"
 #include "AbilitySystemBlueprintLibrary.h"
-
 #include "AbilitySystemComponent.h"
-
 #include "Abilities/GameplayAbility.h"
-#include "Abilities/GameplayAbility_Instanced.h"
-
 #include "Abilities/Tasks/AbilityTask_PlayMontageAndWait.h"
 #include "Abilities/Tasks/AbilityTask_WaitMovementModeChange.h"
 #include "Abilities/Tasks/AbilityTask_WaitOverlap.h"
@@ -68,7 +64,7 @@ FHitResult UAbilitySystemBlueprintLibrary::GetHitResultFromTargetData(FGameplayA
 UAbilityTask_PlayMontageAndWait* UAbilitySystemBlueprintLibrary::CreatePlayMontageAndWaitProxy(class UObject* WorldContextObject, UAnimMontage *MontageToPlay)
 {
 	check(WorldContextObject);
-	UGameplayAbility_Instanced * Ability = CastChecked<UGameplayAbility_Instanced>(WorldContextObject);
+	UGameplayAbility* Ability = CastChecked<UGameplayAbility>(WorldContextObject);
 	if (Ability)
 	{
 		AActor * ActorOwner = Cast<AActor>(Ability->GetOuter());
@@ -95,7 +91,7 @@ UAbilityTask_PlayMontageAndWait* UAbilitySystemBlueprintLibrary::CreatePlayMonta
 UAbilityTask_WaitMovementModeChange* UAbilitySystemBlueprintLibrary::CreateWaitMovementModeChange(class UObject* WorldContextObject, EMovementMode NewMode)
 {
 	check(WorldContextObject);
-	UGameplayAbility_Instanced * Ability = CastChecked<UGameplayAbility_Instanced>(WorldContextObject);
+	UGameplayAbility* Ability = CastChecked<UGameplayAbility>(WorldContextObject);
 	if (Ability)
 	{
 		AActor * ActorOwner = Cast<AActor>(Ability->GetOuter());
@@ -127,7 +123,7 @@ UAbilityTask_WaitMovementModeChange* UAbilitySystemBlueprintLibrary::CreateWaitM
 UAbilityTask_WaitOverlap* UAbilitySystemBlueprintLibrary::CreateWaitOverlap(class UObject* WorldContextObject, EMovementMode NewMode)
 {
 	check(WorldContextObject);
-	UGameplayAbility_Instanced * Ability = CastChecked<UGameplayAbility_Instanced>(WorldContextObject);
+	UGameplayAbility* Ability = CastChecked<UGameplayAbility>(WorldContextObject);
 	if (Ability)
 	{
 		AActor * ActorOwner = Cast<AActor>(Ability->GetOuter());
diff --git a/Engine/Source/Runtime/GameplayAbilities/Public/Abilities/GameplayAbility.h b/Engine/Source/Runtime/GameplayAbilities/Public/Abilities/GameplayAbility.h
--- a/Engine/Source/Runtime/GameplayAbilities/Public/Abilities/GameplayAbility.h
+++ b/Engine/Source/Runtime/GameplayAbilities/Public/Abilities/GameplayAbility.h
@@ -33,9 +33,7 @@
  *			-'Giving' abilities (that can be used) to actors
  *	
  *	
+ *	 
  *	
- *	UGameplayAbility is the base class. It is non instanced, non replicating, and not blueprintable. 
- *	This is essentially the 'cheapest' type of ability with the least amount of overhead.
- *	
  *	The intention is for programmers to create these non instanced abilities in C++. Designers can then
  *	extend them as data assets (E.g., they can change default properties, they cannot implement blueprint graphs).
@@ -56,7 +54,7 @@
  *	A lot is possible with non instanced abilities but care must be taken.
  *	
  *	
- *	To support state or event replication, an ability must be instanced. This can be done with the UGameplayAbility_Instanced class.
+ *	To support state or event replication, an ability must be instanced. This can be done with the InstancingPolicy property.
  *	
  *
  *	
@@ -87,20 +85,15 @@
 ***/
 
 
-UCLASS()
+UCLASS(Blueprintable)
 class GAMEPLAYABILITIES_API UGameplayAbility : public UDataAsset
 {
 	GENERATED_UCLASS_BODY()
 
+	friend class UAbilitySystemComponent;
+
 public:
 
-	UFUNCTION(Client, Reliable)
-	void ClientActivateAbilitySucceed(int32 PredictionKey);
-
-	virtual void ClientActivateAbilitySucceed_Internal(int32 PredictionKey);
-
-	bool IsSupportedForNetworking() const OVERRIDE;
-	
 	// ----------------------------------------------------------------------------------------------------------------
 	//
 	//	The important functions:
@@ -122,7 +115,13 @@
 	//		EndAbility()			- The ability has ended. This is intended to be called by the ability to end itself.
 	//	
 	// ----------------------------------------------------------------------------------------------------------------
+	
+	/** Input binding. Base implementation calls TryActivateAbility */
+	virtual void InputPressed(int32 InputID, const FGameplayAbilityActorInfo* ActorInfo);
 
+	/** Input binding. Base implementation does nothing */
+	virtual void InputReleased(int32 InputID, const FGameplayAbilityActorInfo* ActorInfo);
+	
 	/**
 	 * Attempts to activate the ability.
 	 *	-This function calls CanActivateAbility
@@ -131,142 +130,177 @@
 	 *	-If all goes well, CallActivateAbility is called next.
 	 */
 	bool TryActivateAbility(const FGameplayAbilityActorInfo* ActorInfo, int32 PredictionKey = 0, UGameplayAbility ** OutInstancedAbility = NULL);
-	
+
 	/** Returns true if this ability can be activated right now. Has no side effects */
 	virtual bool CanActivateAbility(const FGameplayAbilityActorInfo* ActorInfo) const;
-
-	// -----------------------------------------------
-
-	/** Input binding. Base implementation calls TryActivateAbility */
-	virtual void InputPressed(int32 InputID, const FGameplayAbilityActorInfo* ActorInfo);
-
-	/** Input binding. Base implementation does nothing */
-	virtual void InputReleased(int32 InputID, const FGameplayAbilityActorInfo* ActorInfo);
-
-	// ----------------------------------------------------------------------------------------------------------------
-	//
-	//	Outward facing API for other systems
-	//
-	// ----------------------------------------------------------------------------------------------------------------
-
+	
 	float GetCooldownTimeRemaining(const FGameplayAbilityActorInfo* ActorInfo) const;
-
-	// ----------------------------------------------------------------------------------------------------------------
-	//
-	//	Standardized State
-	//
-	// ----------------------------------------------------------------------------------------------------------------	
-
-	UPROPERTY(EditDefaultsOnly, Category=Advanced)
-	TEnumAsByte<EGameplayAbilityNetExecutionPolicy::Type>	NetExecutionPolicy;
-
-	/** This GameplayEffect represents the cooldown. It will be applied when the ability is commited and the ability cannot be used again until it is expired. */
-	UPROPERTY(EditDefaultsOnly, Category=Cooldowns)
-	class UGameplayEffect * CooldownGameplayEffect;
-
-	/** This GameplayEffect represents the cost (mana, stamina, etc) of the ability. It will be applied when the ability is commited. */
-	UPROPERTY(EditDefaultsOnly, Category=Costs)
-	class UGameplayEffect * CostGameplayEffect;
-
-	// ----------------------------------------------------------------------------------------------------------------
-	//
-	//	Ability exclusion / cancelling
-	//
-	// ----------------------------------------------------------------------------------------------------------------
-	
-	/** Abilities with these tags are canelled when this ability is executed */
-	UPROPERTY(EditDefaultsOnly, Category = Tags)
-	FGameplayTagContainer CancelAbilitiesWithTag;
-
-	/** This ability has these tags */
-	UPROPERTY(EditDefaultsOnly, Category = Tags)
-	FGameplayTagContainer AbilityTags;
-
-	// ----------------------------------------------------------------------------------------------------------------
-	//
-	//	Animation callbacks
-	//
-	// ----------------------------------------------------------------------------------------------------------------
-
-	void MontageBranchPoint_AbilityDecisionStop(const FGameplayAbilityActorInfo* ActorInfo) const;
-
-	void MontageBranchPoint_AbilityDecisionStart(const FGameplayAbilityActorInfo* ActorInfo) const;
-
-	// ----------------------------------------------------------------------------------------------------------------
-
-	virtual void PredictiveActivateAbility(const FGameplayAbilityActorInfo* ActorInfo, const FGameplayAbilityActivationInfo ActivationInfo);
-
-	virtual void ServerTryActivateAbility(const FGameplayAbilityActorInfo* ActorInfo, const FGameplayAbilityActivationInfo ActivationInfo);
-
-	virtual void ClientActivateAbility(const FGameplayAbilityActorInfo* ActorInfo, const FGameplayAbilityActivationInfo ActivationInfo);
-
-	// ----------------------------------------------------------------------------------------------------------------
-
-	virtual EGameplayAbilityInstancingPolicy::Type GetInstancingPolicy() const
+		
+	EGameplayAbilityInstancingPolicy::Type GetInstancingPolicy() const
 	{
-		return EGameplayAbilityInstancingPolicy::NonInstanced;
+		return InstancingPolicy;
 	}
 
-	virtual EGameplayAbilityReplicationPolicy::Type GetReplicationPolicy() const
+	EGameplayAbilityReplicationPolicy::Type GetReplicationPolicy() const
 	{
-		return EGameplayAbilityReplicationPolicy::ReplicateNone;
+		return ReplicationPolicy;
 	}
 
-	// ----------------------------------------------------------------------------------------------------------------
-
 	/** Gets the current actor info bound to this ability - can only be called on instanced abilities. */
-	virtual const FGameplayAbilityActorInfo* GetCurrentActorInfo();
+	const FGameplayAbilityActorInfo* GetCurrentActorInfo()
+	{
+		check(IsInstantiated());
+		return CurrentActorInfo;
+	}
 
 	/** Gets the current activation info bound to this ability - can only be called on instanced abilities. */
-	virtual FGameplayAbilityActivationInfo GetCurrentActivationInfo();
+	FGameplayAbilityActivationInfo GetCurrentActivationInfo()
+	{
+		check(IsInstantiated());
+		return CurrentActivationInfo;
+	}
 
-	virtual UWorld* GetWorld() const OVERRIDE
+	virtual UWorld* GetWorld() const override
 	{
 		return GetOuter()->GetWorld();
 	}
 
+	int32 GetFunctionCallspace(UFunction* Function, void* Parameters, FFrame* Stack);
+
+	bool CallRemoteFunction(UFunction* Function, void* Parameters, FOutParmRec* OutParms, FFrame* Stack);
+
 protected:
+		
+	// --------------------------------------
+	//	CanActivate
+	// --------------------------------------
+	
+	/** Returns true if this ability can be activated right now. Has no side effects */
+	UFUNCTION(BlueprintImplementableEvent, Category = Ability, FriendlyName="CanActivateAbility")
+	virtual bool K2_CanActivateAbility(FGameplayAbilityActorInfo ActorInfo) const;
+
+	bool HasBlueprintCanUse;
+
+	// --------------------------------------
+	//	ActivateAbility
+	// --------------------------------------
 
 	/**
 	 * The main function that defines what an ability does.
 	 *  -Child classes will want to override this
 	 *  -This function must call CommitAbility()
 	 */
+	UFUNCTION(BlueprintImplementableEvent, Category = Ability, FriendlyName = "ActivateAbility")
+	virtual void K2_ActivateAbility();	
 
 	virtual void ActivateAbility(const FGameplayAbilityActorInfo * ActorInfo, const FGameplayAbilityActivationInfo ActivationInfo);
 
-	/** 
+	bool HasBlueprintActivate;
+
+	void CallActivateAbility(const FGameplayAbilityActorInfo* ActorInfo, const FGameplayAbilityActivationInfo ActivationInfo);
+
+	// --------------------------------------
+	//	CommitAbility
+	// --------------------------------------
+
+	/**
 	 * Attempts to commit the ability (spend resources, etc). This our last chance to fail.
 	 *	-Child classes that override ActivateAbility must call this themselves!
 	 */
+	UFUNCTION(BlueprintCallable, Category = Ability, FriendlyName = "CommitAbility")
+	virtual bool K2_CommitAbility();
 
 	virtual bool CommitAbility(const FGameplayAbilityActorInfo* ActorInfo, const FGameplayAbilityActivationInfo ActivationInfo);
-	
-	/** Do boilerplate init stuff and then call ActivateAbility */
-	void PreActivate(const FGameplayAbilityActorInfo* ActorInfo, const FGameplayAbilityActivationInfo ActivationInfo);
 
-	void CallActivateAbility(const FGameplayAbilityActorInfo* ActorInfo, const FGameplayAbilityActivationInfo ActivationInfo);
-
-	void CallPredictiveActivateAbility(const FGameplayAbilityActorInfo* ActorInfo, const FGameplayAbilityActivationInfo ActivationInfo);
-
-	virtual void EndAbility(const FGameplayAbilityActorInfo* ActorInfo);
-
-	/** 
+	/**
 	 * The last chance to fail before commiting
 	 *	-This will usually be the same as CanActivateAbility. Some abilities may need to do extra checks here if they are consuming extra stuff in CommitExecute
 	 */
 	virtual bool CommitCheck(const FGameplayAbilityActorInfo* ActorInfo, const FGameplayAbilityActivationInfo ActivationInfo);
 
+	// --------------------------------------
+	//	CommitExecute
+	// --------------------------------------
+
+	UFUNCTION(BlueprintImplementableEvent, Category = Ability, FriendlyName="CommitExecute")
+	virtual void K2_CommitExecute();
+
 	/** Does the commmit atomically (consume resources, do cooldowns, etc) */
 	virtual void CommitExecute(const FGameplayAbilityActorInfo* ActorInfo, const FGameplayAbilityActivationInfo ActivationInfo);
 
+	/** Do boilerplate init stuff and then call ActivateAbility */
+	void PreActivate(const FGameplayAbilityActorInfo* ActorInfo, const FGameplayAbilityActivationInfo ActivationInfo);
+
+	// --------------------------------------
+	//	CancelAbility
+	// --------------------------------------
+
 	/** Destroys intanced-per-execution abilities. Instance-per-actor abilities should 'reset'. Non instance abilities - what can we do? */
 	virtual void CancelAbility(const FGameplayAbilityActorInfo* ActorInfo, const FGameplayAbilityActivationInfo ActivationInfo) { }
 
-	// --------------------------------------------------------------------------------
+	// -------------------------------------
+	//	PredictiveActivate
+	// -------------------------------------
 
+	UFUNCTION(BlueprintImplementableEvent, Category = Ability, FriendlyName="PredictiveActivateAbility")
+	virtual void K2_PredictiveActivateAbility();
 
-	// --------------------------------------------------------------------------------
+	virtual void PredictiveActivateAbility(const FGameplayAbilityActorInfo* ActorInfo, const FGameplayAbilityActivationInfo ActivationInfo);
 
+	void CallPredictiveActivateAbility(const FGameplayAbilityActorInfo* ActorInfo, const FGameplayAbilityActivationInfo ActivationInfo);
+
+	bool HasBlueprintPredictiveActivate;
+
+	// -------------------------------------
+	//	EndAbility
+	// -------------------------------------
+
+	UFUNCTION(BlueprintCallable, Category = Ability, FriendlyName="EndAbility")
+	virtual void K2_EndAbility();
+
+	virtual void EndAbility(const FGameplayAbilityActorInfo* ActorInfo);
+
+	// -------------------------------------
+	
+	/** Ask server to try to activate ability */
+	virtual void ServerTryActivateAbility(const FGameplayAbilityActorInfo* ActorInfo, const FGameplayAbilityActivationInfo ActivationInfo);
+
+	/** Called once server OKs a client ability activation */
+	virtual void ClientActivateAbility(const FGameplayAbilityActorInfo* ActorInfo, const FGameplayAbilityActivationInfo ActivationInfo);
+
+	// -------------------------------------
+
+	bool IsInstantiated() const
+	{
+		return !HasAllFlags(RF_ClassDefaultObject);
+	}
+
+	void SetCurrentActorInfo(const FGameplayAbilityActorInfo* ActorInfo) const
+	{
+		if (IsInstantiated())
+		{
+			CurrentActorInfo = ActorInfo;
+		}
+	}
+
+	void SetCurrentInfo(const FGameplayAbilityActorInfo* ActorInfo, const FGameplayAbilityActivationInfo ActivationInfo)
+	{
+		if (IsInstantiated())
+		{
+			CurrentActivationInfo = ActivationInfo;
+			CurrentActorInfo = ActorInfo;
+		}
+	}
+
+	UFUNCTION(BlueprintPure, Category = Ability)
+	FGameplayAbilityActorInfo GetActorInfo();
+
+	UFUNCTION(Client, Reliable)
+	void ClientActivateAbilitySucceed(int32 PredictionKey);
+
+	virtual void ClientActivateAbilitySucceed_Internal(int32 PredictionKey);
+
+	bool IsSupportedForNetworking() const OVERRIDE;
+
 	/** Checks cooldown. returns true if we can be used again. False if not */
 	bool	CheckCooldown(const FGameplayAbilityActorInfo* ActorInfo) const;
@@ -278,10 +312,64 @@
 
 	void	ApplyCost(const FGameplayAbilityActorInfo* ActorInfo, const FGameplayAbilityActivationInfo ActivationInfo);
 
-	// --------------------------------------------------------------------------------
+	// -----------------------------------------------	
 
+	UPROPERTY(EditDefaultsOnly, Category = Advanced)
+	TEnumAsByte<EGameplayAbilityReplicationPolicy::Type> ReplicationPolicy;
+
+	UPROPERTY(EditDefaultsOnly, Category = Advanced)
+	TEnumAsByte<EGameplayAbilityInstancingPolicy::Type>	InstancingPolicy;						
+
+	/** This is information specific to this instance of the ability. E.g, whether it is predicting, authorting, confirmed, etc. */
+	UPROPERTY(BlueprintReadOnly, Category = Ability)
+	FGameplayAbilityActivationInfo	CurrentActivationInfo;
+
+	UPROPERTY(EditDefaultsOnly, Category=Advanced)
+	TEnumAsByte<EGameplayAbilityNetExecutionPolicy::Type> NetExecutionPolicy;
+
+	/** This GameplayEffect represents the cooldown. It will be applied when the ability is commited and the ability cannot be used again until it is expired. */
+	UPROPERTY(EditDefaultsOnly, Category=Cooldowns)
+	class UGameplayEffect* CooldownGameplayEffect;
+
+	/** This GameplayEffect represents the cost (mana, stamina, etc) of the ability. It will be applied when the ability is commited. */
+	UPROPERTY(EditDefaultsOnly, Category=Costs)
+	class UGameplayEffect* CostGameplayEffect;
+
+	// ----------------------------------------------------------------------------------------------------------------
+	//
+	//	Ability exclusion / cancelling
+	//
+	// ----------------------------------------------------------------------------------------------------------------
+	
+	/** Abilities with these tags are canelled when this ability is executed */
 	UPROPERTY(EditDefaultsOnly, Category = Tags)
+	FGameplayTagContainer CancelAbilitiesWithTag;
+
+	/** This ability has these tags */
+	UPROPERTY(EditDefaultsOnly, Category = Tags)
+	FGameplayTagContainer AbilityTags;
+
+	// ----------------------------------------------------------------------------------------------------------------
+	//
+	//	Animation callbacks
+	//
+	// ----------------------------------------------------------------------------------------------------------------
+
+	void MontageBranchPoint_AbilityDecisionStop(const FGameplayAbilityActorInfo* ActorInfo) const;
+
+	void MontageBranchPoint_AbilityDecisionStart(const FGameplayAbilityActorInfo* ActorInfo) const;
+	
+	UPROPERTY(EditDefaultsOnly, Category = Tags)
 	bool Prediction;
 
-	friend class UAbilitySystemComponent;
+private:
+
+	/** 
+	 *  This is shared, cached information about the thing using us
+	 *	 E.g, Actor*, MovementComponent*, AnimInstance, etc.
+	 *	 This is hopefully allocated once per actor and shared by many abilities.
+	 *	 The actual struct may be overridden per game to include game specific data.
+	 *	 (E.g, child classes may want to cast to FMyGameAbilityActorInfo)
+	 */
+	mutable const FGameplayAbilityActorInfo* CurrentActorInfo;
 };
diff --git a/Engine/Source/Runtime/GameplayAbilities/Public/Abilities/GameplayAbility_Instanced.h b/Engine/Source/Runtime/GameplayAbilities/Public/Abilities/GameplayAbility_Instanced.h
deleted file mode 100644
diff --git a/Engine/Source/Runtime/GameplayAbilities/Public/Abilities/Tasks/AbilityTask.h b/Engine/Source/Runtime/GameplayAbilities/Public/Abilities/Tasks/AbilityTask.h
--- a/Engine/Source/Runtime/GameplayAbilities/Public/Abilities/Tasks/AbilityTask.h
+++ b/Engine/Source/Runtime/GameplayAbilities/Public/Abilities/Tasks/AbilityTask.h
@@ -1,6 +1,6 @@
 // Copyright 1998-2014 Epic Games, Inc. All Rights Reserved.
 #pragma once
-#include "GameplayAbility_Instanced.h"
+#include "GameplayAbility.h"
 #include "AbilityTask.Generated.h"
 
 UCLASS(MinimalAPI)
diff --git a/Engine/Source/Runtime/GameplayAbilities/Public/Abilities/Tasks/AbilityTask_WaitConfirmCancel.h b/Engine/Source/Runtime/GameplayAbilities/Public/Abilities/Tasks/AbilityTask_WaitConfirmCancel.h
--- a/Engine/Source/Runtime/GameplayAbilities/Public/Abilities/Tasks/AbilityTask_WaitConfirmCancel.h
+++ b/Engine/Source/Runtime/GameplayAbilities/Public/Abilities/Tasks/AbilityTask_WaitConfirmCancel.h
@@ -2,7 +2,7 @@
 #pragma once
 
 #include "AbilityTask.h"
-#include "GameplayAbility_Instanced.h"
+#include "GameplayAbility.h"
 #include "AbilityTask_WaitConfirmCancel.generated.h"
 
 DECLARE_DYNAMIC_MULTICAST_DELEGATE_OneParam(FWaitConfirmCancelDelegate, bool, Test);
