make UGameplayAbility not a datasset

diff --git a/Engine/Source/Runtime/GameplayAbilities/Private/Abilities/GameplayAbilitySet.cpp b/Engine/Source/Runtime/GameplayAbilities/Private/Abilities/GameplayAbilitySet.cpp
--- a/Engine/Source/Runtime/GameplayAbilities/Private/Abilities/GameplayAbilitySet.cpp
+++ b/Engine/Source/Runtime/GameplayAbilities/Private/Abilities/GameplayAbilitySet.cpp
@@ -61,7 +61,7 @@ void UGameplayAbilitySet::InitializeAbilities(UAbilitySystemComponent *AbilitySy
 {
 	for (const FGameplayAbilityBindInfo &BindInfo : this->Abilities)
 	{
-		UGameplayAbility * Ability = BindInfo.GameplayAbilityClass ? BindInfo.GameplayAbilityClass->GetDefaultObject<UGameplayAbility>() : BindInfo.GameplayAbilityInstance;
+		UGameplayAbility * Ability = BindInfo.GameplayAbilityClass ? BindInfo.GameplayAbilityClass->GetDefaultObject<UGameplayAbility>() : NULL;
 		if (Ability)
 		{
 			if (Ability->GetInstancingPolicy() == EGameplayAbilityInstancingPolicy::InstancedPerActor)
diff --git a/Engine/Source/Runtime/GameplayAbilities/Public/Abilities/GameplayAbility.h b/Engine/Source/Runtime/GameplayAbilities/Public/Abilities/GameplayAbility.h
--- a/Engine/Source/Runtime/GameplayAbilities/Public/Abilities/GameplayAbility.h
+++ b/Engine/Source/Runtime/GameplayAbilities/Public/Abilities/GameplayAbility.h
@@ -85,8 +85,11 @@ Server:		     --                      --                    --
 ***/
 
 
+/**
+ *	Abilities define custom gameplay logic that can be activated by players or external game logic.
+ */
 UCLASS(Blueprintable)
-class GAMEPLAYABILITIES_API UGameplayAbility : public UDataAsset
+class GAMEPLAYABILITIES_API UGameplayAbility : public UObject
 {
 	GENERATED_UCLASS_BODY()
 
diff --git a/Engine/Source/Runtime/GameplayAbilities/Public/GameplayAbilitySet.h b/Engine/Source/Runtime/GameplayAbilities/Public/GameplayAbilitySet.h
--- a/Engine/Source/Runtime/GameplayAbilities/Public/GameplayAbilitySet.h
+++ b/Engine/Source/Runtime/GameplayAbilities/Public/GameplayAbilitySet.h
@@ -27,9 +27,6 @@ struct FGameplayAbilityBindInfo
 
 	UPROPERTY(EditAnywhere, Category=BindInfo)
 	TSubclassOf<UGameplayAbility>	GameplayAbilityClass;
-
-	UPROPERTY(EditAnywhere, Category = BindInfo)
-	class UGameplayAbility * GameplayAbilityInstance;
 };
 
 /**
