cis fix

diff --git a/Engine/Source/Runtime/GameplayAbilities/Private/Abilities/GameplayAbility.cpp b/Engine/Source/Runtime/GameplayAbilities/Private/Abilities/GameplayAbility.cpp
--- a/Engine/Source/Runtime/GameplayAbilities/Private/Abilities/GameplayAbility.cpp
+++ b/Engine/Source/Runtime/GameplayAbilities/Private/Abilities/GameplayAbility.cpp
@@ -69,7 +69,7 @@ bool UGameplayAbility::CanActivateAbility(const FGameplayAbilityActorInfo* Actor
 		 *	Basically: CanActivateAbility is only used by user activated abilities now. If triggered abilities need to check costs/cooldowns, then we may
 		 *	want to split this function up and change the calling API to distinguish between 'can I initiate an ability activation' and 'can this ability be activated'.
 		 */ 
-		false;
+		return false;
 	}
 	
 	if (!CheckCooldown(ActorInfo))
diff --git a/Engine/Source/Runtime/GameplayAbilities/Public/Abilities/GameplayAbility.h b/Engine/Source/Runtime/GameplayAbilities/Public/Abilities/GameplayAbility.h
--- a/Engine/Source/Runtime/GameplayAbilities/Public/Abilities/GameplayAbility.h
+++ b/Engine/Source/Runtime/GameplayAbilities/Public/Abilities/GameplayAbility.h
@@ -362,8 +362,7 @@ protected:
 
 	void MontageBranchPoint_AbilityDecisionStart(const FGameplayAbilityActorInfo* ActorInfo) const;
 	
-	UPROPERTY(EditDefaultsOnly, Category = Tags)
-	bool Prediction;
+	
 
 private:
 
diff --git a/Engine/Source/Runtime/GameplayAbilities/Public/GameplayAbilitySet.h b/Engine/Source/Runtime/GameplayAbilities/Public/GameplayAbilitySet.h
--- a/Engine/Source/Runtime/GameplayAbilities/Public/GameplayAbilitySet.h
+++ b/Engine/Source/Runtime/GameplayAbilities/Public/GameplayAbilitySet.h
@@ -1,4 +1,4 @@
-// Copyright 1998-2013 Epic Games, Inc. All Rights Reserved.
+// Copyright 1998-2014 Epic Games, Inc. All Rights Reserved.
 
 #pragma once
 
