GameplayAbilities targeting updates

diff --git a/Engine/Source/Runtime/GameplayAbilities/Private/Abilities/GameplayAbility.cpp b/Engine/Source/Runtime/GameplayAbilities/Private/Abilities/GameplayAbility.cpp
--- a/Engine/Source/Runtime/GameplayAbilities/Private/Abilities/GameplayAbility.cpp
+++ b/Engine/Source/Runtime/GameplayAbilities/Private/Abilities/GameplayAbility.cpp
@@ -111,7 +111,25 @@ bool UGameplayAbility::CommitAbility(const FGameplayAbilityActorInfo* ActorInfo,
 
 bool UGameplayAbility::CommitCheck(const FGameplayAbilityActorInfo* ActorInfo, const FGameplayAbilityActivationInfo ActivationInfo)
 {
-	return CanActivateAbility(ActorInfo);
+	/**
+	 *	Checks if we can (still) commit this ability. There are some subtleties here.
+	 *		-An ability can start activating, play an animation, wait for a user confirmation/target data, and then actually commit
+	 *		-Commit = spend resources/cooldowns. Its possible the source has changed state since he started activation, so a commit may fail.
+	 *		-We don't want to just call CanActivateAbility() since right now that also checks things like input inhibition.
+	 *			-E.g., its possible the act of starting your ability makes it no longer activatable (CanaCtivateAbility() may be false if called here).
+	 */
+
+	if (!CheckCooldown(ActorInfo))
+	{
+		return false;
+	}
+
+	if (!CheckCost(ActorInfo))
+	{
+		return false;
+	}
+
+	return true;
 }
 
 void UGameplayAbility::CommitExecute(const FGameplayAbilityActorInfo* ActorInfo, const FGameplayAbilityActivationInfo ActivationInfo)
diff --git a/Engine/Source/Runtime/GameplayAbilities/Private/Abilities/GameplayAbilityTargetActor.cpp b/Engine/Source/Runtime/GameplayAbilities/Private/Abilities/GameplayAbilityTargetActor.cpp
--- a/Engine/Source/Runtime/GameplayAbilities/Private/Abilities/GameplayAbilityTargetActor.cpp
+++ b/Engine/Source/Runtime/GameplayAbilities/Private/Abilities/GameplayAbilityTargetActor.cpp
@@ -13,4 +13,19 @@ AGameplayAbilityTargetActor::AGameplayAbilityTargetActor(const class FPostConstr
 	: Super(PCIP)
 {
 	StaticTargetFunction = false;
+}
+
+void AGameplayAbilityTargetActor::StartTargeting(UGameplayAbility* Ability)
+{
+}
+
+void AGameplayAbilityTargetActor::ConfirmTargeting()
+{
+	TargetDataReadyDelegate.Broadcast(FGameplayAbilityTargetDataHandle());
+}
+
+/** Outside code is saying 'stop everything and just forget about it' */
+void AGameplayAbilityTargetActor::CancelTargeting()
+{
+	CanceledDelegate.Broadcast(FGameplayAbilityTargetDataHandle());
 }
\ No newline at end of file
diff --git a/Engine/Source/Runtime/GameplayAbilities/Private/Abilities/GameplayAbilityTargetActor_SingleLineTrace.cpp b/Engine/Source/Runtime/GameplayAbilities/Private/Abilities/GameplayAbilityTargetActor_SingleLineTrace.cpp
--- a/Engine/Source/Runtime/GameplayAbilities/Private/Abilities/GameplayAbilityTargetActor_SingleLineTrace.cpp
+++ b/Engine/Source/Runtime/GameplayAbilities/Private/Abilities/GameplayAbilityTargetActor_SingleLineTrace.cpp
@@ -18,6 +18,7 @@ AGameplayAbilityTargetActor_SingleLineTrace::AGameplayAbilityTargetActor_SingleL
 	PrimaryActorTick.TickGroup = TG_PrePhysics;
 	StaticTargetFunction = false;
 	bDebug = false;
+	bBindToConfirmCancelInputs = false;
 }
 
 FGameplayAbilityTargetDataHandle AGameplayAbilityTargetActor_SingleLineTrace::StaticGetTargetData(UWorld * World, const FGameplayAbilityActorInfo* ActorInfo, FGameplayAbilityActivationInfo ActivationInfo)
@@ -60,11 +61,15 @@ void AGameplayAbilityTargetActor_SingleLineTrace::StartTargeting(UGameplayAbilit
 {
 	Ability = InAbility;
 
-	UAbilitySystemComponent* ASC = Ability->GetCurrentActorInfo()->AbilitySystemComponent.Get();
-	if (ASC)
+	// We can bind directly to our ASC's confirm/cancel events, or wait to be told from an outside source to confirm or cancel
+	if (bBindToConfirmCancelInputs)
 	{
-		ASC->ConfirmCallbacks.AddDynamic(this, &AGameplayAbilityTargetActor_SingleLineTrace::Confirm);
-		ASC->CancelCallbacks.AddDynamic(this, &AGameplayAbilityTargetActor_SingleLineTrace::Cancel);
+		UAbilitySystemComponent* ASC = Ability->GetCurrentActorInfo()->AbilitySystemComponent.Get();
+		if (ASC)
+		{
+			ASC->ConfirmCallbacks.AddDynamic(this, &AGameplayAbilityTargetActor_SingleLineTrace::Confirm);
+			ASC->CancelCallbacks.AddDynamic(this, &AGameplayAbilityTargetActor_SingleLineTrace::Cancel);
+		}
 	}
 	
 	bDebug = true;
@@ -96,4 +101,16 @@ void AGameplayAbilityTargetActor_SingleLineTrace::Tick(float DeltaSeconds)
 	{
 		StaticGetTargetData(Ability->GetWorld(), Ability->GetCurrentActorInfo(), Ability->GetCurrentActivationInfo());
 	}
-}
\ No newline at end of file
+}
+
+void AGameplayAbilityTargetActor_SingleLineTrace::ConfirmTargeting()
+{
+	if (Ability.IsValid())
+	{
+		bDebug = false;
+		FGameplayAbilityTargetDataHandle Handle = StaticGetTargetData(Ability->GetWorld(), Ability->GetCurrentActorInfo(), Ability->GetCurrentActivationInfo());
+		TargetDataReadyDelegate.Broadcast(Handle);
+	}
+
+	Destroy();
+}
diff --git a/Engine/Source/Runtime/GameplayAbilities/Private/Abilities/Tasks/AbilityTask_WaitConfirmCancel.cpp b/Engine/Source/Runtime/GameplayAbilities/Private/Abilities/Tasks/AbilityTask_WaitConfirmCancel.cpp
--- a/Engine/Source/Runtime/GameplayAbilities/Private/Abilities/Tasks/AbilityTask_WaitConfirmCancel.cpp
+++ b/Engine/Source/Runtime/GameplayAbilities/Private/Abilities/Tasks/AbilityTask_WaitConfirmCancel.cpp
@@ -60,7 +60,12 @@ void UAbilityTask_WaitConfirmCancel::Activate()
 		}
 		
 		// We have to wait for the callback from the AbilitySystemComponent. Which may be instigated locally or through network replication
-		ASC->ConfirmCallbacks.AddDynamic(this, &UAbilityTask_WaitConfirmCancel::OnConfirmCallback);
-		ASC->CancelCallbacks.AddDynamic(this, &UAbilityTask_WaitConfirmCancel::OnCancelCallback);
+
+	
+		ASC->ConfirmCallbacks.AddDynamic(this, &UAbilityTask_WaitConfirmCancel::OnConfirmCallback);	// Tell me if the confirm input is pressed
+		ASC->CancelCallbacks.AddDynamic(this, &UAbilityTask_WaitConfirmCancel::OnCancelCallback);	// Tell me if the cancel input is pressed
+
+		// Tell me if something else tells me to 'force' my target data (an ability may allow targeting during animation/limited time)
+		// Tell me if something else tells me to cancel my targeting
 	}
 }
\ No newline at end of file
diff --git a/Engine/Source/Runtime/GameplayAbilities/Private/Abilities/Tasks/AbilityTask_WaitTargetData.cpp b/Engine/Source/Runtime/GameplayAbilities/Private/Abilities/Tasks/AbilityTask_WaitTargetData.cpp
--- a/Engine/Source/Runtime/GameplayAbilities/Private/Abilities/Tasks/AbilityTask_WaitTargetData.cpp
+++ b/Engine/Source/Runtime/GameplayAbilities/Private/Abilities/Tasks/AbilityTask_WaitTargetData.cpp
@@ -39,6 +39,7 @@ void UAbilityTask_WaitTargetData::Activate()
 			else
 			{
 				ASC->ReplicatedTargetDataDelegate.AddUObject(this, &UAbilityTask_WaitTargetData::OnTargetDataReplicatedCallback);
+				ASC->ReplicatedTargetDataCancelledDelegate.AddDynamic(this, &UAbilityTask_WaitTargetData::OnTargetDataReplicatedCancelledCallback);
 				return;
 			}
 		}
@@ -57,6 +58,10 @@ void UAbilityTask_WaitTargetData::Activate()
 				// This is a latent thing, probably with visuals and other gameplay related stuff. Spawn an actor locally and it will tell us when it has TargetData.
 				AGameplayAbilityTargetActor* SpawnedActor = CastChecked<AGameplayAbilityTargetActor>(Ability->GetWorld()->SpawnActor(TargetClass));
 				SpawnedActor->TargetDataReadyDelegate.AddUObject(this, &UAbilityTask_WaitTargetData::OnTargetDataReadyCallback);
+				SpawnedActor->CanceledDelegate.AddUObject(this, &UAbilityTask_WaitTargetData::OnTargetDataCancelledCallback);
+								
+				ASC->SetUserAbilityActivationInhibited(true);
+				ASC->SpawnedTargetActors.Push(SpawnedActor);
 
 				SpawnedActor->StartTargeting(Ability.Get());
 			}
@@ -71,10 +76,25 @@ void UAbilityTask_WaitTargetData::OnTargetDataReplicatedCallback(FGameplayAbilit
 	ASC->ConsumeAbilityTargetData();
 }
 
+void UAbilityTask_WaitTargetData::OnTargetDataReplicatedCancelledCallback()
+{
+	check(ASC);
+	Cancelled.Broadcast(FGameplayAbilityTargetDataHandle());
+}
+
 void UAbilityTask_WaitTargetData::OnTargetDataReadyCallback(FGameplayAbilityTargetDataHandle Data)
 {
 	check(ASC);
 	ASC->ServerSetReplicatedTargetData(Data);
 	ValidData.Broadcast(Data);
 	ASC->ConsumeAbilityTargetData();
+	ASC->SetUserAbilityActivationInhibited(false);
+}
+
+void UAbilityTask_WaitTargetData::OnTargetDataCancelledCallback(FGameplayAbilityTargetDataHandle Data)
+{
+	check(ASC);
+	ASC->ServerSetReplicatedTargetDataCancelled();
+	Cancelled.Broadcast(Data);
+	ASC->SetUserAbilityActivationInhibited(false);
 }
\ No newline at end of file
diff --git a/Engine/Source/Runtime/GameplayAbilities/Private/AbilitySystemComponent_Abilities.cpp b/Engine/Source/Runtime/GameplayAbilities/Private/AbilitySystemComponent_Abilities.cpp
--- a/Engine/Source/Runtime/GameplayAbilities/Private/AbilitySystemComponent_Abilities.cpp
+++ b/Engine/Source/Runtime/GameplayAbilities/Private/AbilitySystemComponent_Abilities.cpp
@@ -4,6 +4,7 @@
 #include "AbilitySystemPrivatePCH.h"
 #include "AbilitySystemComponent.h"
 #include "Abilities/GameplayAbility.h"
+#include "Abilities/GameplayAbilityTargetActor.h"
 
 #include "Net/UnrealNetwork.h"
 #include "MessageLog.h"
@@ -314,11 +315,37 @@ void UAbilitySystemComponent::InputCancel()
 	CancelCallbacks.Broadcast();
 }
 
+void UAbilitySystemComponent::TargetConfirm()
+{
+	for (AGameplayAbilityTargetActor* TargetActor : SpawnedTargetActors)
+	{
+		if (TargetActor)
+		{
+			TargetActor->ConfirmTargeting();
+		}
+	}
+
+	SpawnedTargetActors.Empty();
+}
+
+void UAbilitySystemComponent::TargetCancel()
+{
+	for (AGameplayAbilityTargetActor* TargetActor : SpawnedTargetActors)
+	{
+		if (TargetActor)
+		{
+			TargetActor->CancelTargeting();
+		}
+	}
+
+	SpawnedTargetActors.Empty();
+}
+
 // --------------------------------------------------------------------------
 
 void UAbilitySystemComponent::ServerSetReplicatedConfirm_Implementation(bool Confirmed)
 {
-	ReplicatedConfirmAbility  = Confirmed;
+	ReplicatedConfirmAbility = Confirmed;
 	ConfirmCallbacks.Broadcast();
 }
 
@@ -327,6 +354,8 @@ bool UAbilitySystemComponent::ServerSetReplicatedConfirm_Validate(bool Confirmed
 	return true;
 }
 
+// -------
+
 void UAbilitySystemComponent::ServerSetReplicatedTargetData_Implementation(FGameplayAbilityTargetDataHandle Confirmed)
 {
 	ReplicatedTargetData = Confirmed;
@@ -338,6 +367,20 @@ bool UAbilitySystemComponent::ServerSetReplicatedTargetData_Validate(FGameplayAb
 	return true;
 }
 
+// -------
+
+void UAbilitySystemComponent::ServerSetReplicatedTargetDataCancelled_Implementation()
+{
+	ReplicatedTargetDataCancelledDelegate.Broadcast();
+}
+
+bool UAbilitySystemComponent::ServerSetReplicatedTargetDataCancelled_Validate()
+{
+	return true;
+}
+
+// -------
+
 void UAbilitySystemComponent::SetTargetAbility(UGameplayAbility* NewTargetingAbility)
 {
 	TargetingAbility = NewTargetingAbility;
diff --git a/Engine/Source/Runtime/GameplayAbilities/Public/Abilities/GameplayAbilityTargetActor.h b/Engine/Source/Runtime/GameplayAbilities/Public/Abilities/GameplayAbilityTargetActor.h
--- a/Engine/Source/Runtime/GameplayAbilities/Public/Abilities/GameplayAbilityTargetActor.h
+++ b/Engine/Source/Runtime/GameplayAbilities/Public/Abilities/GameplayAbilityTargetActor.h
@@ -21,10 +21,17 @@ public:
 	UFUNCTION(BlueprintImplementableEvent, Category = "Targeting")
 	FGameplayAbilityTargetDataHandle StaticGetTargetData(UWorld * World, FGameplayAbilityActorInfo ActorInfo, FGameplayAbilityActivationInfo ActivationInfo);
 
-	// temp
-	virtual void StartTargeting(UGameplayAbility* Ability) { };
+	
+	virtual void StartTargeting(UGameplayAbility* Ability);
+
+	/** Outside code is saying 'stop and just give me what you have' */
+	virtual void ConfirmTargeting();
+
+	/** Outside code is saying 'stop everything and just forget about it' */
+	virtual void CancelTargeting();
 
 	// ------------------------------
 	
 	FAbilityTargetData	TargetDataReadyDelegate;
+	FAbilityTargetData	CanceledDelegate;
 };
\ No newline at end of file
diff --git a/Engine/Source/Runtime/GameplayAbilities/Public/Abilities/GameplayAbilityTargetActor_SingleLineTrace.h b/Engine/Source/Runtime/GameplayAbilities/Public/Abilities/GameplayAbilityTargetActor_SingleLineTrace.h
--- a/Engine/Source/Runtime/GameplayAbilities/Public/Abilities/GameplayAbilityTargetActor_SingleLineTrace.h
+++ b/Engine/Source/Runtime/GameplayAbilities/Public/Abilities/GameplayAbilityTargetActor_SingleLineTrace.h
@@ -16,19 +16,19 @@ public:
 
 	virtual void StartTargeting(UGameplayAbility* Ability);
 
-	//temp
-
-	void DoIt(TWeakObjectPtr<UGameplayAbility> Ability);
-
 	UFUNCTION()
 	void Confirm();
 
 	UFUNCTION()
 	void Cancel();
 
+	virtual void ConfirmTargeting();	
+
 	virtual void Tick(float DeltaSeconds) override;
 	
 	TWeakObjectPtr<UGameplayAbility> Ability;
 
 	bool bDebug;
+
+	bool bBindToConfirmCancelInputs;
 };
\ No newline at end of file
diff --git a/Engine/Source/Runtime/GameplayAbilities/Public/Abilities/Tasks/AbilityTask_WaitTargetData.h b/Engine/Source/Runtime/GameplayAbilities/Public/Abilities/Tasks/AbilityTask_WaitTargetData.h
--- a/Engine/Source/Runtime/GameplayAbilities/Public/Abilities/Tasks/AbilityTask_WaitTargetData.h
+++ b/Engine/Source/Runtime/GameplayAbilities/Public/Abilities/Tasks/AbilityTask_WaitTargetData.h
@@ -17,13 +17,19 @@ class UAbilityTask_WaitTargetData: public UAbilityTask
 	FWaitTargetDataDelegate	ValidData;
 
 	UPROPERTY(BlueprintAssignable)
-	FWaitTargetDataDelegate	Error;
+	FWaitTargetDataDelegate	Cancelled;
 
 	UFUNCTION()
 	void OnTargetDataReplicatedCallback(FGameplayAbilityTargetDataHandle Data);
 
+	UFUNCTION()
+	void OnTargetDataReplicatedCancelledCallback();
+
 	UFUNCTION()
 	void OnTargetDataReadyCallback(FGameplayAbilityTargetDataHandle Data);
+
+	UFUNCTION()
+	void OnTargetDataCancelledCallback(FGameplayAbilityTargetDataHandle Data);
 	
 	
 	UFUNCTION(BlueprintCallable, meta=(HidePin = "WorldContextObject", DefaultToSelf = "WorldContextObject", BlueprintInternalUseOnly = "true"), Category = "Abilities")
diff --git a/Engine/Source/Runtime/GameplayAbilities/Public/AbilitySystemComponent.h b/Engine/Source/Runtime/GameplayAbilities/Public/AbilitySystemComponent.h
--- a/Engine/Source/Runtime/GameplayAbilities/Public/AbilitySystemComponent.h
+++ b/Engine/Source/Runtime/GameplayAbilities/Public/AbilitySystemComponent.h
@@ -13,6 +13,8 @@ DECLARE_MULTICAST_DELEGATE(FAbilitySystemComponentPredictionKeyClear);
 /** Used to register callbacks to confirm/cancel input */
 DECLARE_DYNAMIC_MULTICAST_DELEGATE(FAbilityConfirmOrCancel);
 
+class AGameplayAbilityTargetActor;
+
 USTRUCT()
 struct GAMEPLAYABILITIES_API FAttributeDefaults
 {
@@ -323,12 +325,26 @@ class GAMEPLAYABILITIES_API UAbilitySystemComponent : public UActorComponent
 
 	virtual void BindToInputComponent(UInputComponent *InputComponent);
 
+	UFUNCTION(BlueprintCallable, Category="Abilities")
 	void InputConfirm();
+
+	UFUNCTION(BlueprintCallable, Category="Abilities")
 	void InputCancel();
 
 	FAbilityConfirmOrCancel	ConfirmCallbacks;
 	FAbilityConfirmOrCancel	CancelCallbacks;
 
+	UPROPERTY()
+	TArray<AGameplayAbilityTargetActor*>	SpawnedTargetActors;
+
+	/** Any active targeting actors will be told to stop and return current targeting data */
+	UFUNCTION(BlueprintCallable, Category = "Abilities")
+	void TargetConfirm();
+
+	/** Any active targeting actors will be stopped and cancelled, not returning any targeting data */
+	UFUNCTION(BlueprintCallable, Category = "Abilities")
+	void TargetCancel();
+
 	// ----------------------------------------------------------------------------------------------------------------
 
 	/** This is temp for testing. We want to think a bit more on the API for outside stuff activating abilities */
@@ -361,6 +377,10 @@ class GAMEPLAYABILITIES_API UAbilitySystemComponent : public UActorComponent
 	UFUNCTION(Server, reliable, WithValidation)
 	void ServerSetReplicatedTargetData(FGameplayAbilityTargetDataHandle ReplicatedTargetData);
 
+	UFUNCTION(Server, reliable, WithValidation)
+	void ServerSetReplicatedTargetDataCancelled();
+	
+
 	void SetTargetAbility(UGameplayAbility* NewTargetingAbility);
 
 	void ConsumeAbilityConfirm();
@@ -375,8 +395,12 @@ class GAMEPLAYABILITIES_API UAbilitySystemComponent : public UActorComponent
 
 	FGameplayAbilityTargetDataHandle ReplicatedTargetData;
 
+	/** ReplicatedTargetData was received */
 	FAbilityTargetData	ReplicatedTargetDataDelegate;
 
+	/** ReplicatedTargetData was 'cancelled' for this activation */
+	FAbilityConfirmOrCancel	ReplicatedTargetDataCancelledDelegate;
+
 private:
 
 	// --------------------------------------------
