remove K2Node_WaitOverlap

diff --git a/Engine/Source/Editor/GameplayAbilitiesEditor/Classes/K2Node_WaitOverlap.h b/Engine/Source/Editor/GameplayAbilitiesEditor/Classes/K2Node_WaitOverlap.h
deleted file mode 100644
diff --git a/Engine/Source/Editor/GameplayAbilitiesEditor/Private/K2Node_WaitOverlap.cpp b/Engine/Source/Editor/GameplayAbilitiesEditor/Private/K2Node_WaitOverlap.cpp
deleted file mode 100644
diff --git a/Engine/Source/Runtime/GameplayAbilities/Private/Abilities/Tasks/AbilityTask_WaitOverlap.cpp b/Engine/Source/Runtime/GameplayAbilities/Private/Abilities/Tasks/AbilityTask_WaitOverlap.cpp
--- a/Engine/Source/Runtime/GameplayAbilities/Private/Abilities/Tasks/AbilityTask_WaitOverlap.cpp
+++ b/Engine/Source/Runtime/GameplayAbilities/Private/Abilities/Tasks/AbilityTask_WaitOverlap.cpp
@@ -33,4 +33,40 @@ void UAbilityTask_WaitOverlap::OnHitCallback(AActor* OtherActor, UPrimitiveCompo
 		// We are done. Kill us so we don't keep getting broadcast messages
 		MarkPendingKill();
 	}
+}
+
+/**
+*	Need:
+*	-Easy way to specify which primitive components should be used for this overlap test
+*	-Easy way to specify which types of actors/collision overlaps that we care about/want to block on
+*/
+
+UAbilityTask_WaitOverlap* UAbilityTask_WaitOverlap::WaitForOverlap(class UObject* WorldContextObject)
+{
+	check(WorldContextObject);
+
+	UAbilityTask_WaitOverlap * MyObj = NewObject<UAbilityTask_WaitOverlap>();
+	UGameplayAbility* ThisAbility = CastChecked<UGameplayAbility>(WorldContextObject);
+	MyObj->Ability = ThisAbility;
+
+	return MyObj;
+}
+
+void UAbilityTask_WaitOverlap::Activate()
+{
+	if (Ability.IsValid())
+	{
+		AActor * ActorOwner = Cast<AActor>(Ability->GetOuter());
+
+		// TEMP - we are just using root component's collision. A real system will need more data to specify which component to use
+		UPrimitiveComponent * PrimComponent = Cast<UPrimitiveComponent>(ActorOwner->GetRootComponent());
+		if (!PrimComponent)
+		{
+			PrimComponent = ActorOwner->FindComponentByClass<UPrimitiveComponent>();
+		}
+		check(PrimComponent);
+
+		PrimComponent->OnComponentBeginOverlap.AddDynamic(this, &UAbilityTask_WaitOverlap::OnOverlapCallback);
+		PrimComponent->OnComponentHit.AddDynamic(this, &UAbilityTask_WaitOverlap::OnHitCallback);
+	}
 }
\ No newline at end of file
diff --git a/Engine/Source/Runtime/GameplayAbilities/Private/Abilities/Tasks/AbilityTask_WaitTargetData.cpp b/Engine/Source/Runtime/GameplayAbilities/Private/Abilities/Tasks/AbilityTask_WaitTargetData.cpp
--- a/Engine/Source/Runtime/GameplayAbilities/Private/Abilities/Tasks/AbilityTask_WaitTargetData.cpp
+++ b/Engine/Source/Runtime/GameplayAbilities/Private/Abilities/Tasks/AbilityTask_WaitTargetData.cpp
@@ -20,6 +20,11 @@ UAbilityTask_WaitTargetData* UAbilityTask_WaitTargetData::WaitTargetData(UObject
 	return MyObj;	
 }
 
+UAbilityTask_WaitTargetData* UAbilityTask_WaitTargetData::CreateWaitTargetData(UObject* WorldContextObject, TSubclassOf<AGameplayAbilityTargetActor> InTargetClass)
+{
+	return WaitTargetData(WorldContextObject, InTargetClass);
+}
+
 void UAbilityTask_WaitTargetData::Activate()
 {
 	if (Ability.IsValid())
diff --git a/Engine/Source/Runtime/GameplayAbilities/Private/AbilitySystemBlueprintLibrary.cpp b/Engine/Source/Runtime/GameplayAbilities/Private/AbilitySystemBlueprintLibrary.cpp
--- a/Engine/Source/Runtime/GameplayAbilities/Private/AbilitySystemBlueprintLibrary.cpp
+++ b/Engine/Source/Runtime/GameplayAbilities/Private/AbilitySystemBlueprintLibrary.cpp
@@ -114,37 +114,3 @@ UAbilityTask_WaitMovementModeChange* UAbilitySystemBlueprintLibrary::CreateWaitM
 	return NULL;
 }
 
-/**
- *	Need:
- *	-Easy way to specify which primitive components should be used for this overlap test
- *	-Easy way to specify which types of actors/collision overlaps that we care about/want to block on
- */
-
-UAbilityTask_WaitOverlap* UAbilitySystemBlueprintLibrary::CreateWaitOverlap(class UObject* WorldContextObject, EMovementMode NewMode)
-{
-	check(WorldContextObject);
-	UGameplayAbility* Ability = CastChecked<UGameplayAbility>(WorldContextObject);
-	if (Ability)
-	{
-		AActor * ActorOwner = Cast<AActor>(Ability->GetOuter());
-
-		UAbilityTask_WaitOverlap * MyObj = NULL;
-		MyObj = NewObject<UAbilityTask_WaitOverlap>();
-
-
-		// TEMP - we are just using root component's collision. A real system will need more data to specify which component to use
-		UPrimitiveComponent * PrimComponent = Cast<UPrimitiveComponent>(ActorOwner->GetRootComponent());
-		if (!PrimComponent)
-		{			
-			PrimComponent = ActorOwner->FindComponentByClass<UPrimitiveComponent>();
-		}
-		check(PrimComponent);
-
-		PrimComponent->OnComponentBeginOverlap.AddDynamic(MyObj, &UAbilityTask_WaitOverlap::OnOverlapCallback);
-		PrimComponent->OnComponentHit.AddDynamic(MyObj, &UAbilityTask_WaitOverlap::OnHitCallback);
-
-		return MyObj;
-	}
-
-	return NULL;
-}
\ No newline at end of file
diff --git a/Engine/Source/Runtime/GameplayAbilities/Private/AbilitySystemComponent_Abilities.cpp b/Engine/Source/Runtime/GameplayAbilities/Private/AbilitySystemComponent_Abilities.cpp
--- a/Engine/Source/Runtime/GameplayAbilities/Private/AbilitySystemComponent_Abilities.cpp
+++ b/Engine/Source/Runtime/GameplayAbilities/Private/AbilitySystemComponent_Abilities.cpp
@@ -299,11 +299,9 @@ void UAbilitySystemComponent::BindToInputComponent(UInputComponent *InputCompone
 
 void UAbilitySystemComponent::InputConfirm()
 {
-	if (GetOwnerRole() != ROLE_Authority)
+	if (GetOwnerRole() != ROLE_Authority && ConfirmCallbacks.IsBound())
 	{
-		// Tell the server we confirmed input
-		// FIXME TODO: we don't need to do this everytime! just when there is an activae ability waiting for this!
-		// 
+		// Tell the server we confirmed input.
 		ServerSetReplicatedConfirm(true);
 	}
 	
diff --git a/Engine/Source/Runtime/GameplayAbilities/Public/Abilities/Tasks/AbilityTask_WaitOverlap.h b/Engine/Source/Runtime/GameplayAbilities/Public/Abilities/Tasks/AbilityTask_WaitOverlap.h
--- a/Engine/Source/Runtime/GameplayAbilities/Public/Abilities/Tasks/AbilityTask_WaitOverlap.h
+++ b/Engine/Source/Runtime/GameplayAbilities/Public/Abilities/Tasks/AbilityTask_WaitOverlap.h
@@ -23,6 +23,12 @@ class UAbilityTask_WaitOverlap : public UAbilityTask
 	UFUNCTION()
 	void OnHitCallback(AActor* OtherActor, UPrimitiveComponent* OtherComp, FVector NormalImpulse, const FHitResult& Hit);
 
+	virtual void Activate() override;
+
+	/** Wait until an overlap occurs. This will need to be better fleshed out so we can specify game specific collision requirements */
+	UFUNCTION(BlueprintCallable, meta = (HidePin = "WorldContextObject", DefaultToSelf = "WorldContextObject", BlueprintInternalUseOnly = "TRUE"))
+	static UAbilityTask_WaitOverlap* WaitForOverlap(UObject* WorldContextObject);
+
 public:
 	
 };
\ No newline at end of file
diff --git a/Engine/Source/Runtime/GameplayAbilities/Public/Abilities/Tasks/AbilityTask_WaitTargetData.h b/Engine/Source/Runtime/GameplayAbilities/Public/Abilities/Tasks/AbilityTask_WaitTargetData.h
--- a/Engine/Source/Runtime/GameplayAbilities/Public/Abilities/Tasks/AbilityTask_WaitTargetData.h
+++ b/Engine/Source/Runtime/GameplayAbilities/Public/Abilities/Tasks/AbilityTask_WaitTargetData.h
@@ -35,6 +35,9 @@ class UAbilityTask_WaitTargetData: public UAbilityTask
 	UFUNCTION(BlueprintCallable, meta=(HidePin = "WorldContextObject", DefaultToSelf = "WorldContextObject", BlueprintInternalUseOnly = "true"), Category = "Abilities")
 	static UAbilityTask_WaitTargetData* WaitTargetData(UObject* WorldContextObject, TSubclassOf<AGameplayAbilityTargetActor> TargetClass);
 
+	UFUNCTION(BlueprintCallable, meta = (HidePin = "WorldContextObject", DefaultToSelf = "WorldContextObject", BlueprintInternalUseOnly = "true"), Category = "Abilities")
+	static UAbilityTask_WaitTargetData* CreateWaitTargetData(UObject* WorldContextObject, TSubclassOf<AGameplayAbilityTargetActor> TargetClass);
+
 	virtual void Activate() OVERRIDE;
 
 	TSubclassOf<AGameplayAbilityTargetActor> TargetClass;
diff --git a/Engine/Source/Runtime/GameplayAbilities/Public/AbilitySystemBlueprintLibrary.h b/Engine/Source/Runtime/GameplayAbilities/Public/AbilitySystemBlueprintLibrary.h
--- a/Engine/Source/Runtime/GameplayAbilities/Public/AbilitySystemBlueprintLibrary.h
+++ b/Engine/Source/Runtime/GameplayAbilities/Public/AbilitySystemBlueprintLibrary.h
@@ -35,7 +35,5 @@ class UAbilitySystemBlueprintLibrary : public UBlueprintFunctionLibrary
 
 	UFUNCTION(BlueprintCallable, meta = (HidePin = "WorldContextObject", DefaultToSelf = "WorldContextObject", BlueprintInternalUseOnly = "TRUE"))
 	static UAbilityTask_WaitMovementModeChange* CreateWaitMovementModeChange(UObject* WorldContextObject, EMovementMode NewMode);
-
-	UFUNCTION(BlueprintCallable, meta = (HidePin = "WorldContextObject", DefaultToSelf = "WorldContextObject", BlueprintInternalUseOnly = "TRUE"))
-	static UAbilityTask_WaitOverlap* CreateWaitOverlap(UObject* WorldContextObject, EMovementMode NewMode);
+	
 };
\ No newline at end of file
