Ability System updates

diff --git a/Engine/Source/Runtime/GameplayAbilities/Private/Abilities/GameplayAbility.cpp b/Engine/Source/Runtime/GameplayAbilities/Private/Abilities/GameplayAbility.cpp
--- a/Engine/Source/Runtime/GameplayAbilities/Private/Abilities/GameplayAbility.cpp
+++ b/Engine/Source/Runtime/GameplayAbilities/Private/Abilities/GameplayAbility.cpp
@@ -36,6 +36,26 @@ UGameplayAbility::UGameplayAbility(const class FPostConstructInitializePropertie
 	}
 }
 
+void UGameplayAbility::PostNetInit()
+{
+	/** We were dynamically spawned from replication - we need to init a currentactorinfo by looking at outer.
+	 *  This may need to be updated further if we start having abilities live on different outers than player AbilitySystemComponents.
+	 */
+	
+	if (CurrentActorInfo == NULL)
+	{
+		AActor *OwnerActor = Cast<AActor>(GetOuter());
+		if (ensure(OwnerActor))
+		{
+			UAbilitySystemComponent* AbilitySystemComponent = OwnerActor->FindComponentByClass<UAbilitySystemComponent>();
+			if (ensure(AbilitySystemComponent))
+			{
+				CurrentActorInfo = AbilitySystemComponent->AbilityActorInfo.Get();
+			}
+		}
+	}
+}
+
 bool UGameplayAbility::IsSupportedForNetworking() const
 {
 	/**
@@ -217,13 +237,18 @@ bool UGameplayAbility::TryActivateAbility(const FGameplayAbilityActorInfo* Actor
 
 void UGameplayAbility::EndAbility(const FGameplayAbilityActorInfo* ActorInfo)
 {
+	// MarkPending kill if we are instance per execution
 	if (InstancingPolicy == EGameplayAbilityInstancingPolicy::InstancedPerExecution)
 	{
-		MarkPendingKill();
+		// If not replicated, or if authority (clients should not delete subobjects on their own)
+		if (ReplicationPolicy == EGameplayAbilityReplicationPolicy::ReplicateNone || ActorInfo->IsNetAuthority())
+		{
+			MarkPendingKill();
+		}
 	}
 
-	// Remove from owning AbilitySystemComponent?
-	// generic way of releasing all callbacks!
+	// -Remove from owning AbilitySystemComponent
+	// -generic way of releasing all callbacks
 }
 
 void UGameplayAbility::ActivateAbility(const FGameplayAbilityActorInfo* ActorInfo, const FGameplayAbilityActivationInfo ActivationInfo)
@@ -422,6 +447,8 @@ void UGameplayAbility::MontageBranchPoint_AbilityDecisionStart(const FGameplayAb
 
 void UGameplayAbility::ClientActivateAbilitySucceed_Implementation(int32 PredictionKey)
 {
+	PostNetInit();
+
 	// Child classes can't override ClientActivateAbilitySucceed_Implementation, so just call ClientActivateAbilitySucceed_Internal
 	ClientActivateAbilitySucceed_Internal(PredictionKey);
 }
@@ -484,6 +511,20 @@ bool FGameplayAbilityActorInfo::IsLocallyControlled() const
 	return false;
 }
 
+bool FGameplayAbilityActorInfo::IsNetAuthority() const
+{
+	if (Actor.IsValid())
+	{
+		return (Actor->Role == ROLE_Authority);
+	}
+
+	// If we encounter issues with this being called before or after the owning actor is destroyed,
+	// we may need to cache off the authority (or look for it on some global/world state).
+
+	ensure(false);
+	return false;
+}
+
 void FGameplayAbilityActivationInfo::GeneratePredictionKey(UAbilitySystemComponent * Component) const
 {
 	check(Component);
diff --git a/Engine/Source/Runtime/GameplayAbilities/Private/AbilitySystemComponent.cpp b/Engine/Source/Runtime/GameplayAbilities/Private/AbilitySystemComponent.cpp
--- a/Engine/Source/Runtime/GameplayAbilities/Private/AbilitySystemComponent.cpp
+++ b/Engine/Source/Runtime/GameplayAbilities/Private/AbilitySystemComponent.cpp
@@ -220,6 +220,11 @@ bool UAbilitySystemComponent::HasAnyTags(FGameplayTagContainer &Tags)
 	return ActiveGameplayEffects.HasAnyTags(Tags);
 }
 
+bool UAbilitySystemComponent::HasTag(FGameplayTag Tag)
+{
+	return ActiveGameplayEffects.HasTag(Tag);
+}
+
 bool UAbilitySystemComponent::HasAllTags(FGameplayTagContainer &Tags)
 {
 	return ActiveGameplayEffects.HasAllTags(Tags);
@@ -523,7 +528,7 @@ bool UAbilitySystemComponent::ReplicateSubobjects(class UActorChannel *Channel,
 
 	for (UGameplayAbility* Ability : ReplicatedInstancedAbilities)
 	{
-		if (Ability)
+		if (Ability && !Ability->HasAnyFlags(RF_PendingKill))
 		{
 			WroteSomething |= Channel->ReplicateSubobject(Ability, *Bunch, *RepFlags);
 		}
diff --git a/Engine/Source/Runtime/GameplayAbilities/Private/GameplayCueView.cpp b/Engine/Source/Runtime/GameplayAbilities/Private/GameplayCueView.cpp
--- a/Engine/Source/Runtime/GameplayAbilities/Private/GameplayCueView.cpp
+++ b/Engine/Source/Runtime/GameplayAbilities/Private/GameplayCueView.cpp
@@ -8,6 +8,7 @@
 #include "Particles/ParticleLODLevel.h"
 #include "Particles/ParticleModuleRequired.h"
 #include "Particles/ParticleSystemComponent.h"
+#include "GameplayCueInterface.h"
 #include "GameplayTagsModule.h"
 
 UGameplayCueView::UGameplayCueView(const class FPostConstructInitializeProperties& PCIP)
@@ -29,7 +30,7 @@ FGameplayCueViewInfo * FGameplayCueHandler::GetBestMatchingView(EGameplayCueEven
 		{
 			for (FGameplayCueViewInfo & View : Def->Views)
 			{
-				if (View.CueType == Type && View.Tags.HasTag(Tag, EGameplayTagMatchType::IncludeParentTags, EGameplayTagMatchType::Explicit))
+				if (View.CueType == Type && View.Tags.HasTag(Tag, EGameplayTagMatchType::Explicit, EGameplayTagMatchType::Explicit))
 				{
 					return &View;
 				}
@@ -45,7 +46,7 @@ void FGameplayCueHandler::GameplayCueActivated(const FGameplayTagContainer & Gam
 	check(Owner);
 	for (auto TagIt = GameplayCueTags.CreateConstIterator(); TagIt; ++TagIt)
 	{
-		if (FGameplayCueViewInfo * View = GetBestMatchingView(EGameplayCueEvent::Applied, *TagIt))
+		if (FGameplayCueViewInfo * View = GetBestMatchingView(EGameplayCueEvent::OnActive, *TagIt))
 		{
 			View->SpawnViewEffects(Owner, NULL, InstigatorContext);
 		}
@@ -61,6 +62,12 @@ void FGameplayCueHandler::GameplayCueExecuted(const FGameplayTagContainer & Game
 		{
 			View->SpawnViewEffects(Owner, NULL, InstigatorContext);
 		}
+
+		
+		{
+			IGameplayCueInterface::Execute_HandleGameplayCue(Owner, *TagIt, EGameplayCueEvent::Executed);
+		}
+
 	}
 }
 
@@ -76,13 +83,15 @@ void FGameplayCueHandler::GameplayCueAdded(const FGameplayTagContainer & Gamepla
 		ClearEffects(Effects);
 		check(Effects.Num() == 0);
 
-		if (FGameplayCueViewInfo * View = GetBestMatchingView(EGameplayCueEvent::Added, *TagIt))
+		if (FGameplayCueViewInfo * View = GetBestMatchingView(EGameplayCueEvent::WhileActive, *TagIt))
 		{
-			View->SpawnViewEffects(Owner, NULL, InstigatorContext);
-
 			TSharedPtr<FGameplayCueViewEffects> SpawnedEffects = View->SpawnViewEffects(Owner, &SpawnedObjects, InstigatorContext);
 			Effects.Add(SpawnedEffects);
 		}
+
+		{
+			IGameplayCueInterface::Execute_HandleGameplayCue(Owner, *TagIt, EGameplayCueEvent::WhileActive);
+		}
 	}
 }
 
@@ -105,6 +114,10 @@ void FGameplayCueHandler::GameplayCueRemoved(const FGameplayTagContainer & Gamep
 		{
 			View->SpawnViewEffects(Owner, NULL, InstigatorContext);
 		}
+
+		{
+			IGameplayCueInterface::Execute_HandleGameplayCue(Owner, *TagIt, EGameplayCueEvent::Removed);
+		}
 	}
 }
 
diff --git a/Engine/Source/Runtime/GameplayAbilities/Private/GameplayEffect.cpp b/Engine/Source/Runtime/GameplayAbilities/Private/GameplayEffect.cpp
--- a/Engine/Source/Runtime/GameplayAbilities/Private/GameplayEffect.cpp
+++ b/Engine/Source/Runtime/GameplayAbilities/Private/GameplayEffect.cpp
@@ -1842,24 +1842,7 @@ bool FActiveGameplayEffectsContainer::HasAllTags(FGameplayTagContainer &Tags)
 
 	for (auto It = Tags.CreateConstIterator(); It;  ++It)
 	{
-		bool Found = false;
-		for (FActiveGameplayEffect &ActiveEffect : GameplayEffects)
-		{
-			// Fixme: check stacking rules!
-
-			/**
-			 *	The desired behavior here is:
-			 *		"Do we have tag a.b.c and the GameplayEffect has a.b.c.d -> match!"
-			 *		"Do we have tag a.b.c.d. and the GameplayEffect has a.b.c -> no match!"
-			 */
-
-			if (ActiveEffect.Spec.Def->OwnedTagsContainer.HasTag(*It, EGameplayTagMatchType::IncludeParentTags, EGameplayTagMatchType::Explicit))
-			{
-				Found = true;
-				break;
-			}
-		}
-		if (!Found)
+		if (!HasTag(*It))
 		{
 			return false;
 		}
@@ -1868,6 +1851,24 @@ bool FActiveGameplayEffectsContainer::HasAllTags(FGameplayTagContainer &Tags)
 	return true;
 }
 
+bool FActiveGameplayEffectsContainer::HasTag(const FGameplayTag Tag)
+{
+	for (FActiveGameplayEffect &ActiveEffect : GameplayEffects)
+	{
+		/**
+		 *	The desired behavior here is:
+		 *		"Do we have tag a.b.c and the GameplayEffect has a.b.c.d -> match!"
+		 *		"Do we have tag a.b.c.d. and the GameplayEffect has a.b.c -> no match!"
+		 */
+
+		if (ActiveEffect.Spec.Def->OwnedTagsContainer.HasTag(Tag, EGameplayTagMatchType::IncludeParentTags, EGameplayTagMatchType::Explicit))
+		{
+			return true;
+		}
+	}
+	return false;
+}
+
 bool FActiveGameplayEffectsContainer::CanApplyAttributeModifiers(const UGameplayEffect *GameplayEffect, float Level, AActor *Instigator)
 {
 	SCOPE_CYCLE_COUNTER(STAT_GameplayEffectsCanApplyAttributeModifiers);
diff --git a/Engine/Source/Runtime/GameplayAbilities/Public/Abilities/GameplayAbility.h b/Engine/Source/Runtime/GameplayAbilities/Public/Abilities/GameplayAbility.h
--- a/Engine/Source/Runtime/GameplayAbilities/Public/Abilities/GameplayAbility.h
+++ b/Engine/Source/Runtime/GameplayAbilities/Public/Abilities/GameplayAbility.h
@@ -168,9 +168,11 @@ public:
 		return GetOuter()->GetWorld();
 	}
 
-	int32 GetFunctionCallspace(UFunction* Function, void* Parameters, FFrame* Stack);
+	int32 GetFunctionCallspace(UFunction* Function, void* Parameters, FFrame* Stack) override;
 
-	bool CallRemoteFunction(UFunction* Function, void* Parameters, FOutParmRec* OutParms, FFrame* Stack);
+	bool CallRemoteFunction(UFunction* Function, void* Parameters, FOutParmRec* OutParms, FFrame* Stack) override;
+
+	void PostNetInit();
 
 protected:
 		
diff --git a/Engine/Source/Runtime/GameplayAbilities/Public/Abilities/GameplayAbilityTypes.h b/Engine/Source/Runtime/GameplayAbilities/Public/Abilities/GameplayAbilityTypes.h
--- a/Engine/Source/Runtime/GameplayAbilities/Public/Abilities/GameplayAbilityTypes.h
+++ b/Engine/Source/Runtime/GameplayAbilities/Public/Abilities/GameplayAbilityTypes.h
@@ -175,6 +175,8 @@ struct GAMEPLAYABILITIES_API FGameplayAbilityActorInfo
 
 	bool IsLocallyControlled() const;
 
+	bool IsNetAuthority() const;
+
 	virtual void InitFromActor(AActor *Actor);
 };
 
diff --git a/Engine/Source/Runtime/GameplayAbilities/Public/AbilitySystemComponent.h b/Engine/Source/Runtime/GameplayAbilities/Public/AbilitySystemComponent.h
--- a/Engine/Source/Runtime/GameplayAbilities/Public/AbilitySystemComponent.h
+++ b/Engine/Source/Runtime/GameplayAbilities/Public/AbilitySystemComponent.h
@@ -207,6 +207,9 @@ class GAMEPLAYABILITIES_API UAbilitySystemComponent : public UActorComponent
 	UFUNCTION(BlueprintCallable, Category=GameplayEffects)
 	bool HasAnyTags(FGameplayTagContainer &Tags);
 
+	UFUNCTION(BlueprintCallable, Category = GameplayEffects)
+	bool HasTag(const FGameplayTag Tag);
+
 	UFUNCTION(BlueprintCallable, Category = GameplayEffects)
 	bool HasAllTags(FGameplayTagContainer &Tags);
 	
diff --git a/Engine/Source/Runtime/GameplayAbilities/Public/GameplayCueInterface.h b/Engine/Source/Runtime/GameplayAbilities/Public/GameplayCueInterface.h
--- a/Engine/Source/Runtime/GameplayAbilities/Public/GameplayCueInterface.h
+++ b/Engine/Source/Runtime/GameplayAbilities/Public/GameplayCueInterface.h
@@ -6,6 +6,8 @@
 #include "GameplayEffect.h"
 #include "GameplayCueInterface.generated.h"
 
+
+
 /** Interface for actors that wish to handle GameplayCue events from GameplayEffects */
 UINTERFACE(MinimalAPI)
 class UGameplayCueInterface: public UInterface
@@ -17,18 +19,27 @@ class GAMEPLAYABILITIES_API IGameplayCueInterface
 {
 	GENERATED_IINTERFACE_BODY()
 
-	// EVENT
+	UFUNCTION(BlueprintImplementableEvent, Category = "GameplayCue")
+	void HandleGameplayCue(FGameplayTag Tag, EGameplayCueEvent::Type CueType);
+
+	// EVENT -----------------------
+
+	/** EGameplayCueEvent::OnActive - Called when GameplayCue is really activated (not called again in relevancy, Join In Progress situations) */
 	UFUNCTION()
 	virtual void GameplayCueActivated(const FGameplayTagContainer & GameplayCueTags, float NormalizedMagnitude, const FGameplayEffectInstigatorContext InstigatorContext) = 0;
 
+	/** EGameplayCueEvent::Executed - Called when GameplayCue is executed (instant effect or periodic effect ticks) */
 	UFUNCTION()
 	virtual void GameplayCueExecuted(const FGameplayTagContainer & GameplayCueTags, float NormalizedMagnitude, const FGameplayEffectInstigatorContext InstigatorContext) = 0;
 
-
-	// STATE
+	/** EGameplayCueEvent::Removed - Called when GameplayCue is removed */
 	UFUNCTION()
-	virtual void GameplayCueAdded(const FGameplayTagContainer & GameplayCueTags, float NormalizedMagnitude, const FGameplayEffectInstigatorContext InstigatorContext) = 0;
+	virtual void GameplayCueRemoved(const FGameplayTagContainer & GameplayCueTags, float NormalizedMagnitude, const FGameplayEffectInstigatorContext InstigatorContext) = 0;
+
+
+	// STATE -----------------------
 
+	/** EGameplayCueEvent::WhileActive - Called when GameplayCue is added locally (e.g., casted for the first time or in relevancy / Join in Progress situations) */
 	UFUNCTION()
-	virtual void GameplayCueRemoved(const FGameplayTagContainer & GameplayCueTags, float NormalizedMagnitude, const FGameplayEffectInstigatorContext InstigatorContext) = 0;
+	virtual void GameplayCueAdded(const FGameplayTagContainer & GameplayCueTags, float NormalizedMagnitude, const FGameplayEffectInstigatorContext InstigatorContext) = 0;
 };
\ No newline at end of file
diff --git a/Engine/Source/Runtime/GameplayAbilities/Public/GameplayCueView.h b/Engine/Source/Runtime/GameplayAbilities/Public/GameplayCueView.h
--- a/Engine/Source/Runtime/GameplayAbilities/Public/GameplayCueView.h
+++ b/Engine/Source/Runtime/GameplayAbilities/Public/GameplayCueView.h
@@ -13,18 +13,6 @@
 	It is inevitable that specific actors and games will have to handle GameplayCues in their own way.
 */
 
-UENUM(BlueprintType)
-namespace EGameplayCueEvent
-{
-	enum Type
-	{
-		Applied,
-		Added,
-		Executed,
-		Removed
-	};
-}
-
 USTRUCT()
 struct FGameplayCueViewEffects
 {
@@ -57,7 +45,10 @@ struct FGameplayCueViewInfo
 	TEnumAsByte<EGameplayCueEvent::Type> CueType;
 
 	// This needs to be fleshed out more:
+	UPROPERTY(EditAnywhere, BlueprintReadWrite, Category = GameplayCue)
 	bool	InstigatorLocalOnly;
+
+	UPROPERTY(EditAnywhere, BlueprintReadWrite, Category = GameplayCue)
 	bool	TargetLocalOnly;
 
 	// -----------------------------
diff --git a/Engine/Source/Runtime/GameplayAbilities/Public/GameplayEffect.h b/Engine/Source/Runtime/GameplayAbilities/Public/GameplayEffect.h
--- a/Engine/Source/Runtime/GameplayAbilities/Public/GameplayEffect.h
+++ b/Engine/Source/Runtime/GameplayAbilities/Public/GameplayEffect.h
@@ -1,4 +1,4 @@
-// Copyright 1998-2013 Epic Games, Inc. All Rights Reserved.
+// Copyright 1998-2014 Epic Games, Inc. All Rights Reserved.
 
 #pragma once
 
@@ -6,6 +6,7 @@
 #include "GameplayTagAssetInterface.h"
 #include "AttributeSet.h"
 #include "TimerManager.h"
+#include "GameplayEffectTypes.h"
 #include "GameplayEffect.generated.h"
 
 DECLARE_DELEGATE_OneParam(FOnGameplayAttributeEffectExecuted, struct FGameplayModifierEvaluatedData &);
@@ -25,126 +26,7 @@ struct FActiveGameplayEffect;
 class UGameplayEffect;
 class UAbilitySystemComponent;
 
-UENUM(BlueprintType)
-namespace EGameplayModOp
-{
-	enum Type
-	{		
-		// Numeric
-		Additive = 0		UMETA(DisplayName="Add"),
-		Multiplicitive		UMETA(DisplayName="Multiply"),
-		Division			UMETA(DisplayName="Divide"),
-
-		// Other
-		Override 			UMETA(DisplayName="Override"),	// This should always be the first non numeric ModOp
-		Callback			UMETA(DisplayName="Custom"),
-
-		// This must always be at the end
-		Max					UMETA(DisplayName="Invalid")
-	};
-}
-
-/**
- * Tells us what thing a GameplayModifier modifies.
- */
-UENUM(BlueprintType)
-namespace EGameplayMod
-{
-	enum Type
-	{
-		Attribute = 0,		// Modifies this Attributes
-		OutgoingGE,			// Modifies Outgoing Gameplay Effects (that modify this Attribute)
-		IncomingGE,			// Modifies Incoming Gameplay Effects (that modify this Attribute)
-		ActiveGE,			// Modifies currently active Gameplay Effects
-
-		// This must always be at the end
-		Max					UMETA(DisplayName="Invalid")
-	};
-}
-
-/**
- * Tells us what thing a GameplayEffect provides immunity to. This must mirror the values in EGameplayMod
- */
-UENUM(BlueprintType)
-namespace EGameplayImmunity
-{
-	enum Type
-	{
-		None = 0,			// Does not provide immunity
-		OutgoingGE,			// Provides immunity to outgoing GEs
-		IncomingGE,			// Provides immunity to incoming GEs
-		ActiveGE,			// Provides immunity from currently active GEs
-
-		// This must always be at the end
-		Max					UMETA(DisplayName="Invalid")
-	};
-}
-
-static_assert((int32)EGameplayMod::OutgoingGE == (int32)EGameplayImmunity::OutgoingGE, "EGameplayMod::OutgoingGE and EGameplayImmunity::OutgoingGE must match. Did you forget to modify one of them?");
-static_assert((int32)EGameplayMod::IncomingGE == (int32)EGameplayImmunity::IncomingGE, "EGameplayMod::IncomingGE and EGameplayImmunity::IncomingGE must match. Did you forget to modify one of them?");
-static_assert((int32)EGameplayMod::ActiveGE == (int32)EGameplayImmunity::ActiveGE, "EGameplayMod::ActiveGE and EGameplayImmunity::ActiveGE must match. Did you forget to modify one of them?");
-static_assert((int32)EGameplayMod::Max == (int32)EGameplayImmunity::Max, "EGameplayMod and EGameplayImmunity must cover the same cases. Did you forget to modify one of them?");
-
-/**
- * Tells us what a GameplayEffect modifies when being applied to another GameplayEffect
- */
-UENUM(BlueprintType)
-namespace EGameplayModEffect
-{
-	enum Type
-	{
-		Magnitude			= 0x01,		// Modifies magnitude of a GameplayEffect (Always default for Attribute mod)
-		Duration			= 0x02,		// Modifies duration of a GameplayEffect
-		ChanceApplyTarget	= 0x04,		// Modifies chance to apply GameplayEffect to target
-		ChanceExecuteEffect	= 0x08,		// Modifies chance to execute GameplayEffect on GameplayEffect
-		LinkedGameplayEffect= 0x10,		// Adds a linked GameplayEffect to a GameplayEffect
-
-		// This must always be at the end
-		All					= 0xFF		UMETA(DisplayName="Invalid")
-	};
-}
-
-/**
- * Tells us how to handle copying gameplay effect when it is applied.
- *	Default means to use context - e.g, OutgoingGE is are always snapshots, IncomingGE is always Link
- *	AlwaysSnapshot vs AlwaysLink let mods themselves override
- */
-UENUM(BlueprintType)
-namespace EGameplayEffectCopyPolicy
-{
-	enum Type
-	{
-		Default = 0			UMETA(DisplayName="Default"),
-		AlwaysSnapshot		UMETA(DisplayName="AlwaysSnapshot"),
-		AlwaysLink			UMETA(DisplayName="AlwaysLink")
-	};
-}
 
-UENUM(BlueprintType)
-namespace EGameplayEffectStackingPolicy
-{
-	enum Type
-	{
-		Unlimited = 0		UMETA(DisplayName = "NoRule"),
-		Highest				UMETA(DisplayName = "Strongest"),
-		Lowest				UMETA(DisplayName = "Weakest"),
-		Replaces			UMETA(DisplayName = "MostRecent"),
-		Callback			UMETA(DisplayName = "Custom"),
-
-		// This must always be at the end
-		Max					UMETA(DisplayName = "Invalid")
-	};
-}
-
-FString EGameplayModOpToString( int32 Type );
-
-FString EGameplayModToString( int32 Type );
-
-FString EGameplayModEffectToString( int32 Type );
-
-FString EGameplayEffectCopyPolicyToString(int32 Type);
-
-FString EGameplayEffectStackingPolicyToString(int32 Type);
 
 USTRUCT()
 struct FGameplayModifierCallbacks
@@ -1392,6 +1274,8 @@ struct FActiveGameplayEffectsContainer : public FFastArraySerializer
 
 	bool HasAllTags(FGameplayTagContainer &Tags);
 
+	bool HasTag(const FGameplayTag Tag);
+
 	bool CanApplyAttributeModifiers(const UGameplayEffect *GameplayEffect, float Level, AActor *Instigator);
 	
 	TArray<float> GetActiveEffectsTimeRemaining(const FActiveGameplayEffectQuery Query) const;
diff --git a/Engine/Source/Runtime/GameplayAbilities/Public/GameplayEffectTypes.h b/Engine/Source/Runtime/GameplayAbilities/Public/GameplayEffectTypes.h
new file mode 100644
--- /dev/null
+++ b/Engine/Source/Runtime/GameplayAbilities/Public/GameplayEffectTypes.h
@@ -0,0 +1,138 @@
+// Copyright 1998-2014 Epic Games, Inc. All Rights Reserved.
+
+#pragma once
+
+#include "GameplayEffectTypes.generated.h"
+
+UENUM(BlueprintType)
+namespace EGameplayModOp
+{
+	enum Type
+	{		
+		// Numeric
+		Additive = 0		UMETA(DisplayName="Add"),
+		Multiplicitive		UMETA(DisplayName="Multiply"),
+		Division			UMETA(DisplayName="Divide"),
+
+		// Other
+		Override 			UMETA(DisplayName="Override"),	// This should always be the first non numeric ModOp
+		Callback			UMETA(DisplayName="Custom"),
+
+		// This must always be at the end
+		Max					UMETA(DisplayName="Invalid")
+	};
+}
+
+/**
+ * Tells us what thing a GameplayModifier modifies.
+ */
+UENUM(BlueprintType)
+namespace EGameplayMod
+{
+	enum Type
+	{
+		Attribute = 0,		// Modifies this Attributes
+		OutgoingGE,			// Modifies Outgoing Gameplay Effects (that modify this Attribute)
+		IncomingGE,			// Modifies Incoming Gameplay Effects (that modify this Attribute)
+		ActiveGE,			// Modifies currently active Gameplay Effects
+
+		// This must always be at the end
+		Max					UMETA(DisplayName="Invalid")
+	};
+}
+
+/**
+ * Tells us what thing a GameplayEffect provides immunity to. This must mirror the values in EGameplayMod
+ */
+UENUM(BlueprintType)
+namespace EGameplayImmunity
+{
+	enum Type
+	{
+		None = 0,			// Does not provide immunity
+		OutgoingGE,			// Provides immunity to outgoing GEs
+		IncomingGE,			// Provides immunity to incoming GEs
+		ActiveGE,			// Provides immunity from currently active GEs
+
+		// This must always be at the end
+		Max					UMETA(DisplayName="Invalid")
+	};
+}
+
+static_assert((int32)EGameplayMod::OutgoingGE == (int32)EGameplayImmunity::OutgoingGE, "EGameplayMod::OutgoingGE and EGameplayImmunity::OutgoingGE must match. Did you forget to modify one of them?");
+static_assert((int32)EGameplayMod::IncomingGE == (int32)EGameplayImmunity::IncomingGE, "EGameplayMod::IncomingGE and EGameplayImmunity::IncomingGE must match. Did you forget to modify one of them?");
+static_assert((int32)EGameplayMod::ActiveGE == (int32)EGameplayImmunity::ActiveGE, "EGameplayMod::ActiveGE and EGameplayImmunity::ActiveGE must match. Did you forget to modify one of them?");
+static_assert((int32)EGameplayMod::Max == (int32)EGameplayImmunity::Max, "EGameplayMod and EGameplayImmunity must cover the same cases. Did you forget to modify one of them?");
+
+/**
+ * Tells us what a GameplayEffect modifies when being applied to another GameplayEffect
+ */
+UENUM(BlueprintType)
+namespace EGameplayModEffect
+{
+	enum Type
+	{
+		Magnitude			= 0x01,		// Modifies magnitude of a GameplayEffect (Always default for Attribute mod)
+		Duration			= 0x02,		// Modifies duration of a GameplayEffect
+		ChanceApplyTarget	= 0x04,		// Modifies chance to apply GameplayEffect to target
+		ChanceExecuteEffect	= 0x08,		// Modifies chance to execute GameplayEffect on GameplayEffect
+		LinkedGameplayEffect= 0x10,		// Adds a linked GameplayEffect to a GameplayEffect
+
+		// This must always be at the end
+		All					= 0xFF		UMETA(DisplayName="Invalid")
+	};
+}
+
+/**
+ * Tells us how to handle copying gameplay effect when it is applied.
+ *	Default means to use context - e.g, OutgoingGE is are always snapshots, IncomingGE is always Link
+ *	AlwaysSnapshot vs AlwaysLink let mods themselves override
+ */
+UENUM(BlueprintType)
+namespace EGameplayEffectCopyPolicy
+{
+	enum Type
+	{
+		Default = 0			UMETA(DisplayName="Default"),
+		AlwaysSnapshot		UMETA(DisplayName="AlwaysSnapshot"),
+		AlwaysLink			UMETA(DisplayName="AlwaysLink")
+	};
+}
+
+UENUM(BlueprintType)
+namespace EGameplayEffectStackingPolicy
+{
+	enum Type
+	{
+		Unlimited = 0		UMETA(DisplayName = "NoRule"),
+		Highest				UMETA(DisplayName = "Strongest"),
+		Lowest				UMETA(DisplayName = "Weakest"),
+		Replaces			UMETA(DisplayName = "MostRecent"),
+		Callback			UMETA(DisplayName = "Custom"),
+
+		// This must always be at the end
+		Max					UMETA(DisplayName = "Invalid")
+	};
+}
+
+UENUM(BlueprintType)
+namespace EGameplayCueEvent
+{
+	enum Type
+	{
+		OnActive,
+		WhileActive,
+		Executed,
+		Removed
+	};
+}
+
+FString EGameplayModOpToString(int32 Type);
+
+FString EGameplayModToString(int32 Type);
+
+FString EGameplayModEffectToString(int32 Type);
+
+FString EGameplayEffectCopyPolicyToString(int32 Type);
+
+FString EGameplayEffectStackingPolicyToString(int32 Type);
\ No newline at end of file
