GameplayAbilities
-Updates to GameplayCues

diff --git a/Engine/Source/Editor/GameplayAbilitiesEditor/GameplayAbilitiesEditor.Build.cs b/Engine/Source/Editor/GameplayAbilitiesEditor/GameplayAbilitiesEditor.Build.cs
--- a/Engine/Source/Editor/GameplayAbilitiesEditor/GameplayAbilitiesEditor.Build.cs
+++ b/Engine/Source/Editor/GameplayAbilitiesEditor/GameplayAbilitiesEditor.Build.cs
@@ -20,6 +20,7 @@ namespace UnrealBuildTool.Rules
 					"CoreUObject",
 					"Engine",
 					"AssetTools",
+                    "GameplayTags",
 					"GameplayAbilities",
                     "InputCore",
 					"Slate",
diff --git a/Engine/Source/Runtime/GameplayAbilities/Private/Abilities/Tasks/AbilityTask_WaitConfirmCancel.cpp b/Engine/Source/Runtime/GameplayAbilities/Private/Abilities/Tasks/AbilityTask_WaitConfirmCancel.cpp
--- a/Engine/Source/Runtime/GameplayAbilities/Private/Abilities/Tasks/AbilityTask_WaitConfirmCancel.cpp
+++ b/Engine/Source/Runtime/GameplayAbilities/Private/Abilities/Tasks/AbilityTask_WaitConfirmCancel.cpp
@@ -15,13 +15,13 @@ UAbilityTask_WaitConfirmCancel::UAbilityTask_WaitConfirmCancel(const class FPost
 void UAbilityTask_WaitConfirmCancel::OnConfirmCallback()
 {
 	OnConfirm.Broadcast(true);
-	ASC->ConsumeAbilityConfirm();
+	ASC->ConsumeAbilityConfirmCancel();
 }
 
 void UAbilityTask_WaitConfirmCancel::OnCancelCallback()
 {
 	OnCancel.Broadcast(false);
-	ASC->ConsumeAbilityConfirm();
+	ASC->ConsumeAbilityConfirmCancel();
 }
 
 UAbilityTask_WaitConfirmCancel* UAbilityTask_WaitConfirmCancel::WaitConfirmCancel(UObject* WorldContextObject)
@@ -54,7 +54,11 @@ void UAbilityTask_WaitConfirmCancel::Activate()
 				// It could have already been confirmed (we may have been behind a client predicting this ability, so it could have got here
 				// before we actually executed this task)
 				OnConfirmCallback();
-				ASC->ConsumeAbilityConfirm();
+				return;
+			}
+			else if(ASC->ReplicatedCancelAbility)
+			{
+				OnCancelCallback();
 				return;
 			}
 		}
diff --git a/Engine/Source/Runtime/GameplayAbilities/Private/AbilitySystemBlueprintLibrary.cpp b/Engine/Source/Runtime/GameplayAbilities/Private/AbilitySystemBlueprintLibrary.cpp
--- a/Engine/Source/Runtime/GameplayAbilities/Private/AbilitySystemBlueprintLibrary.cpp
+++ b/Engine/Source/Runtime/GameplayAbilities/Private/AbilitySystemBlueprintLibrary.cpp
@@ -114,3 +114,26 @@ UAbilityTask_WaitMovementModeChange* UAbilitySystemBlueprintLibrary::CreateWaitM
 	return NULL;
 }
 
+// -------------------------------------------------------------------------------------
+
+
+bool UAbilitySystemBlueprintLibrary::IsInstigatorLocallyControlled(FGameplayCueParameters Parameters)
+{
+	return Parameters.InstigatorContext.IsLocallyControlled();
+}
+
+
+FHitResult UAbilitySystemBlueprintLibrary::GetHitResult(FGameplayCueParameters Parameters)
+{
+	if (Parameters.InstigatorContext.HitResult.IsValid())
+	{
+		return *Parameters.InstigatorContext.HitResult.Get();
+	}
+	
+	return FHitResult();
+}
+
+bool UAbilitySystemBlueprintLibrary::HasHitResult(FGameplayCueParameters Parameters)
+{
+	return Parameters.InstigatorContext.HitResult.IsValid();
+}
\ No newline at end of file
diff --git a/Engine/Source/Runtime/GameplayAbilities/Private/AbilitySystemComponent.cpp b/Engine/Source/Runtime/GameplayAbilities/Private/AbilitySystemComponent.cpp
--- a/Engine/Source/Runtime/GameplayAbilities/Private/AbilitySystemComponent.cpp
+++ b/Engine/Source/Runtime/GameplayAbilities/Private/AbilitySystemComponent.cpp
@@ -230,6 +230,11 @@ bool UAbilitySystemComponent::HasAllTags(FGameplayTagContainer &Tags)
 	return ActiveGameplayEffects.HasAllTags(Tags);
 }
 
+FOnGameplayEffectTagCountChanged& UAbilitySystemComponent::RegisterGameplayTagEvent(FGameplayTag Tag)
+{
+	return ActiveGameplayEffects.RegisterGameplayTagEvent(Tag);
+}
+
 void UAbilitySystemComponent::TEMP_ApplyActiveGameplayEffects()
 {
 	for (int32 idx=0; idx < ActiveGameplayEffects.GameplayEffects.Num(); ++idx)
@@ -331,8 +336,8 @@ FActiveGameplayEffectHandle UAbilitySystemComponent::ApplyGameplayEffectSpecToSe
 		// if the Cue was actually added+activated or just added (due to relevancy)
 
 		// Fixme: what if we wanted to scale Cue magnitude based on damage? E.g, scale an cue effect when the GE is buffed?
-		InvokeGameplayCueAdded(*OurCopyOfSpec);
-		InvokeGameplayCueActivated(*OurCopyOfSpec);
+		InvokeGameplayCueEvent(*OurCopyOfSpec, EGameplayCueEvent::OnActive);
+		InvokeGameplayCueEvent(*OurCopyOfSpec, EGameplayCueEvent::WhileActive);
 	}
 	
 	// Execute the GE at least once (if instant, this will execute once and be done. If persistent, it was added to ActiveGameplayEffects above)
@@ -390,11 +395,11 @@ float UAbilitySystemComponent::GetGameplayEffectMagnitude(FActiveGameplayEffectH
 	return ActiveGameplayEffects.GetGameplayEffectMagnitude(Handle, Attribute);
 }
 
-void UAbilitySystemComponent::InvokeGameplayCueExecute(const FGameplayEffectSpec &Spec)
+void UAbilitySystemComponent::InvokeGameplayCueEvent(const FGameplayEffectSpec &Spec, EGameplayCueEvent::Type EventType)
 {
 	if (DebugGameplayCues)
 	{
-		ABILITY_LOG(Warning, TEXT("InvokeGameplayCueExecute: %s"), *Spec.ToSimpleString());
+		ABILITY_LOG(Warning, TEXT("InvokeGameplayCueEvent: %s"), *Spec.ToSimpleString());
 	}
 
 
@@ -404,13 +409,17 @@ void UAbilitySystemComponent::InvokeGameplayCueExecute(const FGameplayEffectSpec
 	{
 		return;
 	}
-	
+
 	// FIXME: Replication of level not finished
 	float ExecuteLevel = Spec.ModifierLevel.Get()->IsValid() ? Spec.ModifierLevel.Get()->GetLevel() : 1.f;
+
+	FGameplayCueParameters CueParameters;
+	CueParameters.InstigatorContext = Spec.InstigatorContext;
+
 	for (FGameplayEffectCue CueInfo : Spec.Def->GameplayCues)
 	{
-		float NormalizedMagnitude = CueInfo.NormalizeLevel(ExecuteLevel);
-		GameplayCueInterface->GameplayCueExecuted(CueInfo.GameplayCueTags, NormalizedMagnitude, Spec.InstigatorContext);
+		CueParameters.NormalizedMagnitude = CueInfo.NormalizeLevel(ExecuteLevel);
+		GameplayCueInterface->HandleGameplayCues(ActorOwner, CueInfo.GameplayCueTags, EventType, CueParameters);
 
 		if (DebugGameplayCues)
 		{
@@ -420,63 +429,10 @@ void UAbilitySystemComponent::InvokeGameplayCueExecute(const FGameplayEffectSpec
 	}
 }
 
-void UAbilitySystemComponent::InvokeGameplayCueActivated(const FGameplayEffectSpec &Spec)
-{
-	AActor *ActorOwner = GetOwner();
-	IGameplayCueInterface * GameplayCueInterface = InterfaceCast<IGameplayCueInterface>(ActorOwner);
-	if (!GameplayCueInterface)
-	{
-		return;
-	}
-
-	// FIXME: Replication of level not finished
-	float ExecuteLevel = Spec.ModifierLevel.Get()->IsValid() ? Spec.ModifierLevel.Get()->GetLevel() : 1.f; 
-	for (FGameplayEffectCue CueInfo : Spec.Def->GameplayCues)
-	{
-		float NormalizedMagnitude = CueInfo.NormalizeLevel(ExecuteLevel);
-		GameplayCueInterface->GameplayCueActivated(CueInfo.GameplayCueTags, NormalizedMagnitude, Spec.InstigatorContext);
-	}
-}
 
 void UAbilitySystemComponent::NetMulticast_InvokeGameplayCueExecuted_Implementation(const FGameplayEffectSpec Spec)
 {
-	InvokeGameplayCueExecute(Spec);
-}
-
-void UAbilitySystemComponent::InvokeGameplayCueAdded(const FGameplayEffectSpec &Spec)
-{
-	AActor *ActorOwner = GetOwner();
-	IGameplayCueInterface * GameplayCueInterface = InterfaceCast<IGameplayCueInterface>(ActorOwner);
-	if (!GameplayCueInterface)
-	{
-		return;
-	}
-
-	// FIXME: Replication of level not finished
-	float ExecuteLevel = Spec.ModifierLevel.Get()->IsValid() ? Spec.ModifierLevel.Get()->GetLevel() : 1.f;
-	for (FGameplayEffectCue CueInfo : Spec.Def->GameplayCues)
-	{
-		float NormalizedMagnitude = CueInfo.NormalizeLevel(ExecuteLevel);
-		GameplayCueInterface->GameplayCueAdded(CueInfo.GameplayCueTags, NormalizedMagnitude, Spec.InstigatorContext);
-	}
-}
-
-void UAbilitySystemComponent::InvokeGameplayCueRemoved(const FGameplayEffectSpec &Spec)
-{
-	AActor *ActorOwner = GetOwner();
-	IGameplayCueInterface * GameplayCueInterface = InterfaceCast<IGameplayCueInterface>(ActorOwner);
-	if (!GameplayCueInterface)
-	{
-		return;
-	}
-
-	// FIXME: Replication of level not finished
-	float ExecuteLevel = Spec.ModifierLevel.Get()->IsValid() ? Spec.ModifierLevel.Get()->GetLevel() : 1.f;
-	for (FGameplayEffectCue CueInfo : Spec.Def->GameplayCues)
-	{
-		float NormalizedMagnitude = CueInfo.NormalizeLevel(ExecuteLevel);
-		GameplayCueInterface->GameplayCueRemoved(CueInfo.GameplayCueTags, NormalizedMagnitude, Spec.InstigatorContext);
-	}
+	InvokeGameplayCueEvent(Spec, EGameplayCueEvent::Executed);
 }
 
 void UAbilitySystemComponent::AddDependancyToAttribute(FGameplayAttribute Attribute, const TWeakPtr<FAggregator> InDependant)
@@ -604,6 +560,15 @@ int32 UAbilitySystemComponent::GetNextPredictionKey()
 	return ++LocalPredictionKey;
 }
 
+void UAbilitySystemComponent::DispatchBlueprintCustomHandler(AActor* Actor, UFunction* Func, EGameplayCueEvent::Type EventType, FGameplayCueParameters Parameters)
+{
+	AbilitySystemComponent_eventBlueprintCustomHandler_Parms Parms;
+	Parms.EventType = EventType;
+	Parms.Parameters = Parameters;
+
+	Actor->ProcessEvent(Func, &Parms);
+}
+
 // ---------------------------------------------------------------------------------------
 
 void UAbilitySystemComponent::PrintAllGameplayEffects() const
diff --git a/Engine/Source/Runtime/GameplayAbilities/Private/AbilitySystemComponent_Abilities.cpp b/Engine/Source/Runtime/GameplayAbilities/Private/AbilitySystemComponent_Abilities.cpp
--- a/Engine/Source/Runtime/GameplayAbilities/Private/AbilitySystemComponent_Abilities.cpp
+++ b/Engine/Source/Runtime/GameplayAbilities/Private/AbilitySystemComponent_Abilities.cpp
@@ -310,6 +310,12 @@ void UAbilitySystemComponent::InputConfirm()
 
 void UAbilitySystemComponent::InputCancel()
 {
+	if (GetOwnerRole() != ROLE_Authority && CancelCallbacks.IsBound())
+	{
+		// Tell the server we confirmed input.
+		ServerSetReplicatedConfirm(false);
+	}
+
 	CancelCallbacks.Broadcast();
 }
 
@@ -343,8 +349,16 @@ void UAbilitySystemComponent::TargetCancel()
 
 void UAbilitySystemComponent::ServerSetReplicatedConfirm_Implementation(bool Confirmed)
 {
-	ReplicatedConfirmAbility = Confirmed;
-	ConfirmCallbacks.Broadcast();
+	if (Confirmed)
+	{
+		ReplicatedConfirmAbility = true;
+		ConfirmCallbacks.Broadcast();
+	}
+	else
+	{
+		ReplicatedCancelAbility = true;
+		CancelCallbacks.Broadcast();
+	}
 }
 
 bool UAbilitySystemComponent::ServerSetReplicatedConfirm_Validate(bool Confirmed)
@@ -387,9 +401,10 @@ void UAbilitySystemComponent::SetTargetAbility(UGameplayAbility* NewTargetingAbi
 	//ReplicatedTargetData.Clear();
 }
 
-void UAbilitySystemComponent::ConsumeAbilityConfirm()
+void UAbilitySystemComponent::ConsumeAbilityConfirmCancel()
 {
 	ReplicatedConfirmAbility = false;
+	ReplicatedCancelAbility = false;
 }
 
 void UAbilitySystemComponent::ConsumeAbilityTargetData()
diff --git a/Engine/Source/Runtime/GameplayAbilities/Private/AbilitySystemGlobals.cpp b/Engine/Source/Runtime/GameplayAbilities/Private/AbilitySystemGlobals.cpp
--- a/Engine/Source/Runtime/GameplayAbilities/Private/AbilitySystemGlobals.cpp
+++ b/Engine/Source/Runtime/GameplayAbilities/Private/AbilitySystemGlobals.cpp
@@ -3,6 +3,8 @@
 #include "AbilitySystemPrivatePCH.h"
 #include "AbilitySystemComponent.h"
 #include "Abilities/GameplayAbilityTypes.h"
+#include "GameplayTagsModule.h"
+#include "GameplayCueInterface.h"
 
 UAbilitySystemGlobals::UAbilitySystemGlobals(const class FPostConstructInitializeProperties& PCIP)
 : Super(PCIP)
@@ -91,3 +93,46 @@ UAbilitySystemGlobals& UAbilitySystemGlobals::Get()
 
 	return *GlobalPtr;
 }
+
+
+// --------------------------------------------------------------------
+UFunction* UAbilitySystemGlobals::GetGameplayCueFunction(const FGameplayTag& ChildTag, UClass* Class, FName &MatchedTag)
+{
+	SCOPE_CYCLE_COUNTER(STAT_GetGameplayCueFunction);
+
+	// A global cached map to lookup these functions might be a good idea. Keep in mind though that FindFunctionByName
+	// is fast and already gives us a reliable map lookup.
+	// 
+	// We would get some speed by caching off the 'fully qualified name' to 'best match' lookup. E.g. we can directly map
+	// 'GameplayCue.X.Y' to the function 'GameplayCue.X' without having to look for GameplayCue.X.Y first.
+	// 
+	// The native remapping (Gameplay.X.Y to Gameplay_X_Y) is also annoying and slow and could be fixed by this as well.
+	// 
+	// Keep in mind that any UFunction* cacheing is pretty unsafe. Classes can be loaded (and unloaded) during runtime
+	// and will be regenerated all the time in the editor. Just doing a single pass at startup won't be enough,
+	// we'll need a mechanism for registering classes when they are loaded on demand.
+	
+	IGameplayTagsModule& GameplayTagsModule = IGameplayTagsModule::Get();
+	FGameplayTagContainer TagAndParentsContainer = GameplayTagsModule.GetGameplayTagsManager().RequestGameplayTagParents(ChildTag);
+
+	for (auto InnerTagIt = TagAndParentsContainer.CreateConstIterator(); InnerTagIt; ++InnerTagIt)
+	{
+		FName CueName = InnerTagIt->GetTagName();
+		if (UFunction* Func = Class->FindFunctionByName(CueName, EIncludeSuperFlag::IncludeSuper))
+		{
+			MatchedTag = CueName;
+			return Func;
+		}
+
+		// Native functions cant be named with ".", so look for them with _. 
+		FName NativeCueFuncName = *CueName.ToString().Replace(TEXT("."), TEXT("_"));
+		if (UFunction* Func = Class->FindFunctionByName(NativeCueFuncName, EIncludeSuperFlag::IncludeSuper))
+		{
+			MatchedTag = CueName; // purposefully returning the . qualified name.
+			return Func;
+		}
+	}
+
+	return nullptr;
+}
+
diff --git a/Engine/Source/Runtime/GameplayAbilities/Private/AbilitySystemStats.cpp b/Engine/Source/Runtime/GameplayAbilities/Private/AbilitySystemStats.cpp
--- a/Engine/Source/Runtime/GameplayAbilities/Private/AbilitySystemStats.cpp
+++ b/Engine/Source/Runtime/GameplayAbilities/Private/AbilitySystemStats.cpp
@@ -7,4 +7,7 @@ DEFINE_STAT(STAT_GameplayEffectsHasAnyTag);
 DEFINE_STAT(STAT_GameplayEffectsTick);
 DEFINE_STAT(STAT_GameplayEffectsCanApplyAttributeModifiers);
 DEFINE_STAT(STAT_GameplayEffectsGetActiveEffectsData);
-DEFINE_STAT(STAT_GameplayAbilityGetCooldownTimeRemaining);
\ No newline at end of file
+DEFINE_STAT(STAT_GameplayAbilityGetCooldownTimeRemaining);
+DEFINE_STAT(STAT_RemoveActiveGameplayEffect);
+DEFINE_STAT(STAT_CreateNewActiveGameplayEffect);
+DEFINE_STAT(STAT_GetGameplayCueFunction);
diff --git a/Engine/Source/Runtime/GameplayAbilities/Private/AbilitySystemTestPawn.cpp b/Engine/Source/Runtime/GameplayAbilities/Private/AbilitySystemTestPawn.cpp
--- a/Engine/Source/Runtime/GameplayAbilities/Private/AbilitySystemTestPawn.cpp
+++ b/Engine/Source/Runtime/GameplayAbilities/Private/AbilitySystemTestPawn.cpp
@@ -30,28 +30,4 @@ void AAbilitySystemTestPawn::PostInitializeComponents()
 		AbilitySystemComponent->InitializeAbilities(DefaultAbilitySet);
 	}
 	*/
-}
-
-void AAbilitySystemTestPawn::GameplayCueActivated(const FGameplayTagContainer & GameplayCueTags, float NormalizedMagnitude, const FGameplayEffectInstigatorContext InstigatorContext)
-{
-	ABILITY_LOG(Log, TEXT("GameplayCueExecuted: %s. %.2f"), *GetName(), NormalizedMagnitude);
-	GameplayCueHandler.GameplayCueActivated(GameplayCueTags, NormalizedMagnitude, InstigatorContext);
-}
-
-void AAbilitySystemTestPawn::GameplayCueExecuted(const FGameplayTagContainer & GameplayCueTags, float NormalizedMagnitude, const FGameplayEffectInstigatorContext InstigatorContext)
-{
-	ABILITY_LOG(Log, TEXT("GameplayCueExecuted: %s. %.2f"), *GetName(), NormalizedMagnitude);
-	GameplayCueHandler.GameplayCueExecuted(GameplayCueTags, NormalizedMagnitude, InstigatorContext);
-}
-
-void AAbilitySystemTestPawn::GameplayCueAdded(const FGameplayTagContainer & GameplayCueTags, float NormalizedMagnitude, const FGameplayEffectInstigatorContext InstigatorContext)
-{
-	ABILITY_LOG(Log, TEXT("GameplayCueAdded: %s. %.2f"), *GetName(), NormalizedMagnitude);
-	GameplayCueHandler.GameplayCueAdded(GameplayCueTags, NormalizedMagnitude, InstigatorContext);
-}
-
-void AAbilitySystemTestPawn::GameplayCueRemoved(const FGameplayTagContainer & GameplayCueTags, float NormalizedMagnitude, const FGameplayEffectInstigatorContext InstigatorContext)
-{
-	ABILITY_LOG(Log, TEXT("GameplayCueRemoved: %s. %.2f"), *GetName(), NormalizedMagnitude);
-	GameplayCueHandler.GameplayCueRemoved(GameplayCueTags, NormalizedMagnitude, InstigatorContext);
 }
\ No newline at end of file
diff --git a/Engine/Source/Runtime/GameplayAbilities/Private/GameplayCueActor.cpp b/Engine/Source/Runtime/GameplayAbilities/Private/GameplayCueActor.cpp
new file mode 100644
--- /dev/null
+++ b/Engine/Source/Runtime/GameplayAbilities/Private/GameplayCueActor.cpp
@@ -0,0 +1,20 @@
+// Copyright 1998-2014 Epic Games, Inc. All Rights Reserved.
+
+#include "AbilitySystemPrivatePCH.h"
+#include "GameplayCueActor.h"
+
+AGameplayCueActor::AGameplayCueActor(const class FPostConstructInitializeProperties& PCIP)
+: Super(PCIP)
+{
+	AttachToOwner = false;
+}
+
+void AGameplayCueActor::BeginPlay()
+{
+	Super::BeginPlay();
+
+	if (AttachToOwner)
+	{
+		this->AttachRootComponentToActor(GetOwner(), AttachSocketName, EAttachLocation::SnapToTarget);
+	}
+}
\ No newline at end of file
diff --git a/Engine/Source/Runtime/GameplayAbilities/Private/GameplayCueInterface.cpp b/Engine/Source/Runtime/GameplayAbilities/Private/GameplayCueInterface.cpp
--- a/Engine/Source/Runtime/GameplayAbilities/Private/GameplayCueInterface.cpp
+++ b/Engine/Source/Runtime/GameplayAbilities/Private/GameplayCueInterface.cpp
@@ -7,3 +7,35 @@ UGameplayCueInterface::UGameplayCueInterface(const FPostConstructInitializePrope
 : Super(PCIP)
 {
 }
+
+
+void IGameplayCueInterface::DispatchBlueprintCustomHandler(AActor* Actor, UFunction* Func, EGameplayCueEvent::Type EventType, FGameplayCueParameters Parameters)
+{
+	GameplayCueInterface_eventBlueprintCustomHandler_Parms Parms;
+	Parms.EventType = EventType;
+	Parms.Parameters = Parameters;
+
+	Actor->ProcessEvent(Func, &Parms);
+}
+
+void IGameplayCueInterface::HandleGameplayCues(AActor *Self, const FGameplayTagContainer& GameplayCueTags, EGameplayCueEvent::Type EventType, FGameplayCueParameters Parameters)
+{
+	for (auto TagIt = GameplayCueTags.CreateConstIterator(); TagIt; ++TagIt)
+	{
+		HandleGameplayCue(Self, *TagIt, EventType, Parameters);
+	}
+}
+
+void IGameplayCueInterface::HandleGameplayCue(AActor *Self, FGameplayTag GameplayCueTag, EGameplayCueEvent::Type EventType, FGameplayCueParameters Parameters)
+{
+	// Look up a custom function for this gameplay tag. 
+	// This is done in UAbilitySystemGlobals since we will eventually cache these functions off in a global map
+
+	UFunction* Func = UAbilitySystemGlobals::Get().GetGameplayCueFunction(GameplayCueTag, Self->GetClass(), Parameters.MatchedTagName);
+	if (Func)
+	{
+		IGameplayCueInterface::DispatchBlueprintCustomHandler(Self, Func, EventType, Parameters);
+	}
+
+	// Fall back to a possible generic handler if available
+}
\ No newline at end of file
diff --git a/Engine/Source/Runtime/GameplayAbilities/Private/GameplayCueView.cpp b/Engine/Source/Runtime/GameplayAbilities/Private/GameplayCueView.cpp
--- a/Engine/Source/Runtime/GameplayAbilities/Private/GameplayCueView.cpp
+++ b/Engine/Source/Runtime/GameplayAbilities/Private/GameplayCueView.cpp
@@ -8,8 +8,11 @@
 #include "Particles/ParticleLODLevel.h"
 #include "Particles/ParticleModuleRequired.h"
 #include "Particles/ParticleSystemComponent.h"
+#include "AbilitySystemBlueprintLibrary.h"
 #include "GameplayCueInterface.h"
 #include "GameplayTagsModule.h"
+#include "GameplayCueActor.h"
+#include "AbilitySystemComponent.h"
 
 UGameplayCueView::UGameplayCueView(const class FPostConstructInitializeProperties& PCIP)
 : Super(PCIP)
@@ -17,7 +20,7 @@ UGameplayCueView::UGameplayCueView(const class FPostConstructInitializePropertie
 
 }
 
-FGameplayCueViewInfo * FGameplayCueHandler::GetBestMatchingView(EGameplayCueEvent::Type Type, const FGameplayTag BaseTag)
+FGameplayCueViewInfo * FGameplayCueHandler::GetBestMatchingView(EGameplayCueEvent::Type Type, const FGameplayTag BaseTag, bool InstigatorLocal, bool TargetLocal)
 {
 	IGameplayTagsModule& GameplayTagsModule = IGameplayTagsModule::Get();
 	FGameplayTagContainer TagAndParentsContainer = GameplayTagsModule.GetGameplayTagsManager().RequestGameplayTagParents(BaseTag);
@@ -30,7 +33,10 @@ FGameplayCueViewInfo * FGameplayCueHandler::GetBestMatchingView(EGameplayCueEven
 		{
 			for (FGameplayCueViewInfo & View : Def->Views)
 			{
-				if (View.CueType == Type && View.Tags.HasTag(Tag, EGameplayTagMatchType::Explicit, EGameplayTagMatchType::Explicit))
+				if (View.CueType == Type
+					&& (!View.InstigatorLocalOnly || InstigatorLocal)
+					&& (!View.TargetLocalOnly || TargetLocal) 
+					&& View.Tags.HasTag(Tag, EGameplayTagMatchType::Explicit, EGameplayTagMatchType::Explicit))
 				{
 					return &View;
 				}
@@ -44,36 +50,49 @@ FGameplayCueViewInfo * FGameplayCueHandler::GetBestMatchingView(EGameplayCueEven
 void FGameplayCueHandler::GameplayCueActivated(const FGameplayTagContainer & GameplayCueTags, float NormalizedMagnitude, const FGameplayEffectInstigatorContext InstigatorContext)
 {
 	check(Owner);
+	bool InstigatorLocal = InstigatorContext.IsLocallyControlled();
+	bool TargetLocal = OwnerIsLocallyControlled();
+	
 	for (auto TagIt = GameplayCueTags.CreateConstIterator(); TagIt; ++TagIt)
 	{
-		if (FGameplayCueViewInfo * View = GetBestMatchingView(EGameplayCueEvent::OnActive, *TagIt))
+		if (FGameplayCueViewInfo* View = GetBestMatchingView(EGameplayCueEvent::OnActive, *TagIt, InstigatorLocal, TargetLocal))
 		{
 			View->SpawnViewEffects(Owner, NULL, InstigatorContext);
 		}
+
+		FName MatchedTag;
+		UFunction *Func = UAbilitySystemGlobals::Get().GetGameplayCueFunction(*TagIt, Owner->GetClass(), MatchedTag);
+		if (Func)
+		{
+			FGameplayCueParameters Params;
+			Params.NormalizedMagnitude = NormalizedMagnitude;
+			Params.InstigatorContext;
+
+			IGameplayCueInterface::DispatchBlueprintCustomHandler(Owner, Func, EGameplayCueEvent::OnActive, Params);
+		}
 	}
 }
 
 void FGameplayCueHandler::GameplayCueExecuted(const FGameplayTagContainer & GameplayCueTags, float NormalizedMagnitude, const FGameplayEffectInstigatorContext InstigatorContext)
 {
 	check(Owner);
+	bool InstigatorLocal = InstigatorContext.IsLocallyControlled();
+	bool TargetLocal = OwnerIsLocallyControlled();
+
 	for (auto TagIt = GameplayCueTags.CreateConstIterator(); TagIt; ++TagIt)
 	{
-		if (FGameplayCueViewInfo * View = GetBestMatchingView(EGameplayCueEvent::Executed, *TagIt))
+		if (FGameplayCueViewInfo* View = GetBestMatchingView(EGameplayCueEvent::Executed, *TagIt, InstigatorLocal, TargetLocal))
 		{
 			View->SpawnViewEffects(Owner, NULL, InstigatorContext);
 		}
-
-		
-		{
-			IGameplayCueInterface::Execute_HandleGameplayCue(Owner, *TagIt, EGameplayCueEvent::Executed);
-		}
-
 	}
 }
 
 void FGameplayCueHandler::GameplayCueAdded(const FGameplayTagContainer & GameplayCueTags, float NormalizedMagnitude, const FGameplayEffectInstigatorContext InstigatorContext)
 {
 	check(Owner);
+	bool InstigatorLocal = InstigatorContext.IsLocallyControlled();
+	bool TargetLocal = OwnerIsLocallyControlled();
 
 	for (auto TagIt = GameplayCueTags.CreateConstIterator(); TagIt; ++TagIt)
 	{
@@ -83,21 +102,19 @@ void FGameplayCueHandler::GameplayCueAdded(const FGameplayTagContainer & Gamepla
 		ClearEffects(Effects);
 		check(Effects.Num() == 0);
 
-		if (FGameplayCueViewInfo * View = GetBestMatchingView(EGameplayCueEvent::WhileActive, *TagIt))
+		if (FGameplayCueViewInfo * View = GetBestMatchingView(EGameplayCueEvent::WhileActive, *TagIt, InstigatorLocal, TargetLocal))
 		{
 			TSharedPtr<FGameplayCueViewEffects> SpawnedEffects = View->SpawnViewEffects(Owner, &SpawnedObjects, InstigatorContext);
 			Effects.Add(SpawnedEffects);
 		}
-
-		{
-			IGameplayCueInterface::Execute_HandleGameplayCue(Owner, *TagIt, EGameplayCueEvent::WhileActive);
-		}
 	}
 }
 
 void FGameplayCueHandler::GameplayCueRemoved(const FGameplayTagContainer & GameplayCueTags, float NormalizedMagnitude, const FGameplayEffectInstigatorContext InstigatorContext)
 {
 	check(Owner);
+	bool InstigatorLocal = InstigatorContext.IsLocallyControlled();
+	bool TargetLocal = OwnerIsLocallyControlled();
 
 	for (auto TagIt = GameplayCueTags.CreateConstIterator(); TagIt; ++TagIt)
 	{
@@ -110,13 +127,20 @@ void FGameplayCueHandler::GameplayCueRemoved(const FGameplayTagContainer & Gamep
 		
 		// Remove old effects
 		
-		if (FGameplayCueViewInfo * View = GetBestMatchingView(EGameplayCueEvent::Removed, *TagIt))
+		if (FGameplayCueViewInfo * View = GetBestMatchingView(EGameplayCueEvent::Removed, *TagIt, InstigatorLocal, TargetLocal))
 		{
 			View->SpawnViewEffects(Owner, NULL, InstigatorContext);
 		}
 
+		FName MatchedTag;
+		UFunction *Func = UAbilitySystemGlobals::Get().GetGameplayCueFunction(*TagIt, Owner->GetClass(), MatchedTag);
+		if (Func)
 		{
-			IGameplayCueInterface::Execute_HandleGameplayCue(Owner, *TagIt, EGameplayCueEvent::Removed);
+			FGameplayCueParameters Params;
+			Params.NormalizedMagnitude = NormalizedMagnitude;
+			Params.InstigatorContext;
+
+			IGameplayCueInterface::DispatchBlueprintCustomHandler(Owner, Func, EGameplayCueEvent::Removed, Params);
 		}
 	}
 }
@@ -189,12 +213,28 @@ TSharedPtr<FGameplayCueViewEffects> FGameplayCueViewInfo::SpawnViewEffects(AActo
 		{
 			FHitResult &HitResult = *InstigatorContext.HitResult.Get();
 
-			SpawnedEffects->ParticleSystemComponent = UGameplayStatics::SpawnEmitterAttached(ParticleSystem, Owner->GetRootComponent(), NAME_None, HitResult.Location,
-				HitResult.Normal.Rotation(), EAttachLocation::KeepWorldPosition);
+			if (AttachParticleSystem)
+			{
+				SpawnedEffects->ParticleSystemComponent = UGameplayStatics::SpawnEmitterAttached(ParticleSystem, Owner->GetRootComponent(), NAME_None, HitResult.Location,
+					HitResult.Normal.Rotation(), EAttachLocation::KeepWorldPosition);
+			}
+			else
+			{
+				bool AutoDestroy = SpawnedObjects == nullptr;
+				SpawnedEffects->ParticleSystemComponent = UGameplayStatics::SpawnEmitterAtLocation(Owner->GetWorld(), ParticleSystem, HitResult.Location, HitResult.Normal.Rotation(), AutoDestroy);
+			}
 		}
 		else
 		{
-			SpawnedEffects->ParticleSystemComponent = UGameplayStatics::SpawnEmitterAttached(ParticleSystem, Owner->GetRootComponent());
+			if (AttachParticleSystem)
+			{
+				SpawnedEffects->ParticleSystemComponent = UGameplayStatics::SpawnEmitterAttached(ParticleSystem, Owner->GetRootComponent());
+			}
+			else
+			{
+				bool AutoDestroy = SpawnedObjects == nullptr;
+				SpawnedEffects->ParticleSystemComponent = UGameplayStatics::SpawnEmitterAtLocation(Owner->GetWorld(), ParticleSystem, Owner->GetActorLocation(), Owner->GetActorRotation(), AutoDestroy);
+			}
 		}		
 
 		if (SpawnedObjects)
@@ -220,7 +260,8 @@ TSharedPtr<FGameplayCueViewEffects> FGameplayCueViewInfo::SpawnViewEffects(AActo
 	{
 		FVector Location = Owner->GetActorLocation();
 		FRotator Rotation = Owner->GetActorRotation();
-		SpawnedEffects->SpawnedActor = Owner->GetWorld()->SpawnActor(ActorClass, &Location, &Rotation);
+		SpawnedEffects->SpawnedActor = Cast<AGameplayCueActor>(Owner->GetWorld()->SpawnActor(ActorClass, &Location, &Rotation));
+		
 		if (SpawnedObjects)
 		{
 			SpawnedObjects->Add(SpawnedEffects->SpawnedActor.Get());
@@ -229,3 +270,17 @@ TSharedPtr<FGameplayCueViewEffects> FGameplayCueViewInfo::SpawnViewEffects(AActo
 
 	return SpawnedEffects;
 }
+
+bool FGameplayCueHandler::OwnerIsLocallyControlled() const
+{
+	APawn* Pawn = Cast<APawn>(Owner);
+	if (Pawn)
+	{
+		return Pawn->IsLocallyControlled();
+	}
+
+	return false;
+}
+
+// ---------------------------------------------------
+
diff --git a/Engine/Source/Runtime/GameplayAbilities/Private/GameplayEffect.cpp b/Engine/Source/Runtime/GameplayAbilities/Private/GameplayEffect.cpp
--- a/Engine/Source/Runtime/GameplayAbilities/Private/GameplayEffect.cpp
+++ b/Engine/Source/Runtime/GameplayAbilities/Private/GameplayEffect.cpp
@@ -4,6 +4,7 @@
 #include "GameplayEffect.h"
 #include "GameplayEffectExtension.h"
 #include "GameplayEffectStackingExtension.h"
+#include "GameplayTagsModule.h"
 #include "AbilitySystemComponent.h"
 
 const float UGameplayEffect::INFINITE_DURATION = -1.f;
@@ -498,6 +499,17 @@ bool FGameplayEffectInstigatorContext::NetSerialize(FArchive& Ar, class UPackage
 	return true;
 }
 
+bool FGameplayEffectInstigatorContext::IsLocallyControlled() const
+{
+	APawn* Pawn = Cast<APawn>(Instigator);
+	if (Pawn)
+	{
+		return Pawn->IsLocallyControlled();
+	}
+	return false;
+}
+
+
 // --------------------------------------------------------------------------------------------------------------------------------------------------------
 //
 //	FAggregatorRef
@@ -1112,14 +1124,18 @@ bool FActiveGameplayEffect::CanBeStacked(const FActiveGameplayEffect& Other) con
 
 void FActiveGameplayEffect::PreReplicatedRemove(const struct FActiveGameplayEffectsContainer &InArray)
 {
-	InArray.Owner->InvokeGameplayCueRemoved(Spec);
+	const_cast<FActiveGameplayEffectsContainer&>(InArray).InternalOnActiveGameplayEffectRemoved(*this);	// Const cast is ok. It is there to prevent mutation of the GameplayEffects array, which this wont do.
+	
+	InArray.Owner->InvokeGameplayCueEvent(Spec, EGameplayCueEvent::Removed);
 }
 
 void FActiveGameplayEffect::PostReplicatedAdd(const struct FActiveGameplayEffectsContainer &InArray)
 {
+	const_cast<FActiveGameplayEffectsContainer&>(InArray).InternalOnActiveGameplayEffectAdded(*this);	// Const cast is ok. It is there to prevent mutation of the GameplayEffects array, which this wont do.
+
 	static const int32 MAX_DELTA_TIME = 3;
 
-	InArray.Owner->InvokeGameplayCueAdded(Spec);
+	InArray.Owner->InvokeGameplayCueEvent(Spec, EGameplayCueEvent::WhileActive);
 	// Was this actually just activated, or are we just finding out about it due to relevancy/join in progress?
 	float WorldTimeSeconds = InArray.GetWorldTime();
 	int32 GameStateTime = InArray.GetGameStateTime();
@@ -1128,7 +1144,7 @@ void FActiveGameplayEffect::PostReplicatedAdd(const struct FActiveGameplayEffect
 
 	if (GameStateTime > 0 && FMath::Abs(DeltaGameStateTime) < MAX_DELTA_TIME)
 	{
-		InArray.Owner->InvokeGameplayCueActivated(Spec);
+		InArray.Owner->InvokeGameplayCueEvent(Spec, EGameplayCueEvent::OnActive);
 	}
 
 	// Set our local start time accordingly
@@ -1526,6 +1542,8 @@ void FActiveGameplayEffectsContainer::StacksNeedToRecalculate()
 
 FActiveGameplayEffect & FActiveGameplayEffectsContainer::CreateNewActiveGameplayEffect(const FGameplayEffectSpec &Spec, int32 InPredictionKey)
 {
+	SCOPE_CYCLE_COUNTER(STAT_CreateNewActiveGameplayEffect);
+
 	LastAssignedHandle = LastAssignedHandle.GetNextHandle();
 	FActiveGameplayEffect & NewEffect = *new (GameplayEffects)FActiveGameplayEffect(LastAssignedHandle, Spec, GetWorldTime(), GetGameStateTime(), InPredictionKey);
 
@@ -1583,9 +1601,37 @@ FActiveGameplayEffect & FActiveGameplayEffectsContainer::CreateNewActiveGameplay
 
 		Owner->GetPredictionKeyDelegate(InPredictionKey).AddUObject(Owner, &UAbilitySystemComponent::RemoveActiveGameplayEffect_NoReturn, NewEffect.Handle);
 	}
+
+	InternalOnActiveGameplayEffectAdded(NewEffect);	
+
 	return NewEffect;
 }
 
+void FActiveGameplayEffectsContainer::InternalOnActiveGameplayEffectAdded(const FActiveGameplayEffect& Effect)
+{
+	// Update gameplaytag count and broadcast delegate if we just added this tag (count=0, prior to increment)
+	IGameplayTagsModule& GameplayTagsModule = IGameplayTagsModule::Get();
+	for (auto BaseTagIt = Effect.Spec.Def->OwnedTagsContainer.CreateConstIterator(); BaseTagIt; ++BaseTagIt)
+	{
+		const FGameplayTag& BaseTag = *BaseTagIt;
+		FGameplayTagContainer TagAndParentsContainer = GameplayTagsModule.GetGameplayTagsManager().RequestGameplayTagParents(BaseTag);
+		for (auto TagIt = TagAndParentsContainer.CreateConstIterator(); TagIt; ++TagIt)
+		{
+			const FGameplayTag& Tag = *TagIt;
+			int32& Count = GameplayTagCountMap.FindOrAdd(Tag);
+			if (Count++ == 0)
+			{
+				FOnGameplayEffectTagCountChanged *Delegate = GameplayTagEventMap.Find(Tag);
+				if (Delegate)
+				{
+					Delegate->Broadcast(Tag, 1);
+				}
+			}
+		}
+	}
+}
+
+
 bool FActiveGameplayEffectsContainer::RemoveActiveGameplayEffect(FActiveGameplayEffectHandle Handle)
 {
 	// Could make this a map for quicker lookup
@@ -1603,10 +1649,14 @@ bool FActiveGameplayEffectsContainer::RemoveActiveGameplayEffect(FActiveGameplay
 
 bool FActiveGameplayEffectsContainer::InternalRemoveActiveGameplayEffect(int32 Idx)
 {
+	SCOPE_CYCLE_COUNTER(STAT_RemoveActiveGameplayEffect);
+
 	if (ensure(Idx < GameplayEffects.Num()))
 	{
 		FActiveGameplayEffect& Effect = GameplayEffects[Idx];
-		Owner->InvokeGameplayCueRemoved(GameplayEffects[Idx].Spec);
+		Owner->InvokeGameplayCueEvent(GameplayEffects[Idx].Spec, EGameplayCueEvent::Removed);
+
+		InternalOnActiveGameplayEffectRemoved(Effect);
 
 		if (GameplayEffects[Idx].DurationHandle.IsValid())
 		{
@@ -1634,6 +1684,32 @@ bool FActiveGameplayEffectsContainer::InternalRemoveActiveGameplayEffect(int32 I
 	return false;
 }
 
+/** This does cleanup that has to happen whether the effect is being removed locally or due to replication */
+void FActiveGameplayEffectsContainer::InternalOnActiveGameplayEffectRemoved(const FActiveGameplayEffect& Effect)
+{
+	// Update gameplaytag count and broadcast delegate if we are at 0
+	IGameplayTagsModule& GameplayTagsModule = IGameplayTagsModule::Get();
+	for (auto BaseTagIt = Effect.Spec.Def->OwnedTagsContainer.CreateConstIterator(); BaseTagIt; ++BaseTagIt)
+	{
+		const FGameplayTag& BaseTag = *BaseTagIt;
+		FGameplayTagContainer TagAndParentsContainer = GameplayTagsModule.GetGameplayTagsManager().RequestGameplayTagParents(BaseTag);
+		for (auto TagIt = TagAndParentsContainer.CreateConstIterator(); TagIt; ++TagIt)
+		{
+			const FGameplayTag& Tag = *TagIt;
+			int32& Count = GameplayTagCountMap[Tag];
+			if (--Count == 0)
+			{
+				FOnGameplayEffectTagCountChanged *Delegate = GameplayTagEventMap.Find(Tag);
+				if (Delegate)
+				{
+					Delegate->Broadcast(Tag, 0);
+				}
+			}
+			check(Count >= 0);
+		}
+	}
+}
+
 bool FActiveGameplayEffectsContainer::IsNetAuthority() const
 {
 	check(Owner);
@@ -1823,11 +1899,9 @@ void FActiveGameplayEffectsContainer::RecalculateStacking()
 bool FActiveGameplayEffectsContainer::HasAnyTags(FGameplayTagContainer &Tags)
 {
 	SCOPE_CYCLE_COUNTER(STAT_GameplayEffectsHasAnyTag);
-
-	for (FActiveGameplayEffect &ActiveEffect : GameplayEffects)
+	for (auto It = Tags.CreateConstIterator(); It; ++It)
 	{
-		// Fixme: check stacking rules!
-		if (ActiveEffect.Spec.Def->OwnedTagsContainer.MatchesAny(Tags, false))
+		if (HasTag(*It))
 		{
 			return true;
 		}
@@ -1853,20 +1927,29 @@ bool FActiveGameplayEffectsContainer::HasAllTags(FGameplayTagContainer &Tags)
 
 bool FActiveGameplayEffectsContainer::HasTag(const FGameplayTag Tag)
 {
-	for (FActiveGameplayEffect &ActiveEffect : GameplayEffects)
+	/**
+	 *	The desired behavior here is:
+	 *		"Do we have tag a.b.c and the GameplayEffect has a.b.c.d -> match!"
+	 *		"Do we have tag a.b.c.d. and the GameplayEffect has a.b.c -> no match!"
+	 */
+
+	int32* Count = GameplayTagCountMap.Find(Tag);
+	if (Count)
 	{
-		/**
-		 *	The desired behavior here is:
-		 *		"Do we have tag a.b.c and the GameplayEffect has a.b.c.d -> match!"
-		 *		"Do we have tag a.b.c.d. and the GameplayEffect has a.b.c -> no match!"
-		 */
+		return *Count > 0;
+	}
+	return false;
 
+#if 0
+	for (FActiveGameplayEffect &ActiveEffect : GameplayEffects)
+	{
 		if (ActiveEffect.Spec.Def->OwnedTagsContainer.HasTag(Tag, EGameplayTagMatchType::IncludeParentTags, EGameplayTagMatchType::Explicit))
 		{
 			return true;
 		}
 	}
 	return false;
+#endif
 }
 
 bool FActiveGameplayEffectsContainer::CanApplyAttributeModifiers(const UGameplayEffect *GameplayEffect, float Level, AActor *Instigator)
@@ -1923,6 +2006,11 @@ TArray<float> FActiveGameplayEffectsContainer::GetActiveEffectsTimeRemaining(con
 	return ReturnList;
 }
 
+FOnGameplayEffectTagCountChanged& FActiveGameplayEffectsContainer::RegisterGameplayTagEvent(FGameplayTag Tag)
+{
+	return GameplayTagEventMap.FindOrAdd(Tag);
+}
+
 // --------------------------------------------------------------------------------------------------------------------------------------------------------
 //
 //	FGameplayEffectLevelSpec
diff --git a/Engine/Source/Runtime/GameplayAbilities/Public/AbilitySystemBlueprintLibrary.h b/Engine/Source/Runtime/GameplayAbilities/Public/AbilitySystemBlueprintLibrary.h
--- a/Engine/Source/Runtime/GameplayAbilities/Public/AbilitySystemBlueprintLibrary.h
+++ b/Engine/Source/Runtime/GameplayAbilities/Public/AbilitySystemBlueprintLibrary.h
@@ -3,6 +3,7 @@
 #pragma once
 
 #include "Abilities/GameplayAbility.h"
+#include "GameplayCueView.h"
 #include "AbilitySystemBlueprintLibrary.generated.h"
 
 class UAnimMontage;
@@ -19,15 +20,32 @@ class UAbilitySystemBlueprintLibrary : public UBlueprintFunctionLibrary
 	UFUNCTION(BlueprintPure, Category = Ability)
 	static UAbilitySystemComponent* GetAbilitySystemComponent(AActor *Actor);
 
-	UFUNCTION(BlueprintCallable, Category=Ability)
+	// -------------------------------------------------------------------------------
+	//		TargetData
+	// -------------------------------------------------------------------------------
+
+	UFUNCTION(BlueprintCallable, Category="Ability|TargetData")
 	static void ApplyGameplayEffectToTargetData(FGameplayAbilityTargetDataHandle Target, UGameplayEffect *GameplayEffect, const FGameplayAbilityActorInfo InstigatorInfo);
 
-	UFUNCTION(BlueprintCallable, Category = Ability)
+	UFUNCTION(BlueprintCallable, Category="Ability|TargetData")
 	static FGameplayAbilityTargetDataHandle	AbilityTargetDataFromHitResult(FHitResult HitResult);
 
-	UFUNCTION(BlueprintCallable, Category = Ability)
+	UFUNCTION(BlueprintCallable, Category="Ability|TargetData")
 	static FHitResult GetHitResultFromTargetData(FGameplayAbilityTargetDataHandle HitResult);
 
+	// -------------------------------------------------------------------------------
+	//		GameplayCue
+	// -------------------------------------------------------------------------------
+	
+	UFUNCTION(BlueprintCallable, Category="Ability|GameplayCue")
+	static bool IsInstigatorLocallyControlled(FGameplayCueParameters Parameters);
+
+	UFUNCTION(BlueprintCallable, Category="Ability|GameplayCue")
+	static FHitResult GetHitResult(FGameplayCueParameters Parameters);
+
+	UFUNCTION(BlueprintCallable, Category = "Ability|GameplayCue")
+	static bool HasHitResult(FGameplayCueParameters Parameters);
+
 	// -------------------------------------------------------------------------------
 
 	UFUNCTION(BlueprintCallable, meta=(HidePin="WorldContextObject", DefaultToSelf="WorldContextObject", BlueprintInternalUseOnly = "TRUE"))
@@ -35,5 +53,10 @@ class UAbilitySystemBlueprintLibrary : public UBlueprintFunctionLibrary
 
 	UFUNCTION(BlueprintCallable, meta = (HidePin = "WorldContextObject", DefaultToSelf = "WorldContextObject", BlueprintInternalUseOnly = "TRUE"))
 	static UAbilityTask_WaitMovementModeChange* CreateWaitMovementModeChange(UObject* WorldContextObject, EMovementMode NewMode);
-	
+
+	// -------------------------------------------------------------------------------
+
+
+
+
 };
\ No newline at end of file
diff --git a/Engine/Source/Runtime/GameplayAbilities/Public/AbilitySystemComponent.h b/Engine/Source/Runtime/GameplayAbilities/Public/AbilitySystemComponent.h
--- a/Engine/Source/Runtime/GameplayAbilities/Public/AbilitySystemComponent.h
+++ b/Engine/Source/Runtime/GameplayAbilities/Public/AbilitySystemComponent.h
@@ -212,7 +212,9 @@ class GAMEPLAYABILITIES_API UAbilitySystemComponent : public UActorComponent
 
 	UFUNCTION(BlueprintCallable, Category = GameplayEffects)
 	bool HasAllTags(FGameplayTagContainer &Tags);
-	
+
+	/** Allow events to be registered for specific gameplay tags being added or removed */
+	FOnGameplayEffectTagCountChanged& RegisterGameplayTagEvent(FGameplayTag Tag);
 
 	// --------------------------------------------
 	// Possibly useful but not primary API functions:
@@ -258,6 +260,18 @@ class GAMEPLAYABILITIES_API UAbilitySystemComponent : public UActorComponent
 		}
 	}
 
+	// ----------------------------------------------------------------------------------------------------------------
+	//
+	//	GameplayCues
+	// 
+	// ----------------------------------------------------------------------------------------------------------------
+
+
+	UFUNCTION(BlueprintImplementableEvent, Category = GameplayCue, meta = (BlueprintInternalUseOnly = "true"))
+	void BlueprintCustomHandler(EGameplayCueEvent::Type EventType, FGameplayCueParameters Parameters);
+
+	static void DispatchBlueprintCustomHandler(AActor* Actor, UFunction* Func, EGameplayCueEvent::Type EventType, FGameplayCueParameters Parameters);
+
 	// ----------------------------------------------------------------------------------------------------------------
 
 	/**
@@ -386,7 +400,7 @@ class GAMEPLAYABILITIES_API UAbilitySystemComponent : public UActorComponent
 
 	void SetTargetAbility(UGameplayAbility* NewTargetingAbility);
 
-	void ConsumeAbilityConfirm();
+	void ConsumeAbilityConfirmCancel();
 
 	void ConsumeAbilityTargetData();
 
@@ -395,6 +409,7 @@ class GAMEPLAYABILITIES_API UAbilitySystemComponent : public UActorComponent
 	UGameplayAbility* TargetingAbility;
 
 	bool ReplicatedConfirmAbility;
+	bool ReplicatedCancelAbility;
 
 	FGameplayAbilityTargetDataHandle ReplicatedTargetData;
 
@@ -448,11 +463,8 @@ private:
 	UFUNCTION(NetMulticast, unreliable)
 	void NetMulticast_InvokeGameplayCueExecuted(const FGameplayEffectSpec Spec);
 
-	void InvokeGameplayCueExecute(const FGameplayEffectSpec &Spec);
-	void InvokeGameplayCueActivated(const FGameplayEffectSpec &Spec);
-	void InvokeGameplayCueAdded(const FGameplayEffectSpec &Spec);
-	void InvokeGameplayCueRemoved(const FGameplayEffectSpec &Spec);
-
+	void InvokeGameplayCueEvent(const FGameplayEffectSpec &Spec, EGameplayCueEvent::Type EventType);	
+	
 	// ---------------------------------------------
 
 	
diff --git a/Engine/Source/Runtime/GameplayAbilities/Public/AbilitySystemGlobals.h b/Engine/Source/Runtime/GameplayAbilities/Public/AbilitySystemGlobals.h
--- a/Engine/Source/Runtime/GameplayAbilities/Public/AbilitySystemGlobals.h
+++ b/Engine/Source/Runtime/GameplayAbilities/Public/AbilitySystemGlobals.h
@@ -7,12 +7,14 @@
 class AActor;
 class UAbilitySystemComponent;
 struct FGameplayAbilityActorInfo;
+struct FGameplayTag;
 
 /** Holds global data for the skill system. Can be configured per project via config file */
 UCLASS(config=Game)
 class GAMEPLAYABILITIES_API UAbilitySystemGlobals : public UObject
 {
 	GENERATED_UCLASS_BODY()
+	
 
 	/** Holds all of the valid gameplay-related tags that can be applied to assets */
 	UPROPERTY(config)
@@ -36,7 +38,7 @@ class GAMEPLAYABILITIES_API UAbilitySystemGlobals : public UObject
 		GlobalAttributeDataTable = InTable;
 	}
 
-	static UAbilitySystemGlobals&	Get();
+	static UAbilitySystemGlobals& Get();
 
 	/**
 	 *	Games may want to override this in a UMyProjectAbilitySystemsGlobals class and provide
@@ -46,6 +48,8 @@ class GAMEPLAYABILITIES_API UAbilitySystemGlobals : public UObject
 
 	virtual FGameplayAbilityActorInfo * AllocAbilityActorInfo(AActor *Actor) const;
 
+	UFunction* GetGameplayCueFunction(const FGameplayTag &Tag, UClass* Class, FName &MatchedTag);
+
 private:
 
 	class UCurveTable* GlobalCurveTable;
diff --git a/Engine/Source/Runtime/GameplayAbilities/Public/AbilitySystemStats.h b/Engine/Source/Runtime/GameplayAbilities/Public/AbilitySystemStats.h
--- a/Engine/Source/Runtime/GameplayAbilities/Public/AbilitySystemStats.h
+++ b/Engine/Source/Runtime/GameplayAbilities/Public/AbilitySystemStats.h
@@ -1,10 +1,7 @@
-// Copyright 1998-2013 Epic Games, Inc. All Rights Reserved.
+// Copyright 1998-2014 Epic Games, Inc. All Rights Reserved.
 
 #pragma once
 
-//#include "GameplayTagAssetInterface.h"
-//#include "GameplayTagsManager.h"
-//#include "GameplayTagContainer.h"
 #include "AbilitySystemLog.h"
 
 DECLARE_STATS_GROUP(TEXT("AbilitySystem"), STATGROUP_AbilitySystem, STATCAT_Advanced);
@@ -15,3 +12,6 @@ DECLARE_CYCLE_STAT_EXTERN(TEXT("GameplayEffectsTick"), STAT_GameplayEffectsTick,
 DECLARE_CYCLE_STAT_EXTERN(TEXT("CanApplyAttributeModifiers"), STAT_GameplayEffectsCanApplyAttributeModifiers, STATGROUP_AbilitySystem, );
 DECLARE_CYCLE_STAT_EXTERN(TEXT("GetActiveEffectsData"), STAT_GameplayEffectsGetActiveEffectsData, STATGROUP_AbilitySystem, );
 DECLARE_CYCLE_STAT_EXTERN(TEXT("GetCooldownTimeRemaining"), STAT_GameplayAbilityGetCooldownTimeRemaining, STATGROUP_AbilitySystem, );
+DECLARE_CYCLE_STAT_EXTERN(TEXT("RemoveActiveGameplayEffect"), STAT_RemoveActiveGameplayEffect, STATGROUP_AbilitySystem, );
+DECLARE_CYCLE_STAT_EXTERN(TEXT("CreateNewActiveGameplayEffect"), STAT_CreateNewActiveGameplayEffect, STATGROUP_AbilitySystem, );
+DECLARE_CYCLE_STAT_EXTERN(TEXT("GetGameplayCueFunction"), STAT_GetGameplayCueFunction, STATGROUP_AbilitySystem, );
\ No newline at end of file
diff --git a/Engine/Source/Runtime/GameplayAbilities/Public/AbilitySystemTestPawn.h b/Engine/Source/Runtime/GameplayAbilities/Public/AbilitySystemTestPawn.h
--- a/Engine/Source/Runtime/GameplayAbilities/Public/AbilitySystemTestPawn.h
+++ b/Engine/Source/Runtime/GameplayAbilities/Public/AbilitySystemTestPawn.h
@@ -11,20 +11,7 @@ class GAMEPLAYABILITIES_API AAbilitySystemTestPawn : public ADefaultPawn, public
 	GENERATED_UCLASS_BODY()
 
 	virtual void PostInitializeComponents() override;
-
-	UFUNCTION()
-	virtual void GameplayCueActivated(const FGameplayTagContainer & GameplayCueTags, float NormalizedMagnitude, const FGameplayEffectInstigatorContext InstigatorContext) override;
-
-	UFUNCTION()
-	virtual void GameplayCueExecuted(const FGameplayTagContainer & GameplayCueTags, float NormalizedMagnitude, const FGameplayEffectInstigatorContext InstigatorContext) override;
-
-	UFUNCTION()
-	virtual void GameplayCueAdded(const FGameplayTagContainer & GameplayCueTags, float NormalizedMagnitude, const FGameplayEffectInstigatorContext InstigatorContext) override;
-
-	UFUNCTION()
-	virtual void GameplayCueRemoved(const FGameplayTagContainer & GameplayCueTags, float NormalizedMagnitude, const FGameplayEffectInstigatorContext InstigatorContext) override;
-
-
+	
 	/** DefaultPawn collision component */
 	UPROPERTY(Category = AbilitySystem, VisibleAnywhere, BlueprintReadOnly)
 	TSubobjectPtr<class UAbilitySystemComponent> AbilitySystemComponent;
diff --git a/Engine/Source/Runtime/GameplayAbilities/Public/GameplayCueActor.h b/Engine/Source/Runtime/GameplayAbilities/Public/GameplayCueActor.h
new file mode 100644
--- /dev/null
+++ b/Engine/Source/Runtime/GameplayAbilities/Public/GameplayCueActor.h
@@ -0,0 +1,36 @@
+// Copyright 1998-2014 Epic Games, Inc. All Rights Reserved.
+
+#pragma once
+
+#include "GameplayEffect.h"
+#include "GameplayCueActor.generated.h"
+
+/**
+ *	This is meant to be a basic implementation of an actor is spawned from a GameplayCue event.
+ *	It provides basic functionality like attachment.
+ *	
+ *	Individual projects will most likely want their own subclass (or a completely different GameplayCue handling system all together)
+ */
+
+
+UCLASS(BlueprintType)
+class GAMEPLAYABILITIES_API AGameplayCueActor : public AActor
+{
+	// Defines how a GameplayCue is translated into viewable components
+	GENERATED_UCLASS_BODY()
+
+	virtual void BeginPlay() override;
+
+	/** Whether this actor should be attached to its owner on spawn (or otherwise persist in the world on its own) */
+	UPROPERTY(EditAnywhere, BlueprintReadOnly, Category="GameplayCueActor", meta=(ExposeOnSpawn=true))
+	bool	AttachToOwner;
+
+	/** If attaching, an optional socket name to attach to */
+	UPROPERTY(EditAnywhere, BlueprintReadOnly, Category="GameplayCueActor", meta=(ExposeOnSpawn=true))
+	FName	AttachSocketName;
+
+	/** Whether this actor should be destroyed when the originating GameplayCue is removed (for 'WhileAdded' Cues only) */
+	UPROPERTY(EditAnywhere, BlueprintReadOnly, Category="GameplayCueActor", meta=(ExposeOnSpawn=true))
+	bool	DestroyWhenCueEnds;
+};
+
diff --git a/Engine/Source/Runtime/GameplayAbilities/Public/GameplayCueInterface.h b/Engine/Source/Runtime/GameplayAbilities/Public/GameplayCueInterface.h
--- a/Engine/Source/Runtime/GameplayAbilities/Public/GameplayCueInterface.h
+++ b/Engine/Source/Runtime/GameplayAbilities/Public/GameplayCueInterface.h
@@ -7,7 +7,6 @@
 #include "GameplayCueInterface.generated.h"
 
 
-
 /** Interface for actors that wish to handle GameplayCue events from GameplayEffects */
 UINTERFACE(MinimalAPI)
 class UGameplayCueInterface: public UInterface
@@ -19,27 +18,13 @@ class GAMEPLAYABILITIES_API IGameplayCueInterface
 {
 	GENERATED_IINTERFACE_BODY()
 
-	UFUNCTION(BlueprintImplementableEvent, Category = "GameplayCue")
-	void HandleGameplayCue(FGameplayTag Tag, EGameplayCueEvent::Type CueType);
-
-	// EVENT -----------------------
-
-	/** EGameplayCueEvent::OnActive - Called when GameplayCue is really activated (not called again in relevancy, Join In Progress situations) */
-	UFUNCTION()
-	virtual void GameplayCueActivated(const FGameplayTagContainer & GameplayCueTags, float NormalizedMagnitude, const FGameplayEffectInstigatorContext InstigatorContext) = 0;
-
-	/** EGameplayCueEvent::Executed - Called when GameplayCue is executed (instant effect or periodic effect ticks) */
-	UFUNCTION()
-	virtual void GameplayCueExecuted(const FGameplayTagContainer & GameplayCueTags, float NormalizedMagnitude, const FGameplayEffectInstigatorContext InstigatorContext) = 0;
-
-	/** EGameplayCueEvent::Removed - Called when GameplayCue is removed */
-	UFUNCTION()
-	virtual void GameplayCueRemoved(const FGameplayTagContainer & GameplayCueTags, float NormalizedMagnitude, const FGameplayEffectInstigatorContext InstigatorContext) = 0;
-
+	virtual void HandleGameplayCue(AActor *Self, FGameplayTag GameplayCueTag, EGameplayCueEvent::Type EventType, FGameplayCueParameters Parameters);
 
-	// STATE -----------------------
+	virtual void HandleGameplayCues(AActor *Self, const FGameplayTagContainer& GameplayCueTags, EGameplayCueEvent::Type EventType, FGameplayCueParameters Parameters);
+	
+	/** Internal function to map ufunctions directly to gameplaycue tags */
+	UFUNCTION(BlueprintImplementableEvent, Category = GameplayCue, meta = (BlueprintInternalUseOnly = "true"))
+	void BlueprintCustomHandler(EGameplayCueEvent::Type EventType, FGameplayCueParameters Parameters);
 
-	/** EGameplayCueEvent::WhileActive - Called when GameplayCue is added locally (e.g., casted for the first time or in relevancy / Join in Progress situations) */
-	UFUNCTION()
-	virtual void GameplayCueAdded(const FGameplayTagContainer & GameplayCueTags, float NormalizedMagnitude, const FGameplayEffectInstigatorContext InstigatorContext) = 0;
+	static void DispatchBlueprintCustomHandler(AActor* Actor, UFunction* Func, EGameplayCueEvent::Type EventType, FGameplayCueParameters Parameters);	
 };
\ No newline at end of file
diff --git a/Engine/Source/Runtime/GameplayAbilities/Public/GameplayCueView.h b/Engine/Source/Runtime/GameplayAbilities/Public/GameplayCueView.h
--- a/Engine/Source/Runtime/GameplayAbilities/Public/GameplayCueView.h
+++ b/Engine/Source/Runtime/GameplayAbilities/Public/GameplayCueView.h
@@ -6,6 +6,8 @@
 #include "GameplayEffect.h"
 #include "GameplayCueView.generated.h"
 
+class AGameplayCueActor;
+
 /** 
 	This is meant to give a base implementation for handling GameplayCues. It will handle very simple things like
 	creating and destroying ParticleSystemComponents, playing audio, etc.
@@ -26,7 +28,7 @@ struct FGameplayCueViewEffects
 	TWeakObjectPtr<UAudioComponent>	AudioComponent;
 
 	UPROPERTY()
-	TWeakObjectPtr<AActor>	SpawnedActor;
+	TWeakObjectPtr<AGameplayCueActor>	SpawnedActor;
 };
 
 USTRUCT()
@@ -57,15 +59,21 @@ struct FGameplayCueViewInfo
 
 	/** Local/remote sounds to play for weapon attacks against specific surfaces */
 	UPROPERTY(EditAnywhere, BlueprintReadWrite, Category = GameplayCue)
-	class USoundBase* Sound;
+	USoundBase* Sound;
 
 	/** Effects to play for weapon attacks against specific surfaces */
 	UPROPERTY(EditAnywhere, BlueprintReadWrite, Category = GameplayCue)
-	class UParticleSystem* ParticleSystem;
+	UParticleSystem* ParticleSystem;
+
+	/** Whether to attach the ParticleSystem to the owning actor */
+	UPROPERTY(EditAnywhere, BlueprintReadWrite, Category = GameplayCue)
+	bool AttachParticleSystem;
+
+
 
 	/** Effects to play for weapon attacks against specific surfaces */
 	UPROPERTY(EditAnywhere, BlueprintReadWrite, Category = GameplayCue)
-	class TSubclassOf<AActor> ActorClass;
+	class TSubclassOf<AGameplayCueActor> ActorClass;
 
 	virtual TSharedPtr<FGameplayCueViewEffects> SpawnViewEffects(AActor *Owner, TArray<UObject*> *SpawnedObjects, const FGameplayEffectInstigatorContext InstigatorContext) const;
 };
@@ -80,6 +88,7 @@ class GAMEPLAYABILITIES_API UGameplayCueView : public UDataAsset
 	TArray<FGameplayCueViewInfo>	Views;
 };
 
+
 USTRUCT()
 struct GAMEPLAYABILITIES_API FGameplayCueHandler
 {
@@ -111,9 +120,9 @@ struct GAMEPLAYABILITIES_API FGameplayCueHandler
 
 private:
 
-	FGameplayCueViewInfo * GetBestMatchingView(EGameplayCueEvent::Type Type, const FGameplayTag BaseTag);
+	FGameplayCueViewInfo * GetBestMatchingView(EGameplayCueEvent::Type Type, const FGameplayTag BaseTag, bool InstigatorLocal, bool TargetLocal);
 
 	void ClearEffects(TArray< TSharedPtr<FGameplayCueViewEffects > > &Effects);
 
-
+	bool OwnerIsLocallyControlled() const;
 };
\ No newline at end of file
diff --git a/Engine/Source/Runtime/GameplayAbilities/Public/GameplayEffect.h b/Engine/Source/Runtime/GameplayAbilities/Public/GameplayEffect.h
--- a/Engine/Source/Runtime/GameplayAbilities/Public/GameplayEffect.h
+++ b/Engine/Source/Runtime/GameplayAbilities/Public/GameplayEffect.h
@@ -9,19 +9,9 @@
 #include "GameplayEffectTypes.h"
 #include "GameplayEffect.generated.h"
 
-DECLARE_DELEGATE_OneParam(FOnGameplayAttributeEffectExecuted, struct FGameplayModifierEvaluatedData &);
-
-DECLARE_MULTICAST_DELEGATE_ThreeParams(FOnAttributeGameplayEffectSpecExected, const FGameplayAttribute &, const struct FGameplayEffectSpec &, struct FGameplayModifierEvaluatedData &);
-
 struct FActiveGameplayEffect;
 
-#define SKILL_SYSTEM_AGGREGATOR_DEBUG 1
 
-#if SKILL_SYSTEM_AGGREGATOR_DEBUG
-	#define SKILL_AGG_DEBUG( Format, ... ) *FString::Printf(Format, ##__VA_ARGS__)
-#else
-	#define SKILL_AGG_DEBUG( Format, ... ) NULL
-#endif
 
 class UGameplayEffect;
 class UAbilitySystemComponent;
@@ -476,75 +466,6 @@ struct FGameplayEffectLevelSpec
 	FGameplayAttribute	Attribute;
 };
 
-/**
- * FGameplayEffectInstigatorContext
- *	Data struct for an instigator. This is still being fleshed out. We will want to track actors but also be able to provide some level of tracking for actors that are destroyed.
- *	We may need to store some positional information as well.
- *
- */
-USTRUCT()
-struct FGameplayEffectInstigatorContext
-{
-	GENERATED_USTRUCT_BODY()
-
-	FGameplayEffectInstigatorContext()
-		: Instigator(NULL)
-		, InstigatorAbilitySystemComponent(NULL)
-	{
-	}
-
-	void GetOwnedGameplayTags(OUT FGameplayTagContainer &TagContainer)
-	{
-		IGameplayTagAssetInterface* TagInterface = InterfaceCast<IGameplayTagAssetInterface>(Instigator);
-		if (TagInterface)
-		{
-			TagInterface->GetOwnedGameplayTags(TagContainer);
-		}
-	}
-
-	void AddInstigator(class AActor *InInstigator);
-
-	void AddHitResult(const FHitResult InHitResult);
-
-	FString ToString() const
-	{
-		return Instigator ? Instigator->GetName() : FString(TEXT("NONE"));
-	}
-
-	/** Should always return the original instigator that started the whole chain */
-	AActor * GetOriginalInstigator()
-	{
-		return Instigator;
-	}
-
-	UAbilitySystemComponent * GetOriginalInstigatorAbilitySystemComponent() const
-	{
-		return InstigatorAbilitySystemComponent;
-	}
-	
-	/** Instigator controller */
-	UPROPERTY()
-	class AActor* Instigator;
-
-	UPROPERTY(NotReplicated)
-	UAbilitySystemComponent *InstigatorAbilitySystemComponent;
-
-	/** Trace information - may be NULL in many cases */
-	TSharedPtr<FHitResult>	HitResult;
-
-	bool NetSerialize(FArchive& Ar, class UPackageMap* Map, bool& bOutSuccess);
-};
-
-template<>
-struct TStructOpsTypeTraits< FGameplayEffectInstigatorContext > : public TStructOpsTypeTraitsBase
-{
-	enum
-	{
-		WithNetSerializer = true,
-		WithCopy = true		// Necessary so that TSharedPtr<FGameplayAbilityTargetData> Data is copied around
-	};
-};
-
 
 /**
  * FAggregatorRef
@@ -1093,7 +1014,7 @@ struct FGameplayEffectSpec
 	void PrintAll() const;
 
 	// Callbacks
-	FOnAttributeGameplayEffectSpecExected	OnExecute;
+	FOnAttributeGameplayEffectSpecExecuted	OnExecute;
 };
 
 /**
@@ -1212,10 +1133,14 @@ struct FActiveGameplayEffectsContainer : public FFastArraySerializer
 {
 	GENERATED_USTRUCT_BODY();
 
+	friend struct FActiveGameplayEffect;
+
 	FActiveGameplayEffectsContainer() : bNeedToRecalculateStacks(false) {};
 
 	UPROPERTY()
 	TArray< FActiveGameplayEffect >	GameplayEffects;
+
+	UAbilitySystemComponent *Owner;
 	
 	FActiveGameplayEffect & CreateNewActiveGameplayEffect(const FGameplayEffectSpec &Spec, int32 InPredictionKey);
 
@@ -1250,10 +1175,9 @@ struct FActiveGameplayEffectsContainer : public FFastArraySerializer
 
 	float GetGameplayEffectMagnitudeByTag(FActiveGameplayEffectHandle Handle, const FGameplayTag& InTag) const;
 
+	/** Registered as a callback for when a PropertyAggregator is dirty and we need to update the corresponding GameplayAttribute */
 	void OnPropertyAggregatorDirty(FAggregator* Aggregator, FGameplayAttribute Attribute);
 
-	UAbilitySystemComponent *Owner;
-
 	void CheckDuration(FActiveGameplayEffectHandle Handle);
 
 	void StacksNeedToRecalculate();
@@ -1281,23 +1205,34 @@ struct FActiveGameplayEffectsContainer : public FFastArraySerializer
 	TArray<float> GetActiveEffectsTimeRemaining(const FActiveGameplayEffectQuery Query) const;
 
 	int32 GetGameStateTime() const;
+
 	float GetWorldTime() const;
 
 	void OnDurationAggregatorDirty(FAggregator* Aggregator, UAbilitySystemComponent* Owner, FActiveGameplayEffectHandle Handle);
+	
+	FOnGameplayEffectTagCountChanged& RegisterGameplayTagEvent(FGameplayTag Tag);
 
 private:
 
 	FTimerHandle StackHandle;
 
-	bool InternalRemoveActiveGameplayEffect(int32 Idx);
-
 	FAggregatorRef & FindOrCreateAttributeAggregator(FGameplayAttribute Attribute);
 
 	FActiveGameplayEffectHandle LastAssignedHandle;
 
 	TMap<FGameplayAttribute, FAggregatorRef> OngoingPropertyEffects;
 
+	TMap<FGameplayTag, FOnGameplayEffectTagCountChanged> GameplayTagEventMap;
+	TMap<FGameplayTag, int32> GameplayTagCountMap;
+
 	bool IsNetAuthority() const;
+
+	/** Called internally to actually remove a GameplayEffect */
+	bool InternalRemoveActiveGameplayEffect(int32 Idx);
+
+	/** Called both in server side creation and replication creation/deletion */
+	void InternalOnActiveGameplayEffectAdded(const FActiveGameplayEffect& Effect);
+	void InternalOnActiveGameplayEffectRemoved(const FActiveGameplayEffect& Effect);
 };
 
 template<>
diff --git a/Engine/Source/Runtime/GameplayAbilities/Public/GameplayEffectTypes.h b/Engine/Source/Runtime/GameplayAbilities/Public/GameplayEffectTypes.h
--- a/Engine/Source/Runtime/GameplayAbilities/Public/GameplayEffectTypes.h
+++ b/Engine/Source/Runtime/GameplayAbilities/Public/GameplayEffectTypes.h
@@ -4,6 +4,16 @@
 
 #include "GameplayEffectTypes.generated.h"
 
+#define SKILL_SYSTEM_AGGREGATOR_DEBUG 1
+
+#if SKILL_SYSTEM_AGGREGATOR_DEBUG
+	#define SKILL_AGG_DEBUG( Format, ... ) *FString::Printf(Format, ##__VA_ARGS__)
+#else
+	#define SKILL_AGG_DEBUG( Format, ... ) NULL
+#endif
+
+class UAbilitySystemComponent;
+
 UENUM(BlueprintType)
 namespace EGameplayModOp
 {
@@ -115,15 +125,104 @@ namespace EGameplayEffectStackingPolicy
 	};
 }
 
+
+/**
+ * FGameplayEffectInstigatorContext
+ *	Data struct for an instigator. This is still being fleshed out. We will want to track actors but also be able to provide some level of tracking for actors that are destroyed.
+ *	We may need to store some positional information as well.
+ *
+ */
+USTRUCT()
+struct FGameplayEffectInstigatorContext
+{
+	GENERATED_USTRUCT_BODY()
+
+	FGameplayEffectInstigatorContext()
+		: Instigator(NULL)
+		, InstigatorAbilitySystemComponent(NULL)
+	{
+	}
+
+	void GetOwnedGameplayTags(OUT FGameplayTagContainer &TagContainer)
+	{
+		IGameplayTagAssetInterface* TagInterface = InterfaceCast<IGameplayTagAssetInterface>(Instigator);
+		if (TagInterface)
+		{
+			TagInterface->GetOwnedGameplayTags(TagContainer);
+		}
+	}
+
+	void AddInstigator(class AActor *InInstigator);
+
+	void AddHitResult(const FHitResult InHitResult);
+
+	FString ToString() const
+	{
+		return Instigator ? Instigator->GetName() : FString(TEXT("NONE"));
+	}
+
+	/** Should always return the original instigator that started the whole chain */
+	AActor * GetOriginalInstigator()
+	{
+		return Instigator;
+	}
+
+	UAbilitySystemComponent * GetOriginalInstigatorAbilitySystemComponent() const
+	{
+		return InstigatorAbilitySystemComponent;
+	}
+
+	bool IsLocallyControlled() const;
+	
+	/** Instigator controller */
+	UPROPERTY()
+	AActor* Instigator;
+
+	UPROPERTY(NotReplicated)
+	UAbilitySystemComponent *InstigatorAbilitySystemComponent;
+
+	/** Trace information - may be NULL in many cases */
+	TSharedPtr<FHitResult>	HitResult;
+
+	bool NetSerialize(FArchive& Ar, class UPackageMap* Map, bool& bOutSuccess);
+};
+
+template<>
+struct TStructOpsTypeTraits< FGameplayEffectInstigatorContext > : public TStructOpsTypeTraitsBase
+{
+	enum
+	{
+		WithNetSerializer = true,
+		WithCopy = true		// Necessary so that TSharedPtr<FGameplayAbilityTargetData> Data is copied around
+	};
+};
+
+// -----------------------------------------------------------
+
+USTRUCT(BlueprintType)
+struct FGameplayCueParameters
+{
+	GENERATED_USTRUCT_BODY()
+
+	UPROPERTY()
+	float NormalizedMagnitude;
+
+	UPROPERTY()
+	FGameplayEffectInstigatorContext InstigatorContext;
+
+	UPROPERTY()
+	FName MatchedTagName;
+};
+
 UENUM(BlueprintType)
 namespace EGameplayCueEvent
 {
 	enum Type
 	{
-		OnActive,
-		WhileActive,
-		Executed,
-		Removed
+		OnActive,		// Called when GameplayCue is activated.
+		WhileActive,	// Called when GameplayCue is active, even if it wasn't actually just applied (Join in progress, etc)
+		Executed,		// Called when a GameplayCue is executed: instant effects or periodic tick
+		Removed			// Called when GameplayCue is removed
 	};
 }
 
@@ -135,4 +234,12 @@ FString EGameplayModEffectToString(int32 Type);
 
 FString EGameplayEffectCopyPolicyToString(int32 Type);
 
-FString EGameplayEffectStackingPolicyToString(int32 Type);
\ No newline at end of file
+FString EGameplayEffectStackingPolicyToString(int32 Type);
+
+
+DECLARE_DELEGATE_OneParam(FOnGameplayAttributeEffectExecuted, struct FGameplayModifierEvaluatedData &);
+
+DECLARE_MULTICAST_DELEGATE_ThreeParams(FOnAttributeGameplayEffectSpecExecuted, const FGameplayAttribute &, const struct FGameplayEffectSpec &, struct FGameplayModifierEvaluatedData &);
+
+DECLARE_MULTICAST_DELEGATE_TwoParams(FOnGameplayEffectTagCountChanged, const FGameplayTag, int32 );
+
