GameplayAbilities
-task updates

diff --git a/Engine/Source/Editor/GameplayAbilitiesEditor/Classes/K2Node_LatentAbilityCall.h b/Engine/Source/Editor/GameplayAbilitiesEditor/Classes/K2Node_LatentAbilityCall.h
--- a/Engine/Source/Editor/GameplayAbilitiesEditor/Classes/K2Node_LatentAbilityCall.h
+++ b/Engine/Source/Editor/GameplayAbilitiesEditor/Classes/K2Node_LatentAbilityCall.h
@@ -3,6 +3,9 @@
 #pragma once
 #include "K2Node_LatentAbilityCall.generated.h"
 
+
+// TODO: Make this only available in GameplayAbility graphs!
+
 UCLASS()
 class UK2Node_LatentAbilityCall : public UK2Node_BaseAsyncTask
 {
diff --git a/Engine/Source/Runtime/GameplayAbilities/Private/Abilities/GameplayAbility.cpp b/Engine/Source/Runtime/GameplayAbilities/Private/Abilities/GameplayAbility.cpp
--- a/Engine/Source/Runtime/GameplayAbilities/Private/Abilities/GameplayAbility.cpp
+++ b/Engine/Source/Runtime/GameplayAbilities/Private/Abilities/GameplayAbility.cpp
@@ -116,6 +116,7 @@ bool UGameplayAbility::CanActivateAbility(const FGameplayAbilityActorInfo* Actor
 
 bool UGameplayAbility::CommitAbility(const FGameplayAbilityActorInfo* ActorInfo, const FGameplayAbilityActivationInfo ActivationInfo)
 {
+	// Last chance to fail (maybe we no longer have resources to commit since we after we started this ability activation)
 	if (!CommitCheck(ActorInfo, ActivationInfo))
 	{
 		return false;
@@ -126,6 +127,9 @@ bool UGameplayAbility::CommitAbility(const FGameplayAbilityActorInfo* ActorInfo,
 	// Fixme: Should we always call this or only if it is implemented? A noop may not hurt but could be bad for perf (storing a HasBlueprintCommit per instance isn't good either)
 	K2_CommitExecute();
 
+	// Broadcast this commitment
+	ActorInfo->AbilitySystemComponent->NotifyAbilityCommit(this);
+
 	return true;
 }
 
@@ -275,11 +279,15 @@ void UGameplayAbility::ActivateAbility(const FGameplayAbilityActorInfo* ActorInf
 
 void UGameplayAbility::PreActivate(const FGameplayAbilityActorInfo* ActorInfo, const FGameplayAbilityActivationInfo ActivationInfo)
 {
-	ActorInfo->AbilitySystemComponent->CancelAbilitiesWithTags(CancelAbilitiesWithTag, ActorInfo, ActivationInfo, this);
+	UAbilitySystemComponent* Comp = ActorInfo->AbilitySystemComponent.Get();
+
+	Comp->CancelAbilitiesWithTags(CancelAbilitiesWithTag, ActorInfo, ActivationInfo, this);
+
+	Comp->NotifyAbilityActivated(this);
 
 	// Become the 'Targeting Ability'
 	// These may need to be more robust - does every ability that activates become the targeting ability, or only certain ones?
-	ActorInfo->AbilitySystemComponent->SetTargetAbility(this);
+	Comp->SetTargetAbility(this);
 }
 
 void UGameplayAbility::CallActivateAbility(const FGameplayAbilityActorInfo* ActorInfo, const FGameplayAbilityActivationInfo ActivationInfo)
diff --git a/Engine/Source/Runtime/GameplayAbilities/Private/Abilities/Tasks/AbilityTask_WaitAbilityActivate.cpp b/Engine/Source/Runtime/GameplayAbilities/Private/Abilities/Tasks/AbilityTask_WaitAbilityActivate.cpp
new file mode 100644
--- /dev/null
+++ b/Engine/Source/Runtime/GameplayAbilities/Private/Abilities/Tasks/AbilityTask_WaitAbilityActivate.cpp
@@ -0,0 +1,58 @@
+
+#include "AbilitySystemPrivatePCH.h"
+#include "Abilities/Tasks/AbilityTask_WaitAbilityActivate.h"
+
+#include "AbilitySystemComponent.h"
+
+#include "Abilities/GameplayAbility.h"
+
+UAbilityTask_WaitAbilityActivate::UAbilityTask_WaitAbilityActivate(const class FPostConstructInitializeProperties& PCIP)
+	: Super(PCIP)
+{
+}
+
+UAbilityTask_WaitAbilityActivate* UAbilityTask_WaitAbilityActivate::WaitForAbilityActivate(UObject* WorldContextObject, FGameplayTag InWithTag, FGameplayTag InWithoutTag)
+{
+	check(WorldContextObject);
+	UGameplayAbility* ThisAbility = CastChecked<UGameplayAbility>(WorldContextObject);
+	if (ThisAbility)
+	{
+		UAbilityTask_WaitAbilityActivate * MyObj = NULL;
+		MyObj = NewObject<UAbilityTask_WaitAbilityActivate>();
+		MyObj->Ability = ThisAbility;
+		MyObj->WithTag = InWithTag;
+		MyObj->WithoutTag = InWithoutTag;
+
+		return MyObj;
+	}
+	return NULL;
+}
+
+void UAbilityTask_WaitAbilityActivate::Activate()
+{
+	if (Ability.IsValid())
+	{
+		const FGameplayAbilityActorInfo* Info = Ability.Get()->GetCurrentActorInfo();
+
+		UAbilitySystemComponent* ASC = Info->AbilitySystemComponent.Get();
+		ASC->AbilityActivatedCallbacks.AddUObject(this, &UAbilityTask_WaitAbilityActivate::OnAbilityActivate);
+	}
+}
+
+void UAbilityTask_WaitAbilityActivate::OnAbilityActivate(UGameplayAbility *ActivatedAbility)
+{
+	if ((WithTag.IsValid() && !ActivatedAbility->AbilityTags.HasTag(WithTag, EGameplayTagMatchType::IncludeParentTags, EGameplayTagMatchType::Explicit)) ||
+		(WithoutTag.IsValid() && ActivatedAbility->AbilityTags.HasTag(WithoutTag, EGameplayTagMatchType::IncludeParentTags, EGameplayTagMatchType::Explicit)))
+	{
+		// Failed tag check
+		return;
+	}
+
+	UAbilitySystemComponent* ASC = Ability.Get()->GetCurrentActorInfo()->AbilitySystemComponent.Get();
+	if (ASC)
+	{
+		ASC->AbilityActivatedCallbacks.RemoveUObject(this, &UAbilityTask_WaitAbilityActivate::OnAbilityActivate);
+	}
+
+	OnActivate.Broadcast(ActivatedAbility);
+}
\ No newline at end of file
diff --git a/Engine/Source/Runtime/GameplayAbilities/Private/Abilities/Tasks/AbilityTask_WaitAbilityCommit.cpp b/Engine/Source/Runtime/GameplayAbilities/Private/Abilities/Tasks/AbilityTask_WaitAbilityCommit.cpp
new file mode 100644
--- /dev/null
+++ b/Engine/Source/Runtime/GameplayAbilities/Private/Abilities/Tasks/AbilityTask_WaitAbilityCommit.cpp
@@ -0,0 +1,59 @@
+
+#include "AbilitySystemPrivatePCH.h"
+#include "Abilities/Tasks/AbilityTask_WaitAbilityCommit.h"
+
+#include "AbilitySystemComponent.h"
+
+#include "Abilities/GameplayAbility.h"
+
+UAbilityTask_WaitAbilityCommit::UAbilityTask_WaitAbilityCommit(const class FPostConstructInitializeProperties& PCIP)
+	: Super(PCIP)
+{
+}
+
+UAbilityTask_WaitAbilityCommit* UAbilityTask_WaitAbilityCommit::WaitForAbilityCommit(UObject* WorldContextObject, FGameplayTag InWithTag, FGameplayTag InWithoutTag)
+{
+	check(WorldContextObject);
+	UGameplayAbility* ThisAbility = CastChecked<UGameplayAbility>(WorldContextObject);
+	if (ThisAbility)
+	{
+		UAbilityTask_WaitAbilityCommit * MyObj = NULL;
+		MyObj = NewObject<UAbilityTask_WaitAbilityCommit>();
+		MyObj->Ability = ThisAbility;
+		MyObj->WithTag = InWithTag;
+		MyObj->WithoutTag = InWithoutTag;
+		
+
+		return MyObj;
+	}
+	return NULL;
+}
+
+void UAbilityTask_WaitAbilityCommit::Activate()
+{
+	if (Ability.IsValid())
+	{
+		const FGameplayAbilityActorInfo* Info = Ability.Get()->GetCurrentActorInfo();
+
+		UAbilitySystemComponent* ASC = Info->AbilitySystemComponent.Get();
+		ASC->AbilityCommitedCallbacks.AddUObject(this, &UAbilityTask_WaitAbilityCommit::OnAbilityCommit);
+	}
+}
+
+void UAbilityTask_WaitAbilityCommit::OnAbilityCommit(UGameplayAbility *ActivatedAbility)
+{
+	if ( (WithTag.IsValid() && !ActivatedAbility->AbilityTags.HasTag(WithTag, EGameplayTagMatchType::IncludeParentTags, EGameplayTagMatchType::Explicit)) ||
+		 (WithoutTag.IsValid() && ActivatedAbility->AbilityTags.HasTag(WithoutTag, EGameplayTagMatchType::IncludeParentTags, EGameplayTagMatchType::Explicit)))
+	{
+		// Failed tag check
+		return;
+	}
+
+	UAbilitySystemComponent* ASC = Ability.Get()->GetCurrentActorInfo()->AbilitySystemComponent.Get();
+	if (ASC)
+	{
+		ASC->AbilityCommitedCallbacks.RemoveUObject(this, &UAbilityTask_WaitAbilityCommit::OnAbilityCommit);
+	}
+
+	OnCommit.Broadcast(ActivatedAbility);
+}
\ No newline at end of file
diff --git a/Engine/Source/Runtime/GameplayAbilities/Private/Abilities/Tasks/AbilityTask_WaitConfirmCancel.cpp b/Engine/Source/Runtime/GameplayAbilities/Private/Abilities/Tasks/AbilityTask_WaitConfirmCancel.cpp
--- a/Engine/Source/Runtime/GameplayAbilities/Private/Abilities/Tasks/AbilityTask_WaitConfirmCancel.cpp
+++ b/Engine/Source/Runtime/GameplayAbilities/Private/Abilities/Tasks/AbilityTask_WaitConfirmCancel.cpp
@@ -14,13 +14,13 @@ UAbilityTask_WaitConfirmCancel::UAbilityTask_WaitConfirmCancel(const class FPost
 
 void UAbilityTask_WaitConfirmCancel::OnConfirmCallback()
 {
-	OnConfirm.Broadcast(true);
+	OnConfirm.Broadcast();
 	ASC->ConsumeAbilityConfirmCancel();
 }
 
 void UAbilityTask_WaitConfirmCancel::OnCancelCallback()
 {
-	OnCancel.Broadcast(false);
+	OnCancel.Broadcast();
 	ASC->ConsumeAbilityConfirmCancel();
 }
 
diff --git a/Engine/Source/Runtime/GameplayAbilities/Private/AbilitySystemComponent_Abilities.cpp b/Engine/Source/Runtime/GameplayAbilities/Private/AbilitySystemComponent_Abilities.cpp
--- a/Engine/Source/Runtime/GameplayAbilities/Private/AbilitySystemComponent_Abilities.cpp
+++ b/Engine/Source/Runtime/GameplayAbilities/Private/AbilitySystemComponent_Abilities.cpp
@@ -275,6 +275,16 @@ void UAbilitySystemComponent::SetUserAbilityActivationInhibited(bool NewInhibit)
 	}
 }
 
+void UAbilitySystemComponent::NotifyAbilityCommit(UGameplayAbility* Ability)
+{
+	AbilityCommitedCallbacks.Broadcast(Ability);
+}
+
+void UAbilitySystemComponent::NotifyAbilityActivated(UGameplayAbility* Ability)
+{
+	AbilityActivatedCallbacks.Broadcast(Ability);
+}
+
 // --------------------------------------------------------------------------
 
 void UAbilitySystemComponent::BindToInputComponent(UInputComponent *InputComponent)
@@ -396,15 +406,14 @@ bool UAbilitySystemComponent::ServerSetReplicatedTargetDataCancelled_Validate()
 void UAbilitySystemComponent::SetTargetAbility(UGameplayAbility* NewTargetingAbility)
 {
 	TargetingAbility = NewTargetingAbility;
-
-	//ReplicatedConfirmAbility = false;
-	//ReplicatedTargetData.Clear();
 }
 
 void UAbilitySystemComponent::ConsumeAbilityConfirmCancel()
 {
 	ReplicatedConfirmAbility = false;
 	ReplicatedCancelAbility = false;
+	ConfirmCallbacks.Clear();
+	CancelCallbacks.Clear();
 }
 
 void UAbilitySystemComponent::ConsumeAbilityTargetData()
diff --git a/Engine/Source/Runtime/GameplayAbilities/Public/Abilities/GameplayAbilityTypes.h b/Engine/Source/Runtime/GameplayAbilities/Public/Abilities/GameplayAbilityTypes.h
--- a/Engine/Source/Runtime/GameplayAbilities/Public/Abilities/GameplayAbilityTypes.h
+++ b/Engine/Source/Runtime/GameplayAbilities/Public/Abilities/GameplayAbilityTypes.h
@@ -7,11 +7,9 @@
 class UGameplayEffect;
 class UAnimInstance;
 class UAbilitySystemComponent;
+class UGameplayAbility;
+class AGameplayAbilityTargetActor;
 
-/** 
- *
- *
- */
 
 UENUM(BlueprintType)
 namespace EGameplayAbilityInstancingPolicy
@@ -417,4 +415,10 @@ struct TStructOpsTypeTraits<FGameplayAbilityTargetDataHandle> : public TStructOp
 };
 
 /** Generic callback for returning when target data is available */
-DECLARE_MULTICAST_DELEGATE_OneParam(FAbilityTargetData, FGameplayAbilityTargetDataHandle);
\ No newline at end of file
+DECLARE_MULTICAST_DELEGATE_OneParam(FAbilityTargetData, FGameplayAbilityTargetDataHandle);
+
+/** Used for cleaning up predicted data on network clients */
+DECLARE_MULTICAST_DELEGATE(FAbilitySystemComponentPredictionKeyClear);
+
+/** Generic delegate for ability 'events'/notifies */
+DECLARE_MULTICAST_DELEGATE_OneParam(FGenericAbilityDelegate, UGameplayAbility*);
\ No newline at end of file
diff --git a/Engine/Source/Runtime/GameplayAbilities/Public/Abilities/Tasks/AbilityTask.h b/Engine/Source/Runtime/GameplayAbilities/Public/Abilities/Tasks/AbilityTask.h
--- a/Engine/Source/Runtime/GameplayAbilities/Public/Abilities/Tasks/AbilityTask.h
+++ b/Engine/Source/Runtime/GameplayAbilities/Public/Abilities/Tasks/AbilityTask.h
@@ -12,5 +12,6 @@ class UAbilityTask : public UObject
 	UFUNCTION(BlueprintCallable, meta = (BlueprintInternalUseOnly = "true"), Category = "Abilities")
 	virtual void Activate();
 	
+	/** GameplayAbility that created us */
 	TWeakObjectPtr<UGameplayAbility> Ability;
 };
\ No newline at end of file
diff --git a/Engine/Source/Runtime/GameplayAbilities/Public/Abilities/Tasks/AbilityTask_WaitAbilityActivate.h b/Engine/Source/Runtime/GameplayAbilities/Public/Abilities/Tasks/AbilityTask_WaitAbilityActivate.h
new file mode 100644
--- /dev/null
+++ b/Engine/Source/Runtime/GameplayAbilities/Public/Abilities/Tasks/AbilityTask_WaitAbilityActivate.h
@@ -0,0 +1,34 @@
+// Copyright 1998-2014 Epic Games, Inc. All Rights Reserved.
+#pragma once
+#include "AbilityTask.h"
+#include "Abilities/GameplayAbilityTypes.h"
+#include "AbilityTask_WaitAbilityActivate.generated.h"
+
+DECLARE_DYNAMIC_MULTICAST_DELEGATE_OneParam(FWaitAbilityActivateDelegate, UGameplayAbility*, ActivatedAbility);
+
+class AActor;
+class UPrimitiveComponent;
+
+/**
+ *	Waits for the actor to activate another ability
+ */
+UCLASS(MinimalAPI)
+class UAbilityTask_WaitAbilityActivate : public UAbilityTask
+{
+	GENERATED_UCLASS_BODY()
+
+	UPROPERTY(BlueprintAssignable)
+	FWaitAbilityActivateDelegate	OnActivate;
+
+	virtual void Activate() override;
+
+	UFUNCTION()
+	void OnAbilityActivate(UGameplayAbility *ActivatedAbility);
+
+	/** Wait until an overlap occurs. This will need to be better fleshed out so we can specify game specific collision requirements */
+	UFUNCTION(BlueprintCallable, meta = (HidePin = "WorldContextObject", DefaultToSelf = "WorldContextObject", BlueprintInternalUseOnly = "TRUE"))
+	static UAbilityTask_WaitAbilityActivate* WaitForAbilityActivate(UObject* WorldContextObject, FGameplayTag WithTag, FGameplayTag WithoutTag);
+
+	FGameplayTag WithTag;
+	FGameplayTag WithoutTag;
+};
\ No newline at end of file
diff --git a/Engine/Source/Runtime/GameplayAbilities/Public/Abilities/Tasks/AbilityTask_WaitAbilityCommit.h b/Engine/Source/Runtime/GameplayAbilities/Public/Abilities/Tasks/AbilityTask_WaitAbilityCommit.h
new file mode 100644
--- /dev/null
+++ b/Engine/Source/Runtime/GameplayAbilities/Public/Abilities/Tasks/AbilityTask_WaitAbilityCommit.h
@@ -0,0 +1,35 @@
+// Copyright 1998-2014 Epic Games, Inc. All Rights Reserved.
+#pragma once
+#include "AbilityTask.h"
+#include "Abilities/GameplayAbilityTypes.h"
+#include "AbilityTask_WaitAbilityCommit.generated.h"
+
+
+DECLARE_DYNAMIC_MULTICAST_DELEGATE_OneParam(FWaitAbilityCommitDelegate, UGameplayAbility*, ActivatedAbility);
+
+class AActor;
+class UPrimitiveComponent;
+
+/**
+ *	Waits for the actor to activate another ability
+ */
+UCLASS(MinimalAPI)
+class UAbilityTask_WaitAbilityCommit : public UAbilityTask
+{
+	GENERATED_UCLASS_BODY()
+
+	UPROPERTY(BlueprintAssignable)
+	FWaitAbilityCommitDelegate	OnCommit;
+
+	virtual void Activate() override;
+
+	UFUNCTION()
+	void OnAbilityCommit(UGameplayAbility *ActivatedAbility);
+
+	/** Wait until an overlap occurs. This will need to be better fleshed out so we can specify game specific collision requirements */
+	UFUNCTION(BlueprintCallable, meta = (HidePin = "WorldContextObject", DefaultToSelf = "WorldContextObject", BlueprintInternalUseOnly = "TRUE"))
+	static UAbilityTask_WaitAbilityCommit* WaitForAbilityCommit(UObject* WorldContextObject, FGameplayTag WithTag, FGameplayTag WithoutTag);	
+
+	FGameplayTag WithTag;
+	FGameplayTag WithoutTag;
+};
\ No newline at end of file
diff --git a/Engine/Source/Runtime/GameplayAbilities/Public/Abilities/Tasks/AbilityTask_WaitConfirmCancel.h b/Engine/Source/Runtime/GameplayAbilities/Public/Abilities/Tasks/AbilityTask_WaitConfirmCancel.h
--- a/Engine/Source/Runtime/GameplayAbilities/Public/Abilities/Tasks/AbilityTask_WaitConfirmCancel.h
+++ b/Engine/Source/Runtime/GameplayAbilities/Public/Abilities/Tasks/AbilityTask_WaitConfirmCancel.h
@@ -5,7 +5,7 @@
 #include "GameplayAbility.h"
 #include "AbilityTask_WaitConfirmCancel.generated.h"
 
-DECLARE_DYNAMIC_MULTICAST_DELEGATE_OneParam(FWaitConfirmCancelDelegate, bool, Test);
+DECLARE_DYNAMIC_MULTICAST_DELEGATE(FWaitConfirmCancelDelegate);
 
 UCLASS(MinimalAPI)
 class UAbilityTask_WaitConfirmCancel : public UAbilityTask
diff --git a/Engine/Source/Runtime/GameplayAbilities/Public/AbilitySystemBlueprintLibrary.h b/Engine/Source/Runtime/GameplayAbilities/Public/AbilitySystemBlueprintLibrary.h
--- a/Engine/Source/Runtime/GameplayAbilities/Public/AbilitySystemBlueprintLibrary.h
+++ b/Engine/Source/Runtime/GameplayAbilities/Public/AbilitySystemBlueprintLibrary.h
@@ -37,13 +37,13 @@ class UAbilitySystemBlueprintLibrary : public UBlueprintFunctionLibrary
 	//		GameplayCue
 	// -------------------------------------------------------------------------------
 	
-	UFUNCTION(BlueprintCallable, Category="Ability|GameplayCue")
+	UFUNCTION(BlueprintPure, Category="Ability|GameplayCue")
 	static bool IsInstigatorLocallyControlled(FGameplayCueParameters Parameters);
 
-	UFUNCTION(BlueprintCallable, Category="Ability|GameplayCue")
+	UFUNCTION(BlueprintPure, Category = "Ability|GameplayCue")
 	static FHitResult GetHitResult(FGameplayCueParameters Parameters);
 
-	UFUNCTION(BlueprintCallable, Category = "Ability|GameplayCue")
+	UFUNCTION(BlueprintPure, Category = "Ability|GameplayCue")
 	static bool HasHitResult(FGameplayCueParameters Parameters);
 
 	// -------------------------------------------------------------------------------
diff --git a/Engine/Source/Runtime/GameplayAbilities/Public/AbilitySystemComponent.h b/Engine/Source/Runtime/GameplayAbilities/Public/AbilitySystemComponent.h
--- a/Engine/Source/Runtime/GameplayAbilities/Public/AbilitySystemComponent.h
+++ b/Engine/Source/Runtime/GameplayAbilities/Public/AbilitySystemComponent.h
@@ -5,15 +5,13 @@
 #include "Abilities/GameplayAbilityTypes.h"
 #include "AbilitySystemComponent.generated.h"
 
+
 GAMEPLAYABILITIES_API DECLARE_LOG_CATEGORY_EXTERN(LogAbilitySystemComponent, Log, All);
 
-/** Used for cleaning up predicted data on network clients */
-DECLARE_MULTICAST_DELEGATE(FAbilitySystemComponentPredictionKeyClear);
 
 /** Used to register callbacks to confirm/cancel input */
 DECLARE_DYNAMIC_MULTICAST_DELEGATE(FAbilityConfirmOrCancel);
 
-class AGameplayAbilityTargetActor;
 
 USTRUCT()
 struct GAMEPLAYABILITIES_API FAttributeDefaults
@@ -351,6 +349,13 @@ class GAMEPLAYABILITIES_API UAbilitySystemComponent : public UActorComponent
 	FAbilityConfirmOrCancel	ConfirmCallbacks;
 	FAbilityConfirmOrCancel	CancelCallbacks;
 
+	FGenericAbilityDelegate AbilityActivatedCallbacks;
+	FGenericAbilityDelegate AbilityCommitedCallbacks;
+
+	void NotifyAbilityCommit(UGameplayAbility* Ability);
+	void NotifyAbilityActivated(UGameplayAbility* Ability);
+
+
 	UPROPERTY()
 	TArray<AGameplayAbilityTargetActor*>	SpawnedTargetActors;
 
diff --git a/Engine/Source/Runtime/GameplayAbilities/Public/GameplayCueInterface.h b/Engine/Source/Runtime/GameplayAbilities/Public/GameplayCueInterface.h
--- a/Engine/Source/Runtime/GameplayAbilities/Public/GameplayCueInterface.h
+++ b/Engine/Source/Runtime/GameplayAbilities/Public/GameplayCueInterface.h
@@ -23,7 +23,7 @@ class GAMEPLAYABILITIES_API IGameplayCueInterface
 	virtual void HandleGameplayCues(AActor *Self, const FGameplayTagContainer& GameplayCueTags, EGameplayCueEvent::Type EventType, FGameplayCueParameters Parameters);
 	
 	/** Internal function to map ufunctions directly to gameplaycue tags */
-	UFUNCTION(BlueprintImplementableEvent, Category = GameplayCue, meta = (BlueprintInternalUseOnly = "true"))
+	UFUNCTION(BlueprintImplementableEvent, BlueprintCosmetic, Category = GameplayCue, meta = (BlueprintInternalUseOnly = "true"))
 	void BlueprintCustomHandler(EGameplayCueEvent::Type EventType, FGameplayCueParameters Parameters);
 
 	static void DispatchBlueprintCustomHandler(AActor* Actor, UFunction* Func, EGameplayCueEvent::Type EventType, FGameplayCueParameters Parameters);	
