Compiler fix for unused expression.

diff --git a/Engine/Source/Runtime/GameplayAbilities/Private/GameplayCueView.cpp b/Engine/Source/Runtime/GameplayAbilities/Private/GameplayCueView.cpp
--- a/Engine/Source/Runtime/GameplayAbilities/Private/GameplayCueView.cpp
+++ b/Engine/Source/Runtime/GameplayAbilities/Private/GameplayCueView.cpp
@@ -66,7 +66,7 @@ void FGameplayCueHandler::GameplayCueActivated(const FGameplayTagContainer & Gam
 		{
 			FGameplayCueParameters Params;
 			Params.NormalizedMagnitude = NormalizedMagnitude;
-			Params.InstigatorContext;
+			Params.InstigatorContext = InstigatorContext;
 
 			IGameplayCueInterface::DispatchBlueprintCustomHandler(Owner, Func, EGameplayCueEvent::OnActive, Params);
 		}
@@ -138,7 +138,7 @@ void FGameplayCueHandler::GameplayCueRemoved(const FGameplayTagContainer & Gamep
 		{
 			FGameplayCueParameters Params;
 			Params.NormalizedMagnitude = NormalizedMagnitude;
-			Params.InstigatorContext;
+			Params.InstigatorContext;// = InstigatorContext;
 
 			IGameplayCueInterface::DispatchBlueprintCustomHandler(Owner, Func, EGameplayCueEvent::Removed, Params);
 		}
