Compiler fix for unused expression.

diff --git a/Engine/Source/Runtime/GameplayAbilities/Private/GameplayCueView.cpp b/Engine/Source/Runtime/GameplayAbilities/Private/GameplayCueView.cpp
--- a/Engine/Source/Runtime/GameplayAbilities/Private/GameplayCueView.cpp
+++ b/Engine/Source/Runtime/GameplayAbilities/Private/GameplayCueView.cpp
@@ -138,7 +138,7 @@ void FGameplayCueHandler::GameplayCueRemoved(const FGameplayTagContainer & Gamep
 		{
 			FGameplayCueParameters Params;
 			Params.NormalizedMagnitude = NormalizedMagnitude;
-			Params.InstigatorContext;// = InstigatorContext;
+			Params.InstigatorContext = InstigatorContext;
 
 			IGameplayCueInterface::DispatchBlueprintCustomHandler(Owner, Func, EGameplayCueEvent::Removed, Params);
 		}
