Customn event node for GameplayCues
misc cleanup to AbilityTasks

diff --git a/Engine/Source/Editor/GameplayAbilitiesEditor/Classes/K2Node_GameplayCueEvent.h b/Engine/Source/Editor/GameplayAbilitiesEditor/Classes/K2Node_GameplayCueEvent.h
new file mode 100644
--- /dev/null
+++ b/Engine/Source/Editor/GameplayAbilitiesEditor/Classes/K2Node_GameplayCueEvent.h
@@ -0,0 +1,22 @@
+// Copyright 1998-2014 Epic Games, Inc. All Rights Reserved.
+
+#pragma once
+#include "GameplayCueView.h"
+#include "K2Node_GameplayCueEvent.generated.h"
+
+
+UCLASS()
+class UK2Node_GameplayCueEvent : public UK2Node_Event
+{
+	GENERATED_UCLASS_BODY()
+
+	// Begin UEdGraphNode interface
+	virtual FString GetTooltip() const override;
+	virtual FText GetNodeTitle(ENodeTitleType::Type TitleType) const override;
+	virtual FString GetNodeNativeTitle(ENodeTitleType::Type TitleType) const override;
+	// End UEdGraphNode interface
+
+	virtual FString GetCategoryName();
+
+	virtual void GetMenuEntries(struct FGraphContextMenuBuilder& ContextMenuBuilder) const override;
+};
\ No newline at end of file
diff --git a/Engine/Source/Editor/GameplayAbilitiesEditor/Private/K2Node_GameplayCueEvent.cpp b/Engine/Source/Editor/GameplayAbilitiesEditor/Private/K2Node_GameplayCueEvent.cpp
new file mode 100644
--- /dev/null
+++ b/Engine/Source/Editor/GameplayAbilitiesEditor/Private/K2Node_GameplayCueEvent.cpp
@@ -0,0 +1,72 @@
+// Copyright 1998-2014 Epic Games, Inc. All Rights Reserved.
+
+
+#include "AbilitySystemEditorPrivatePCH.h"
+#include "K2Node_GameplayCueEvent.h"
+#include "CompilerResultsLog.h"
+#include "K2ActionMenuBuilder.h"
+#include "GameplayTagsModule.h"
+#include "GameplayCueInterface.h"
+
+#define LOCTEXT_NAMESPACE "K2Node_PlayMontageAndWait"
+
+UK2Node_GameplayCueEvent::UK2Node_GameplayCueEvent(const class FPostConstructInitializeProperties& PCIP)
+	: Super(PCIP)
+{
+	EventSignatureName = GAMEPLAYABILITIES_BlueprintCustomHandler;
+	EventSignatureClass = UGameplayCueInterface::StaticClass();
+}
+
+FString UK2Node_GameplayCueEvent::GetCategoryName()
+{
+	return TEXT("GameplayCueEvent");
+}
+
+FString UK2Node_GameplayCueEvent::GetTooltip() const
+{
+	return TEXT("Handle GameplayCue Event");
+}
+
+FText UK2Node_GameplayCueEvent::GetNodeTitle(ENodeTitleType::Type TitleType) const
+{
+	return FText::FromName(CustomFunctionName);
+	//return LOCTEXT("HandleGameplayCueEvent", "HandleGameplaCueEvent");
+}
+
+FString UK2Node_GameplayCueEvent::GetNodeNativeTitle(ENodeTitleType::Type TitleType) const
+{
+	return CustomFunctionName.ToString();
+	//return LOCTEXT("HandleGameplayCueEvent", "HandleGameplaCueEvent").ToString();
+}
+
+void UK2Node_GameplayCueEvent::GetMenuEntries(FGraphContextMenuBuilder& Context) const
+{
+	Super::GetMenuEntries(Context);
+
+	const FString FunctionCategory(TEXT("GameplayCue Event"));
+
+	IGameplayTagsModule& GameplayTagsModule = IGameplayTagsModule::Get();
+	FGameplayTag RootTag = GameplayTagsModule.GetGameplayTagsManager().RequestGameplayTag(FName(TEXT("GameplayCue")));
+
+	FGameplayTagContainer CueTags = GameplayTagsModule.GetGameplayTagsManager().RequestGameplayTagChildren(RootTag);
+
+	for (auto It = CueTags.CreateConstIterator(); It; ++It)
+	{
+		FGameplayTag Tag = *It;
+
+		UK2Node_GameplayCueEvent* NodeTemplate = Context.CreateTemplateNode<UK2Node_GameplayCueEvent>();
+
+		NodeTemplate->CustomFunctionName = Tag.GetTagName();
+
+		const FString Category = FunctionCategory;
+		const FText MenuDesc = NodeTemplate->GetNodeTitle(ENodeTitleType::ListView);
+		const FString Tooltip = NodeTemplate->GetTooltip();
+		const FString Keywords = NodeTemplate->GetKeywords();
+
+		TSharedPtr<FEdGraphSchemaAction_K2NewNode> NodeAction = FK2ActionMenuBuilder::AddNewNodeAction(Context, Category, MenuDesc, Tooltip, 0, Keywords);
+		NodeAction->NodeTemplate = NodeTemplate;
+		NodeAction->SearchTitle = NodeAction->NodeTemplate->GetNodeSearchTitle();
+	}	
+}
+
+#undef LOCTEXT_NAMESPACE
\ No newline at end of file
diff --git a/Engine/Source/Runtime/GameplayAbilities/Private/Abilities/Tasks/AbilityTask.cpp b/Engine/Source/Runtime/GameplayAbilities/Private/Abilities/Tasks/AbilityTask.cpp
--- a/Engine/Source/Runtime/GameplayAbilities/Private/Abilities/Tasks/AbilityTask.cpp
+++ b/Engine/Source/Runtime/GameplayAbilities/Private/Abilities/Tasks/AbilityTask.cpp
@@ -10,4 +10,11 @@ UAbilityTask::UAbilityTask(const class FPostConstructInitializeProperties& PCIP)
 
 void UAbilityTask::Activate()
 {
+
 }
+
+void UAbilityTask::InitTask(UGameplayAbility* InAbility)
+{
+	Ability = InAbility;
+	AbilitySystemComponent = InAbility->GetCurrentActorInfo()->AbilitySystemComponent;
+}
\ No newline at end of file
diff --git a/Engine/Source/Runtime/GameplayAbilities/Private/Abilities/Tasks/AbilityTask_WaitAbilityActivate.cpp b/Engine/Source/Runtime/GameplayAbilities/Private/Abilities/Tasks/AbilityTask_WaitAbilityActivate.cpp
--- a/Engine/Source/Runtime/GameplayAbilities/Private/Abilities/Tasks/AbilityTask_WaitAbilityActivate.cpp
+++ b/Engine/Source/Runtime/GameplayAbilities/Private/Abilities/Tasks/AbilityTask_WaitAbilityActivate.cpp
@@ -19,7 +19,7 @@ UAbilityTask_WaitAbilityActivate* UAbilityTask_WaitAbilityActivate::WaitForAbili
 	{
 		UAbilityTask_WaitAbilityActivate * MyObj = NULL;
 		MyObj = NewObject<UAbilityTask_WaitAbilityActivate>();
-		MyObj->Ability = ThisAbility;
+		MyObj->InitTask(ThisAbility);
 		MyObj->WithTag = InWithTag;
 		MyObj->WithoutTag = InWithoutTag;
 
@@ -30,12 +30,9 @@ UAbilityTask_WaitAbilityActivate* UAbilityTask_WaitAbilityActivate::WaitForAbili
 
 void UAbilityTask_WaitAbilityActivate::Activate()
 {
-	if (Ability.IsValid())
+	if (AbilitySystemComponent.IsValid())
 	{
-		const FGameplayAbilityActorInfo* Info = Ability.Get()->GetCurrentActorInfo();
-
-		UAbilitySystemComponent* ASC = Info->AbilitySystemComponent.Get();
-		ASC->AbilityActivatedCallbacks.AddUObject(this, &UAbilityTask_WaitAbilityActivate::OnAbilityActivate);
+		AbilitySystemComponent->AbilityActivatedCallbacks.AddUObject(this, &UAbilityTask_WaitAbilityActivate::OnAbilityActivate);
 	}
 }
 
@@ -47,11 +44,10 @@ void UAbilityTask_WaitAbilityActivate::OnAbilityActivate(UGameplayAbility *Activ
 		// Failed tag check
 		return;
 	}
-
-	UAbilitySystemComponent* ASC = Ability.Get()->GetCurrentActorInfo()->AbilitySystemComponent.Get();
-	if (ASC)
+	
+	if (AbilitySystemComponent.IsValid())
 	{
-		ASC->AbilityActivatedCallbacks.RemoveUObject(this, &UAbilityTask_WaitAbilityActivate::OnAbilityActivate);
+		AbilitySystemComponent->AbilityActivatedCallbacks.RemoveUObject(this, &UAbilityTask_WaitAbilityActivate::OnAbilityActivate);
 	}
 
 	OnActivate.Broadcast(ActivatedAbility);
diff --git a/Engine/Source/Runtime/GameplayAbilities/Private/Abilities/Tasks/AbilityTask_WaitAbilityCommit.cpp b/Engine/Source/Runtime/GameplayAbilities/Private/Abilities/Tasks/AbilityTask_WaitAbilityCommit.cpp
--- a/Engine/Source/Runtime/GameplayAbilities/Private/Abilities/Tasks/AbilityTask_WaitAbilityCommit.cpp
+++ b/Engine/Source/Runtime/GameplayAbilities/Private/Abilities/Tasks/AbilityTask_WaitAbilityCommit.cpp
@@ -19,10 +19,9 @@ UAbilityTask_WaitAbilityCommit* UAbilityTask_WaitAbilityCommit::WaitForAbilityCo
 	{
 		UAbilityTask_WaitAbilityCommit * MyObj = NULL;
 		MyObj = NewObject<UAbilityTask_WaitAbilityCommit>();
-		MyObj->Ability = ThisAbility;
+		MyObj->InitTask(ThisAbility);
 		MyObj->WithTag = InWithTag;
 		MyObj->WithoutTag = InWithoutTag;
-		
 
 		return MyObj;
 	}
@@ -31,12 +30,9 @@ UAbilityTask_WaitAbilityCommit* UAbilityTask_WaitAbilityCommit::WaitForAbilityCo
 
 void UAbilityTask_WaitAbilityCommit::Activate()
 {
-	if (Ability.IsValid())
-	{
-		const FGameplayAbilityActorInfo* Info = Ability.Get()->GetCurrentActorInfo();
-
-		UAbilitySystemComponent* ASC = Info->AbilitySystemComponent.Get();
-		ASC->AbilityCommitedCallbacks.AddUObject(this, &UAbilityTask_WaitAbilityCommit::OnAbilityCommit);
+	if (AbilitySystemComponent.IsValid())	
+	{		
+		AbilitySystemComponent->AbilityCommitedCallbacks.AddUObject(this, &UAbilityTask_WaitAbilityCommit::OnAbilityCommit);
 	}
 }
 
@@ -49,10 +45,9 @@ void UAbilityTask_WaitAbilityCommit::OnAbilityCommit(UGameplayAbility *Activated
 		return;
 	}
 
-	UAbilitySystemComponent* ASC = Ability.Get()->GetCurrentActorInfo()->AbilitySystemComponent.Get();
-	if (ASC)
+	if (AbilitySystemComponent.IsValid())
 	{
-		ASC->AbilityCommitedCallbacks.RemoveUObject(this, &UAbilityTask_WaitAbilityCommit::OnAbilityCommit);
+		AbilitySystemComponent->AbilityCommitedCallbacks.RemoveUObject(this, &UAbilityTask_WaitAbilityCommit::OnAbilityCommit);
 	}
 
 	OnCommit.Broadcast(ActivatedAbility);
diff --git a/Engine/Source/Runtime/GameplayAbilities/Private/Abilities/Tasks/AbilityTask_WaitConfirmCancel.cpp b/Engine/Source/Runtime/GameplayAbilities/Private/Abilities/Tasks/AbilityTask_WaitConfirmCancel.cpp
--- a/Engine/Source/Runtime/GameplayAbilities/Private/Abilities/Tasks/AbilityTask_WaitConfirmCancel.cpp
+++ b/Engine/Source/Runtime/GameplayAbilities/Private/Abilities/Tasks/AbilityTask_WaitConfirmCancel.cpp
@@ -9,19 +9,27 @@
 UAbilityTask_WaitConfirmCancel::UAbilityTask_WaitConfirmCancel(const class FPostConstructInitializeProperties& PCIP)
 	: Super(PCIP)
 {
-	ASC = NULL;
+
 }
 
 void UAbilityTask_WaitConfirmCancel::OnConfirmCallback()
 {
 	OnConfirm.Broadcast();
-	ASC->ConsumeAbilityConfirmCancel();
+	if (AbilitySystemComponent.IsValid())
+	{
+		AbilitySystemComponent->ConsumeAbilityConfirmCancel();
+		MarkPendingKill();
+	}
 }
 
 void UAbilityTask_WaitConfirmCancel::OnCancelCallback()
 {
 	OnCancel.Broadcast();
-	ASC->ConsumeAbilityConfirmCancel();
+	if (AbilitySystemComponent.IsValid())
+	{
+		AbilitySystemComponent->ConsumeAbilityConfirmCancel();
+		MarkPendingKill();
+	}
 }
 
 UAbilityTask_WaitConfirmCancel* UAbilityTask_WaitConfirmCancel::WaitConfirmCancel(UObject* WorldContextObject)
@@ -32,7 +40,7 @@ UAbilityTask_WaitConfirmCancel* UAbilityTask_WaitConfirmCancel::WaitConfirmCance
 	{
 		UAbilityTask_WaitConfirmCancel * MyObj = NULL;
 		MyObj = NewObject<UAbilityTask_WaitConfirmCancel>();
-		MyObj->Ability = ThisAbility;
+		MyObj->InitTask(ThisAbility);
 
 		return MyObj;
 	}
@@ -41,22 +49,20 @@ UAbilityTask_WaitConfirmCancel* UAbilityTask_WaitConfirmCancel::WaitConfirmCance
 
 void UAbilityTask_WaitConfirmCancel::Activate()
 {
-	if (Ability.IsValid())
+	if (AbilitySystemComponent.IsValid())
 	{
 		const FGameplayAbilityActorInfo* Info = Ability.Get()->GetCurrentActorInfo();
 
-		ASC = Info->AbilitySystemComponent.Get();
-
+		// If not locally controlled (E.g., we are the server and this is an ability running on the client), check to see if they already sent us
+		// the Confirm/Cancel. If so, immediately end this task.
 		if (!Info->IsLocallyControlled())
 		{
-			if (ASC->ReplicatedConfirmAbility)
+			if (AbilitySystemComponent->ReplicatedConfirmAbility)
 			{
-				// It could have already been confirmed (we may have been behind a client predicting this ability, so it could have got here
-				// before we actually executed this task)
 				OnConfirmCallback();
 				return;
 			}
-			else if(ASC->ReplicatedCancelAbility)
+			else if (AbilitySystemComponent->ReplicatedCancelAbility)
 			{
 				OnCancelCallback();
 				return;
@@ -65,10 +71,9 @@ void UAbilityTask_WaitConfirmCancel::Activate()
 		
 		// We have to wait for the callback from the AbilitySystemComponent. Which may be instigated locally or through network replication
 
-	
-		ASC->ConfirmCallbacks.AddDynamic(this, &UAbilityTask_WaitConfirmCancel::OnConfirmCallback);	// Tell me if the confirm input is pressed
-		ASC->CancelCallbacks.AddDynamic(this, &UAbilityTask_WaitConfirmCancel::OnCancelCallback);	// Tell me if the cancel input is pressed
-
+		AbilitySystemComponent->ConfirmCallbacks.AddDynamic(this, &UAbilityTask_WaitConfirmCancel::OnConfirmCallback);	// Tell me if the confirm input is pressed
+		AbilitySystemComponent->CancelCallbacks.AddDynamic(this, &UAbilityTask_WaitConfirmCancel::OnCancelCallback);	// Tell me if the cancel input is pressed
+		
 		// Tell me if something else tells me to 'force' my target data (an ability may allow targeting during animation/limited time)
 		// Tell me if something else tells me to cancel my targeting
 	}
diff --git a/Engine/Source/Runtime/GameplayAbilities/Private/Abilities/Tasks/AbilityTask_WaitTargetData.cpp b/Engine/Source/Runtime/GameplayAbilities/Private/Abilities/Tasks/AbilityTask_WaitTargetData.cpp
--- a/Engine/Source/Runtime/GameplayAbilities/Private/Abilities/Tasks/AbilityTask_WaitTargetData.cpp
+++ b/Engine/Source/Runtime/GameplayAbilities/Private/Abilities/Tasks/AbilityTask_WaitTargetData.cpp
@@ -7,44 +7,37 @@
 UAbilityTask_WaitTargetData::UAbilityTask_WaitTargetData(const class FPostConstructInitializeProperties& PCIP)
 	: Super(PCIP)
 {
-	ASC = NULL;
+	
 }
 
 UAbilityTask_WaitTargetData* UAbilityTask_WaitTargetData::WaitTargetData(UObject* WorldContextObject, TSubclassOf<AGameplayAbilityTargetActor> InTargetClass)
 {
 	UAbilityTask_WaitTargetData * MyObj = NewObject<UAbilityTask_WaitTargetData>();
 	UGameplayAbility * ThisAbility = CastChecked<UGameplayAbility>(WorldContextObject);
-	MyObj->Ability = ThisAbility;
-	MyObj->ASC = ThisAbility->GetCurrentActorInfo()->AbilitySystemComponent.Get();
+	MyObj->InitTask(ThisAbility);
 	MyObj->TargetClass = InTargetClass;
 	return MyObj;	
 }
 
-UAbilityTask_WaitTargetData* UAbilityTask_WaitTargetData::CreateWaitTargetData(UObject* WorldContextObject, TSubclassOf<AGameplayAbilityTargetActor> InTargetClass)
-{
-	return WaitTargetData(WorldContextObject, InTargetClass);
-}
-
 void UAbilityTask_WaitTargetData::Activate()
 {
 	if (Ability.IsValid())
 	{
-		
 		if (!Ability.Get()->GetCurrentActorInfo()->IsLocallyControlled())
 		{
 			// If not locally controlled (server for remote client), see if TargetData was already sent
 			// else register callback for when it does get here
 
-			if (ASC->ReplicatedTargetData.IsValid())
+			if (AbilitySystemComponent->ReplicatedTargetData.IsValid())
 			{
-				ValidData.Broadcast(ASC->ReplicatedTargetData);
-				ASC->ConsumeAbilityTargetData();
+				ValidData.Broadcast(AbilitySystemComponent->ReplicatedTargetData);
+				Cleanup();
 				return;
 			}
 			else
 			{
-				ASC->ReplicatedTargetDataDelegate.AddUObject(this, &UAbilityTask_WaitTargetData::OnTargetDataReplicatedCallback);
-				ASC->ReplicatedTargetDataCancelledDelegate.AddDynamic(this, &UAbilityTask_WaitTargetData::OnTargetDataReplicatedCancelledCallback);
+				AbilitySystemComponent->ReplicatedTargetDataDelegate.AddUObject(this, &UAbilityTask_WaitTargetData::OnTargetDataReplicatedCallback);
+				AbilitySystemComponent->ReplicatedTargetDataCancelledDelegate.AddDynamic(this, &UAbilityTask_WaitTargetData::OnTargetDataReplicatedCancelledCallback);
 				return;
 			}
 		}
@@ -65,41 +58,72 @@ void UAbilityTask_WaitTargetData::Activate()
 				SpawnedActor->TargetDataReadyDelegate.AddUObject(this, &UAbilityTask_WaitTargetData::OnTargetDataReadyCallback);
 				SpawnedActor->CanceledDelegate.AddUObject(this, &UAbilityTask_WaitTargetData::OnTargetDataCancelledCallback);
 								
-				ASC->SetUserAbilityActivationInhibited(true);
-				ASC->SpawnedTargetActors.Push(SpawnedActor);
+				// User ability activation is inhibited while this is active
+				AbilitySystemComponent->SetUserAbilityActivationInhibited(true);
+				AbilitySystemComponent->SpawnedTargetActors.Push(SpawnedActor);
 
 				SpawnedActor->StartTargeting(Ability.Get());
+
+				MySpawnedTargetActor = SpawnedActor;
 			}
 		}
 	}
 }
 
+/** Valid TargetData was replicated to use (we are server, was sent from client) */
 void UAbilityTask_WaitTargetData::OnTargetDataReplicatedCallback(FGameplayAbilityTargetDataHandle Data)
 {
-	check(ASC);
+	check(AbilitySystemComponent.IsValid());
+
 	ValidData.Broadcast(Data);
-	ASC->ConsumeAbilityTargetData();
+
+	Cleanup();
 }
 
+/** Client cancelled this Targeting Task (we are the server */
 void UAbilityTask_WaitTargetData::OnTargetDataReplicatedCancelledCallback()
 {
-	check(ASC);
+	check(AbilitySystemComponent.IsValid());
+
 	Cancelled.Broadcast(FGameplayAbilityTargetDataHandle());
+	
+	Cleanup();
 }
 
+/** The TargetActor we spawned locally has called back with valid target data */
 void UAbilityTask_WaitTargetData::OnTargetDataReadyCallback(FGameplayAbilityTargetDataHandle Data)
 {
-	check(ASC);
-	ASC->ServerSetReplicatedTargetData(Data);
+	check(AbilitySystemComponent.IsValid());
+
+	AbilitySystemComponent->ServerSetReplicatedTargetData(Data);
 	ValidData.Broadcast(Data);
-	ASC->ConsumeAbilityTargetData();
-	ASC->SetUserAbilityActivationInhibited(false);
+	
+	Cleanup();
 }
 
+/** The TargetActor we spawned locally has called back with a cancel event (they still include the 'last/best' targetdata but the consumer of this may want to discard it) */
 void UAbilityTask_WaitTargetData::OnTargetDataCancelledCallback(FGameplayAbilityTargetDataHandle Data)
 {
-	check(ASC);
-	ASC->ServerSetReplicatedTargetDataCancelled();
+	check(AbilitySystemComponent.IsValid());
+
+	AbilitySystemComponent->ServerSetReplicatedTargetDataCancelled();
 	Cancelled.Broadcast(Data);
-	ASC->SetUserAbilityActivationInhibited(false);
+
+	Cleanup();
+}
+
+void UAbilityTask_WaitTargetData::Cleanup()
+{
+	AbilitySystemComponent->ConsumeAbilityTargetData();
+	AbilitySystemComponent->SetUserAbilityActivationInhibited(false);
+
+	AbilitySystemComponent->ReplicatedTargetDataDelegate.RemoveUObject(this, &UAbilityTask_WaitTargetData::OnTargetDataReplicatedCallback);
+	AbilitySystemComponent->ReplicatedTargetDataCancelledDelegate.RemoveDynamic(this, &UAbilityTask_WaitTargetData::OnTargetDataReplicatedCancelledCallback);
+
+	if (MySpawnedTargetActor.IsValid())
+	{
+		MySpawnedTargetActor->Destroy();
+	}
+
+	MarkPendingKill();
 }
\ No newline at end of file
diff --git a/Engine/Source/Runtime/GameplayAbilities/Public/Abilities/Tasks/AbilityTask.h b/Engine/Source/Runtime/GameplayAbilities/Public/Abilities/Tasks/AbilityTask.h
--- a/Engine/Source/Runtime/GameplayAbilities/Public/Abilities/Tasks/AbilityTask.h
+++ b/Engine/Source/Runtime/GameplayAbilities/Public/Abilities/Tasks/AbilityTask.h
@@ -11,7 +11,11 @@ class UAbilityTask : public UObject
 	// Called to trigger the actual task once the delegates have been set up
 	UFUNCTION(BlueprintCallable, meta = (BlueprintInternalUseOnly = "true"), Category = "Abilities")
 	virtual void Activate();
+
+	virtual void InitTask(UGameplayAbility* InAbility);
 	
 	/** GameplayAbility that created us */
 	TWeakObjectPtr<UGameplayAbility> Ability;
+
+	TWeakObjectPtr<UAbilitySystemComponent>	AbilitySystemComponent;
 };
\ No newline at end of file
diff --git a/Engine/Source/Runtime/GameplayAbilities/Public/Abilities/Tasks/AbilityTask_WaitConfirmCancel.h b/Engine/Source/Runtime/GameplayAbilities/Public/Abilities/Tasks/AbilityTask_WaitConfirmCancel.h
--- a/Engine/Source/Runtime/GameplayAbilities/Public/Abilities/Tasks/AbilityTask_WaitConfirmCancel.h
+++ b/Engine/Source/Runtime/GameplayAbilities/Public/Abilities/Tasks/AbilityTask_WaitConfirmCancel.h
@@ -27,10 +27,7 @@ class UAbilityTask_WaitConfirmCancel : public UAbilityTask
 	UFUNCTION(BlueprintCallable, meta=(HidePin = "WorldContextObject", DefaultToSelf = "WorldContextObject", BlueprintInternalUseOnly = "true", FriendlyName="Wait for Confirm Input"), Category="Abilities")
 	static UAbilityTask_WaitConfirmCancel* WaitConfirmCancel(UObject* WorldContextObject);	
 
-	virtual void Activate() override;
-
-	UPROPERTY()
-	UAbilitySystemComponent* ASC;
+	virtual void Activate() override;	
 
 public:
 	
diff --git a/Engine/Source/Runtime/GameplayAbilities/Public/Abilities/Tasks/AbilityTask_WaitOverlap.h b/Engine/Source/Runtime/GameplayAbilities/Public/Abilities/Tasks/AbilityTask_WaitOverlap.h
--- a/Engine/Source/Runtime/GameplayAbilities/Public/Abilities/Tasks/AbilityTask_WaitOverlap.h
+++ b/Engine/Source/Runtime/GameplayAbilities/Public/Abilities/Tasks/AbilityTask_WaitOverlap.h
@@ -9,6 +9,11 @@ DECLARE_DYNAMIC_MULTICAST_DELEGATE_OneParam(FWaitOverlapDelegate, FGameplayAbili
 class AActor;
 class UPrimitiveComponent;
 
+/**
+ *	Fixme: this is still incomplete and probablyh not what most games want for melee systems.
+ *		-Only actually activates on Blocking hits
+ *		-Uses first PrimitiveComponent instead of being able to specify arbitrary component.
+ */
 UCLASS(MinimalAPI)
 class UAbilityTask_WaitOverlap : public UAbilityTask
 {
diff --git a/Engine/Source/Runtime/GameplayAbilities/Public/Abilities/Tasks/AbilityTask_WaitTargetData.h b/Engine/Source/Runtime/GameplayAbilities/Public/Abilities/Tasks/AbilityTask_WaitTargetData.h
--- a/Engine/Source/Runtime/GameplayAbilities/Public/Abilities/Tasks/AbilityTask_WaitTargetData.h
+++ b/Engine/Source/Runtime/GameplayAbilities/Public/Abilities/Tasks/AbilityTask_WaitTargetData.h
@@ -35,13 +35,11 @@ class UAbilityTask_WaitTargetData: public UAbilityTask
 	UFUNCTION(BlueprintCallable, meta=(HidePin = "WorldContextObject", DefaultToSelf = "WorldContextObject", BlueprintInternalUseOnly = "true"), Category = "Abilities")
 	static UAbilityTask_WaitTargetData* WaitTargetData(UObject* WorldContextObject, TSubclassOf<AGameplayAbilityTargetActor> TargetClass);
 
-	UFUNCTION(BlueprintCallable, meta = (HidePin = "WorldContextObject", DefaultToSelf = "WorldContextObject", BlueprintInternalUseOnly = "true"), Category = "Abilities")
-	static UAbilityTask_WaitTargetData* CreateWaitTargetData(UObject* WorldContextObject, TSubclassOf<AGameplayAbilityTargetActor> TargetClass);
-
 	virtual void Activate() override;
 
 	TSubclassOf<AGameplayAbilityTargetActor> TargetClass;
 
-	UPROPERTY()
-	UAbilitySystemComponent* ASC;
+	void Cleanup();
+
+	TWeakObjectPtr<AActor>	MySpawnedTargetActor;
 };
\ No newline at end of file
