GameplayAbilities updates

diff --git a/Engine/Source/Editor/GameplayAbilitiesEditor/Private/AttributeDetails.cpp b/Engine/Source/Editor/GameplayAbilitiesEditor/Private/AttributeDetails.cpp
--- a/Engine/Source/Editor/GameplayAbilitiesEditor/Private/AttributeDetails.cpp
+++ b/Engine/Source/Editor/GameplayAbilitiesEditor/Private/AttributeDetails.cpp
@@ -9,6 +9,7 @@
 #include "KismetEditorUtilities.h"
 #include "AbilitySystemGlobals.h"
 #include "GameplayAbilitiesModule.h"
+#include "SGameplayAttributeGraphPin.h"
 
 #define LOCTEXT_NAMESPACE "AttributeDetailsCustomization"
 
@@ -52,6 +53,8 @@ void FAttributePropertyDetails::CustomizeHeader( TSharedRef<IPropertyHandle> Str
 		}
 	}
 
+	// Fixme: this should be unified to use SGameplayATtributeWidget instead of custom combo box
+
 	HeaderRow.
 		NameContent()
 		[
@@ -126,6 +129,7 @@ void FAttributePropertyDetails::OnChangeProperty(TSharedPtr<FString> ItemSelecte
 			{
 				const UObject *ObjPtr = Property;
 				MyProperty->SetValue(ObjPtr);
+				
 				return;
 			}
 		}
diff --git a/Engine/Source/Editor/GameplayAbilitiesEditor/Private/GameplayAbilitiesEditorModule.cpp b/Engine/Source/Editor/GameplayAbilitiesEditor/Private/GameplayAbilitiesEditorModule.cpp
--- a/Engine/Source/Editor/GameplayAbilitiesEditor/Private/GameplayAbilitiesEditorModule.cpp
+++ b/Engine/Source/Editor/GameplayAbilitiesEditor/Private/GameplayAbilitiesEditorModule.cpp
@@ -8,6 +8,7 @@
 #include "IAssetTypeActions.h"
 #include "AssetToolsModule.h"
 #include "AssetTypeActions_GameplayAbilitiesBlueprint.h"
+#include "GameplayAbilitiesGraphPanelPinFactory.h"
 
 class FGameplayAbilitiesEditorModule : public IGameplayAbilitiesEditorModule
 {
@@ -40,6 +41,9 @@ void FGameplayAbilitiesEditorModule::StartupModule()
 	TSharedRef<IAssetTypeActions> Action = MakeShareable(new FAssetTypeActions_GameplayAbilitiesBlueprint());
 	AssetTools.RegisterAssetTypeActions(Action);
 	CreatedAssetTypeActions.Add(Action);
+
+	TSharedPtr<FGameplayAbilitiesGraphPanelPinFactory> GameplayAbilitiesGraphPanelPinFactory = MakeShareable(new FGameplayAbilitiesGraphPanelPinFactory());
+	FEdGraphUtilities::RegisterVisualPinFactory(GameplayAbilitiesGraphPanelPinFactory);
 }
 
 void FGameplayAbilitiesEditorModule::RegisterAssetTypeAction(IAssetTools& AssetTools, TSharedRef<IAssetTypeActions> Action)
diff --git a/Engine/Source/Editor/GameplayAbilitiesEditor/Private/GameplayAbilitiesGraphPanelPinFactory.h b/Engine/Source/Editor/GameplayAbilitiesEditor/Private/GameplayAbilitiesGraphPanelPinFactory.h
new file mode 100644
--- /dev/null
+++ b/Engine/Source/Editor/GameplayAbilitiesEditor/Private/GameplayAbilitiesGraphPanelPinFactory.h
@@ -0,0 +1,20 @@
+// Copyright 1998-2014 Epic Games, Inc. All Rights Reserved.
+
+#pragma once
+
+#include "EdGraphUtilities.h"
+#include "SGameplayAttributeGraphPin.h"
+#include "AttributeSet.h"
+
+class FGameplayAbilitiesGraphPanelPinFactory: public FGraphPanelPinFactory
+{
+	virtual TSharedPtr<class SGraphPin> CreatePin(class UEdGraphPin* InPin) const override
+	{
+		const UEdGraphSchema_K2* K2Schema = GetDefault<UEdGraphSchema_K2>();
+		if (InPin->PinType.PinCategory == K2Schema->PC_Struct && InPin->PinType.PinSubCategoryObject == FGameplayAttribute::StaticStruct())
+		{
+			return SNew(SGameplayAttributeGraphPin, InPin);
+		}
+		return NULL;
+	}
+};
diff --git a/Engine/Source/Editor/GameplayAbilitiesEditor/Private/SGameplayAttributeGraphPin.cpp b/Engine/Source/Editor/GameplayAbilitiesEditor/Private/SGameplayAttributeGraphPin.cpp
new file mode 100644
--- /dev/null
+++ b/Engine/Source/Editor/GameplayAbilitiesEditor/Private/SGameplayAttributeGraphPin.cpp
@@ -0,0 +1,72 @@
+// Copyright 1998-2014 Epic Games, Inc. All Rights Reserved.
+
+#include "AbilitySystemEditorPrivatePCH.h"
+#include "SGameplayAttributeGraphPin.h"
+#include "SGameplayAttributeWidget.h"
+#include "GameplayTagsModule.h"
+#include "GameplayTags.h"
+
+#define LOCTEXT_NAMESPACE "K2Node"
+
+void SGameplayAttributeGraphPin::Construct( const FArguments& InArgs, UEdGraphPin* InGraphPinObj )
+{
+	SGraphPin::Construct( SGraphPin::FArguments(), InGraphPinObj );
+	LastSelectedProperty = NULL;
+
+}
+
+TSharedRef<SWidget>	SGameplayAttributeGraphPin::GetDefaultValueWidget()
+{
+	// Parse out current default value
+	// It will be in the form (Attribute=/Script/OrionGame.OrionHealthSet:Health)
+	
+	FString DefaultString = GraphPinObj->GetDefaultAsString();
+	FGameplayAttribute DefaultAttribute;
+
+	if (DefaultString.StartsWith(TEXT("(")) && DefaultString.EndsWith(TEXT(")")))
+	{
+		DefaultString = DefaultString.LeftChop(1);
+		DefaultString = DefaultString.RightChop(1);
+
+		DefaultString.Split("=", NULL, &DefaultString);
+
+		DefaultString = DefaultString.LeftChop(1);
+		DefaultString = DefaultString.RightChop(1);
+
+		DefaultAttribute.SetUProperty(FindObject<UProperty>(ANY_PACKAGE, *DefaultString));
+	}
+
+	//Create widget
+	return SNew(SVerticalBox)
+		+SVerticalBox::Slot()
+		.AutoHeight()
+		[
+			SNew(SGameplayAttributeWidget)
+			.OnAttributeChanged(this, &SGameplayAttributeGraphPin::OnAttributeChanged)
+			.DefaultProperty(DefaultAttribute.GetUProperty())
+		];
+}
+
+void SGameplayAttributeGraphPin::OnAttributeChanged(UProperty* SelectedAttribute)
+{
+	FString FinalValue;
+
+	if (SelectedAttribute == nullptr)
+	{
+		FinalValue = FString(TEXT("()"));
+	}
+	else
+	{
+		FinalValue = FString::Printf(TEXT("(Attribute=\"%s\")"), *SelectedAttribute->GetPathName());
+	}
+
+	GraphPinObj->GetSchema()->TrySetDefaultValue(*GraphPinObj, FinalValue);
+
+
+	FString DefTagString = GraphPinObj->GetDefaultAsString();
+
+	LastSelectedProperty = SelectedAttribute;
+
+}
+
+#undef LOCTEXT_NAMESPACE
diff --git a/Engine/Source/Editor/GameplayAbilitiesEditor/Private/SGameplayAttributeGraphPin.h b/Engine/Source/Editor/GameplayAbilitiesEditor/Private/SGameplayAttributeGraphPin.h
new file mode 100644
--- /dev/null
+++ b/Engine/Source/Editor/GameplayAbilitiesEditor/Private/SGameplayAttributeGraphPin.h
@@ -0,0 +1,26 @@
+// Copyright 1998-2014 Epic Games, Inc. All Rights Reserved.
+
+#pragma once
+
+#include "Slate.h"
+#include "SGraphPin.h"
+
+class SGameplayAttributeGraphPin : public SGraphPin
+{
+public:
+	SLATE_BEGIN_ARGS(SGameplayAttributeGraphPin) {}
+	SLATE_END_ARGS()
+
+	void Construct(const FArguments& InArgs, UEdGraphPin* InGraphPinObj);
+
+	// Begin SGraphPin interface
+	virtual TSharedRef<SWidget>	GetDefaultValueWidget() override;
+	// End SGraphPin interface
+
+	void OnAttributeChanged(UProperty* SelectedAttribute);
+
+	UProperty* LastSelectedProperty;
+
+private:
+
+};
diff --git a/Engine/Source/Editor/GameplayAbilitiesEditor/Private/SGameplayAttributeWidget.cpp b/Engine/Source/Editor/GameplayAbilitiesEditor/Private/SGameplayAttributeWidget.cpp
new file mode 100644
--- /dev/null
+++ b/Engine/Source/Editor/GameplayAbilitiesEditor/Private/SGameplayAttributeWidget.cpp
@@ -0,0 +1,76 @@
+// Copyright 1998-2014 Epic Games, Inc. All Rights Reserved.
+
+#include "AbilitySystemEditorPrivatePCH.h"
+#include "KismetEditorUtilities.h"
+#include "SGameplayAttributeWidget.h"
+
+#define LOCTEXT_NAMESPACE "K2Node"
+
+void SGameplayAttributeWidget::Construct(const FArguments& InArgs)
+{
+	FilterMetaData = InArgs._FilterMetaData;
+	OnAttributeChanged = InArgs._OnAttributeChanged;
+	SelectedProperty = InArgs._DefaultProperty;
+
+	TSharedPtr<FString> InitiallySelected = MakeShareable(new FString("None"));
+
+	PropertyOptions.Empty();
+	PropertyOptions.Add(InitiallySelected);
+
+	// Gather all UAttraibute classes
+	for (TObjectIterator<UClass> ClassIt; ClassIt; ++ClassIt)
+	{
+		UClass *Class = *ClassIt;
+		if (Class->IsChildOf(UAttributeSet::StaticClass()) && !FKismetEditorUtilities::IsClassABlueprintSkeleton(Class))
+		{
+			for (TFieldIterator<UProperty> PropertyIt(Class, EFieldIteratorFlags::ExcludeSuper); PropertyIt; ++PropertyIt)
+			{
+				UProperty *Property = *PropertyIt;
+
+				if (!FilterMetaData.IsEmpty() && Property->HasMetaData(*FilterMetaData))
+				{
+					continue;
+				}
+
+				TSharedPtr<FString> SelectableProperty = MakeShareable(new FString(FString::Printf(TEXT("%s.%s"), *Class->GetName(), *Property->GetName())));
+				if (Property == SelectedProperty)
+				{
+					InitiallySelected = SelectableProperty;
+				}
+				PropertyOptions.Add(SelectableProperty);
+			}
+		}
+	}
+
+	ChildSlot
+	[
+		SNew(STextComboBox)
+		.ContentPadding(FMargin(2.0f, 2.0f))
+		.OptionsSource(&PropertyOptions)
+		.InitiallySelectedItem(InitiallySelected)
+		.OnSelectionChanged(this, &SGameplayAttributeWidget::OnChangeProperty)
+	];
+}
+
+void SGameplayAttributeWidget::OnChangeProperty(TSharedPtr<FString> ItemSelected, ESelectInfo::Type SelectInfo)
+{
+	SelectedProperty = nullptr;
+	if (ItemSelected.IsValid())
+	{
+		FString FullString = *ItemSelected.Get();
+		FString ClassName;
+		FString PropertyName;
+
+		FullString.Split(TEXT("."), &ClassName, &PropertyName);
+
+		UClass *FoundClass = FindObject<UClass>(ANY_PACKAGE, *ClassName);
+		if (FoundClass)
+		{
+			SelectedProperty = FindFieldChecked<UProperty>(FoundClass, *PropertyName);
+		}
+	}	
+
+	OnAttributeChanged.ExecuteIfBound(SelectedProperty);
+}
+
+#undef LOCTEXT_NAMESPACE
diff --git a/Engine/Source/Editor/GameplayAbilitiesEditor/Private/SGameplayAttributeWidget.h b/Engine/Source/Editor/GameplayAbilitiesEditor/Private/SGameplayAttributeWidget.h
new file mode 100644
--- /dev/null
+++ b/Engine/Source/Editor/GameplayAbilitiesEditor/Private/SGameplayAttributeWidget.h
@@ -0,0 +1,37 @@
+// Copyright 1998-2014 Epic Games, Inc. All Rights Reserved.
+
+#pragma once
+
+#include "AttributeSet.h"
+#include "Slate.h"
+
+class SGameplayAttributeWidget : public SCompoundWidget
+{
+public:
+
+	/** Called when a tag status is changed */
+	DECLARE_DELEGATE_OneParam(FOnAttributeChanged, UProperty*)
+
+	SLATE_BEGIN_ARGS(SGameplayAttributeWidget)
+	: _FilterMetaData()
+	, _DefaultProperty(nullptr)
+	{}
+		SLATE_ARGUMENT(FString, FilterMetaData)
+		SLATE_ARGUMENT(UProperty*, DefaultProperty)
+		SLATE_EVENT(FOnAttributeChanged, OnAttributeChanged) // Called when a tag status changes
+	SLATE_END_ARGS()
+
+	void Construct(const FArguments& InArgs);	
+
+private:
+
+	void OnChangeProperty(TSharedPtr<FString> ItemSelected, ESelectInfo::Type SelectInfo);
+
+	FOnAttributeChanged OnAttributeChanged;
+
+	FString FilterMetaData;
+
+	TArray<TSharedPtr<FString>> PropertyOptions;
+
+	UProperty* SelectedProperty;
+};
diff --git a/Engine/Source/Runtime/GameplayAbilities/Private/Abilities/GameplayAbility.cpp b/Engine/Source/Runtime/GameplayAbilities/Private/Abilities/GameplayAbility.cpp
--- a/Engine/Source/Runtime/GameplayAbilities/Private/Abilities/GameplayAbility.cpp
+++ b/Engine/Source/Runtime/GameplayAbilities/Private/Abilities/GameplayAbility.cpp
@@ -473,6 +473,7 @@ void UGameplayAbility::ClientActivateAbilitySucceed_Internal(int32 PredictionKey
 
 void FGameplayAbilityActorInfo::InitFromActor(AActor *InActor)
 {
+	check(InActor);
 	Actor = InActor;
 
 	UStruct * Struct = FGameplayAbilityActorInfo::StaticStruct();
diff --git a/Engine/Source/Runtime/GameplayAbilities/Private/Abilities/Tasks/AbilityTask_WaitAttributeChange.cpp b/Engine/Source/Runtime/GameplayAbilities/Private/Abilities/Tasks/AbilityTask_WaitAttributeChange.cpp
new file mode 100644
--- /dev/null
+++ b/Engine/Source/Runtime/GameplayAbilities/Private/Abilities/Tasks/AbilityTask_WaitAttributeChange.cpp
@@ -0,0 +1,68 @@
+
+#include "AbilitySystemPrivatePCH.h"
+#include "Abilities/Tasks/AbilityTask_WaitAttributeChange.h"
+
+#include "AbilitySystemComponent.h"
+
+#include "Abilities/GameplayAbility.h"
+#include "GameplayEffectExtension.h"
+
+UAbilityTask_WaitAttributeChange::UAbilityTask_WaitAttributeChange(const class FPostConstructInitializeProperties& PCIP)
+	: Super(PCIP)
+{
+}
+
+UAbilityTask_WaitAttributeChange* UAbilityTask_WaitAttributeChange::WaitForAttributeChange(UObject* WorldContextObject, FGameplayAttribute InAttribute, FGameplayTag InWithTag, FGameplayTag InWithoutTag)
+{
+	check(WorldContextObject);
+	UGameplayAbility* ThisAbility = CastChecked<UGameplayAbility>(WorldContextObject);
+	if (ThisAbility)
+	{
+		UAbilityTask_WaitAttributeChange * MyObj = NULL;
+		MyObj = NewObject<UAbilityTask_WaitAttributeChange>();
+		MyObj->InitTask(ThisAbility);
+		MyObj->WithTag = InWithTag;
+		MyObj->WithoutTag = InWithoutTag;
+		MyObj->Attribute = InAttribute;
+
+		return MyObj;
+	}
+	return NULL;
+}
+
+void UAbilityTask_WaitAttributeChange::Activate()
+{
+	if (AbilitySystemComponent.IsValid())
+	{
+		AbilitySystemComponent->RegisterGameplayAttributeEvent(Attribute).AddUObject(this, &UAbilityTask_WaitAttributeChange::OnAttributeChange);
+	}
+}
+
+void UAbilityTask_WaitAttributeChange::OnAttributeChange(float NewValue, const FGameplayEffectModCallbackData* Data)
+{
+	if (Data == nullptr)
+	{
+		// There may be no execution data associated with this change, for example a GE being removed. 
+		// In this case, we auto fail any WithTag requirement and auto pass any WithoutTag requirement
+		if (WithTag.IsValid())
+		{
+			return;
+		}
+	}
+	else
+	{
+		if ((WithTag.IsValid() && !Data->EvaluatedData.Tags.HasTag(WithTag, EGameplayTagMatchType::IncludeParentTags, EGameplayTagMatchType::Explicit)) ||
+			(WithoutTag.IsValid() && Data->EvaluatedData.Tags.HasTag(WithoutTag, EGameplayTagMatchType::IncludeParentTags, EGameplayTagMatchType::Explicit)))
+		{
+			// Failed tag check
+			return;
+		}
+	}
+	
+	if (AbilitySystemComponent.IsValid())
+	{
+		AbilitySystemComponent->RegisterGameplayAttributeEvent(Attribute).RemoveUObject(this, &UAbilityTask_WaitAttributeChange::OnAttributeChange);
+	}
+
+	OnChange.Broadcast();
+}
\ No newline at end of file
diff --git a/Engine/Source/Runtime/GameplayAbilities/Private/Abilities/Tasks/AbilityTask_WaitGameplayEffectRemoved.cpp b/Engine/Source/Runtime/GameplayAbilities/Private/Abilities/Tasks/AbilityTask_WaitGameplayEffectRemoved.cpp
new file mode 100644
--- /dev/null
+++ b/Engine/Source/Runtime/GameplayAbilities/Private/Abilities/Tasks/AbilityTask_WaitGameplayEffectRemoved.cpp
@@ -0,0 +1,51 @@
+
+#include "AbilitySystemPrivatePCH.h"
+#include "Abilities/Tasks/AbilityTask_WaitGameplayEffectRemoved.h"
+
+#include "AbilitySystemComponent.h"
+
+#include "Abilities/GameplayAbility.h"
+
+UAbilityTask_WaitGameplayEffectRemoved::UAbilityTask_WaitGameplayEffectRemoved(const class FPostConstructInitializeProperties& PCIP)
+	: Super(PCIP)
+{
+}
+
+UAbilityTask_WaitGameplayEffectRemoved* UAbilityTask_WaitGameplayEffectRemoved::WaitForGameplayEffectRemoved(UObject* WorldContextObject, FActiveGameplayEffectHandle InHandle)
+{
+	check(WorldContextObject);
+	UGameplayAbility* ThisAbility = CastChecked<UGameplayAbility>(WorldContextObject);
+	if (ThisAbility)
+	{
+		UAbilityTask_WaitGameplayEffectRemoved * MyObj = NULL;
+		MyObj = NewObject<UAbilityTask_WaitGameplayEffectRemoved>();
+		MyObj->InitTask(ThisAbility);
+		MyObj->Handle = InHandle;
+
+		return MyObj;
+	}
+	return NULL;
+}
+
+void UAbilityTask_WaitGameplayEffectRemoved::Activate()
+{
+	if (AbilitySystemComponent.IsValid())	
+	{		
+		FOnActiveGameplayEffectRemoved* DelPtr = AbilitySystemComponent->OnGameplayEffectRemovedDelegate(Handle);
+		if (DelPtr)
+		{
+			DelPtr->AddUObject(this, &UAbilityTask_WaitGameplayEffectRemoved::OnGameplayEffectRemoved);
+		}
+		else
+		{
+			// GameplayEffect was already removed, treat this as a warning? Could be cases of immunity or chained gameplay rules that would instant remove something
+			OnGameplayEffectRemoved();
+		}
+	}
+}
+
+void UAbilityTask_WaitGameplayEffectRemoved::OnGameplayEffectRemoved()
+{
+	OnRemoved.Broadcast();
+	MarkPendingKill();
+}
\ No newline at end of file
diff --git a/Engine/Source/Runtime/GameplayAbilities/Private/AbilitySystemComponent.cpp b/Engine/Source/Runtime/GameplayAbilities/Private/AbilitySystemComponent.cpp
--- a/Engine/Source/Runtime/GameplayAbilities/Private/AbilitySystemComponent.cpp
+++ b/Engine/Source/Runtime/GameplayAbilities/Private/AbilitySystemComponent.cpp
@@ -204,6 +204,17 @@ float UAbilitySystemComponent::GetGameplayEffectMagnitudeByTag(FActiveGameplayEf
 	return ActiveGameplayEffects.GetGameplayEffectMagnitudeByTag(InHandle, InTag);
 }
 
+FOnActiveGameplayEffectRemoved* UAbilitySystemComponent::OnGameplayEffectRemovedDelegate(FActiveGameplayEffectHandle Handle)
+{
+	FActiveGameplayEffect* ActiveEffect = ActiveGameplayEffects.GetActiveGameplayEffect(Handle);
+	if (ActiveEffect)
+	{
+		return &ActiveEffect->OnRemovedDelegate;
+	}
+
+	return nullptr;
+}
+
 int32 UAbilitySystemComponent::GetNumActiveGameplayEffect() const
 {
 	return ActiveGameplayEffects.GetNumGameplayEffects();
@@ -235,6 +246,11 @@ FOnGameplayEffectTagCountChanged& UAbilitySystemComponent::RegisterGameplayTagEv
 	return ActiveGameplayEffects.RegisterGameplayTagEvent(Tag);
 }
 
+FOnGameplayAttributeChange& UAbilitySystemComponent::RegisterGameplayAttributeEvent(FGameplayAttribute Attribute)
+{
+	return ActiveGameplayEffects.RegisterGameplayAttributeEvent(Attribute);
+}
+
 void UAbilitySystemComponent::TEMP_ApplyActiveGameplayEffects()
 {
 	for (int32 idx=0; idx < ActiveGameplayEffects.GameplayEffects.Num(); ++idx)
diff --git a/Engine/Source/Runtime/GameplayAbilities/Private/AttributeSet.cpp b/Engine/Source/Runtime/GameplayAbilities/Private/AttributeSet.cpp
--- a/Engine/Source/Runtime/GameplayAbilities/Private/AttributeSet.cpp
+++ b/Engine/Source/Runtime/GameplayAbilities/Private/AttributeSet.cpp
@@ -1,4 +1,4 @@
-// Copyright 1998-2013 Epic Games, Inc. All Rights Reserved.
+// Copyright 1998-2014 Epic Games, Inc. All Rights Reserved.
 
 #include "AbilitySystemPrivatePCH.h"
 #include "AttributeSet.h"
@@ -85,14 +85,9 @@ void UAttributeSet::PrintDebug()
 	
 }
 
-void UAttributeSet::OnAttributeChange(UProperty *Property)
-{
-
-}
-
-FAttributeMetaData::FAttributeMetaData() :
-	MinValue(0.f),
-	MaxValue(1.f)
+FAttributeMetaData::FAttributeMetaData()
+	: MinValue(0.f)
+	, MaxValue(1.f)
 {
 
 }
diff --git a/Engine/Source/Runtime/GameplayAbilities/Private/GameplayEffect.cpp b/Engine/Source/Runtime/GameplayAbilities/Private/GameplayEffect.cpp
--- a/Engine/Source/Runtime/GameplayAbilities/Private/GameplayEffect.cpp
+++ b/Engine/Source/Runtime/GameplayAbilities/Private/GameplayEffect.cpp
@@ -1328,8 +1328,7 @@ void FActiveGameplayEffectsContainer::ExecuteActiveEffectsFrom(const FGameplayEf
 				
 				const float NewPropertyValue = Aggregator.Evaluate().Magnitude;
 
-				ABILITY_LOG(Log, TEXT("Property %s new value is: %.2f [was: %.2f]"), *Mod.Info.Attribute.GetName(), NewPropertyValue, CurrentValueOfProperty);
-				Owner->SetNumericAttribute(Mod.Info.Attribute, NewPropertyValue);
+				InternalUpdateNumericalAttribute(Mod.Info.Attribute, NewPropertyValue, &ExecuteData);
 			}
 
 			/** This should apply 'gameplay effect specific' rules, such as life steal, shields, etc */
@@ -1381,18 +1380,26 @@ void FActiveGameplayEffectsContainer::ExecuteActiveEffectsFrom(const FGameplayEf
 }
 
 void FActiveGameplayEffectsContainer::ExecutePeriodicGameplayEffect(FActiveGameplayEffectHandle Handle)
+{
+	FActiveGameplayEffect* ActiveEffect = GetActiveGameplayEffect(Handle);
+	if (ActiveEffect)
+	{
+		// Execute
+		ExecuteActiveEffectsFrom(ActiveEffect->Spec, FModifierQualifier().IgnoreHandle(Handle));
+	}
+}
+
+FActiveGameplayEffect* FActiveGameplayEffectsContainer::GetActiveGameplayEffect(const FActiveGameplayEffectHandle Handle)
 {
 	// Could make this a map for quicker lookup
 	for (FActiveGameplayEffect& Effect : GameplayEffects)
 	{
 		if (Effect.Handle == Handle)
 		{
-			// Execute
-			ExecuteActiveEffectsFrom(Effect.Spec, FModifierQualifier().IgnoreHandle(Effect.Handle));
-
-			return;
+			return &Effect;
 		}
 	}
+	return nullptr;
 }
 
 void FActiveGameplayEffectsContainer::AddDependancyToAttribute(FGameplayAttribute Attribute, const TWeakPtr<FAggregator> InDependant)
@@ -1521,11 +1528,21 @@ void FActiveGameplayEffectsContainer::OnPropertyAggregatorDirty(FAggregator* Agg
 	ABILITY_LOG_SCOPE(TEXT("FActiveGameplayEffectsContainer::OnPropertyAggregatorDirty"));
 
 	// Immediately calculate the newest value of the property			
-	float NewPropertyValue = Aggregator->Evaluate().Magnitude;
+	float NewValue = Aggregator->Evaluate().Magnitude;
 
-	ABILITY_LOG(Log, TEXT("Property %s new value is: %.2f"), *Attribute.GetName(), NewPropertyValue);
-	
-	Owner->SetNumericAttribute(Attribute, NewPropertyValue);
+	InternalUpdateNumericalAttribute(Attribute, NewValue, nullptr);
+}
+
+void FActiveGameplayEffectsContainer::InternalUpdateNumericalAttribute(FGameplayAttribute Attribute, float NewValue, const FGameplayEffectModCallbackData* ModData)
+{
+	ABILITY_LOG(Log, TEXT("Property %s new value is: %.2f"), *Attribute.GetName(), NewValue);
+	Owner->SetNumericAttribute(Attribute, NewValue);
+
+	FOnGameplayAttributeChange* Delegate = AttributeChangeDelegates.Find(Attribute);
+	if (Delegate)
+	{
+		Delegate->Broadcast(NewValue, ModData);
+	}
 }
 
 void FActiveGameplayEffectsContainer::StacksNeedToRecalculate()
@@ -1556,7 +1573,6 @@ FActiveGameplayEffect & FActiveGameplayEffectsContainer::CreateNewActiveGameplay
 			FTimerDelegate Delegate = FTimerDelegate::CreateUObject(Owner, &UAbilitySystemComponent::CheckDurationExpired, LastAssignedHandle);
 			TimerManager.SetTimer(NewEffect.DurationHandle, Delegate, Spec.GetDuration(), false, Spec.GetDuration());
 		}
-
 		// The timer manager moves things from the pending list to the active list after checking the active list on the first tick so we need to execute here
 		if (Spec.GetPeriod() != UGameplayEffect::NO_PERIOD)
 		{
@@ -1687,6 +1703,9 @@ bool FActiveGameplayEffectsContainer::InternalRemoveActiveGameplayEffect(int32 I
 /** This does cleanup that has to happen whether the effect is being removed locally or due to replication */
 void FActiveGameplayEffectsContainer::InternalOnActiveGameplayEffectRemoved(const FActiveGameplayEffect& Effect)
 {
+
+	Effect.OnRemovedDelegate.Broadcast();
+
 	// Update gameplaytag count and broadcast delegate if we are at 0
 	IGameplayTagsModule& GameplayTagsModule = IGameplayTagsModule::Get();
 	for (auto BaseTagIt = Effect.Spec.Def->OwnedTagsContainer.CreateConstIterator(); BaseTagIt; ++BaseTagIt)
@@ -2011,6 +2030,11 @@ FOnGameplayEffectTagCountChanged& FActiveGameplayEffectsContainer::RegisterGamep
 	return GameplayTagEventMap.FindOrAdd(Tag);
 }
 
+FOnGameplayAttributeChange& FActiveGameplayEffectsContainer::RegisterGameplayAttributeEvent(FGameplayAttribute Attribute)
+{
+	return AttributeChangeDelegates.FindOrAdd(Attribute);
+}
+
 // --------------------------------------------------------------------------------------------------------------------------------------------------------
 //
 //	FGameplayEffectLevelSpec
diff --git a/Engine/Source/Runtime/GameplayAbilities/Public/Abilities/Tasks/AbilityTask_WaitAbilityCommit.h b/Engine/Source/Runtime/GameplayAbilities/Public/Abilities/Tasks/AbilityTask_WaitAbilityCommit.h
--- a/Engine/Source/Runtime/GameplayAbilities/Public/Abilities/Tasks/AbilityTask_WaitAbilityCommit.h
+++ b/Engine/Source/Runtime/GameplayAbilities/Public/Abilities/Tasks/AbilityTask_WaitAbilityCommit.h
@@ -2,6 +2,7 @@
 #pragma once
 #include "AbilityTask.h"
 #include "Abilities/GameplayAbilityTypes.h"
+#include "AttributeSet.h"
 #include "AbilityTask_WaitAbilityCommit.generated.h"
 
 
@@ -28,7 +29,7 @@ class UAbilityTask_WaitAbilityCommit : public UAbilityTask
 
 	/** Wait until an overlap occurs. This will need to be better fleshed out so we can specify game specific collision requirements */
 	UFUNCTION(BlueprintCallable, meta = (HidePin = "WorldContextObject", DefaultToSelf = "WorldContextObject", BlueprintInternalUseOnly = "TRUE"))
-	static UAbilityTask_WaitAbilityCommit* WaitForAbilityCommit(UObject* WorldContextObject, FGameplayTag WithTag, FGameplayTag WithoutTag);	
+	static UAbilityTask_WaitAbilityCommit* WaitForAbilityCommit(UObject* WorldContextObject, FGameplayTag WithTag, FGameplayTag WithoutTage);	
 
 	FGameplayTag WithTag;
 	FGameplayTag WithoutTag;
diff --git a/Engine/Source/Runtime/GameplayAbilities/Public/Abilities/Tasks/AbilityTask_WaitAttributeChange.h b/Engine/Source/Runtime/GameplayAbilities/Public/Abilities/Tasks/AbilityTask_WaitAttributeChange.h
new file mode 100644
--- /dev/null
+++ b/Engine/Source/Runtime/GameplayAbilities/Public/Abilities/Tasks/AbilityTask_WaitAttributeChange.h
@@ -0,0 +1,35 @@
+// Copyright 1998-2014 Epic Games, Inc. All Rights Reserved.
+#pragma once
+#include "AbilityTask.h"
+#include "Abilities/GameplayAbilityTypes.h"
+#include "AttributeSet.h"
+#include "AbilityTask_WaitAttributeChange.generated.h"
+
+
+DECLARE_DYNAMIC_MULTICAST_DELEGATE(FWaitAttributeChangeDelegate);
+
+struct FGameplayEffectModCallbackData;
+
+/**
+ *	Waits for the actor to activate another ability
+ */
+UCLASS(MinimalAPI)
+class UAbilityTask_WaitAttributeChange : public UAbilityTask
+{
+	GENERATED_UCLASS_BODY()
+
+	UPROPERTY(BlueprintAssignable)
+	FWaitAttributeChangeDelegate	OnChange;
+
+	virtual void Activate() override;
+		
+	void OnAttributeChange(float NewValue, const FGameplayEffectModCallbackData*);
+
+	/** Wait until an overlap occurs. This will need to be better fleshed out so we can specify game specific collision requirements */
+	UFUNCTION(BlueprintCallable, meta = (HidePin = "WorldContextObject", DefaultToSelf = "WorldContextObject", BlueprintInternalUseOnly = "TRUE"))
+	static UAbilityTask_WaitAttributeChange* WaitForAttributeChange(UObject* WorldContextObject, FGameplayAttribute Attribute, FGameplayTag WithTag, FGameplayTag WithoutTag);
+
+	FGameplayTag WithTag;
+	FGameplayTag WithoutTag;
+	FGameplayAttribute	Attribute;
+};
\ No newline at end of file
diff --git a/Engine/Source/Runtime/GameplayAbilities/Public/Abilities/Tasks/AbilityTask_WaitGameplayEffectRemoved.h b/Engine/Source/Runtime/GameplayAbilities/Public/Abilities/Tasks/AbilityTask_WaitGameplayEffectRemoved.h
new file mode 100644
--- /dev/null
+++ b/Engine/Source/Runtime/GameplayAbilities/Public/Abilities/Tasks/AbilityTask_WaitGameplayEffectRemoved.h
@@ -0,0 +1,35 @@
+// Copyright 1998-2014 Epic Games, Inc. All Rights Reserved.
+#pragma once
+#include "AbilityTask.h"
+#include "Abilities/GameplayAbilityTypes.h"
+#include "GameplayEffect.h"
+#include "AbilityTask_WaitGameplayEffectRemoved.generated.h"
+
+
+DECLARE_DYNAMIC_MULTICAST_DELEGATE(FWaitGameplayEffectRemovedDelegate);
+
+class AActor;
+class UPrimitiveComponent;
+
+/**
+ *	Waits for the actor to activate another ability
+ */
+UCLASS(MinimalAPI)
+class UAbilityTask_WaitGameplayEffectRemoved : public UAbilityTask
+{
+	GENERATED_UCLASS_BODY()
+
+	UPROPERTY(BlueprintAssignable)
+	FWaitGameplayEffectRemovedDelegate	OnRemoved;
+
+	virtual void Activate() override;
+
+	UFUNCTION()
+	void OnGameplayEffectRemoved();
+
+	/** Wait until an overlap occurs. This will need to be better fleshed out so we can specify game specific collision requirements */
+	UFUNCTION(BlueprintCallable, meta = (HidePin = "WorldContextObject", DefaultToSelf = "WorldContextObject", BlueprintInternalUseOnly = "TRUE"))
+	static UAbilityTask_WaitGameplayEffectRemoved* WaitForGameplayEffectRemoved(UObject* WorldContextObject, FActiveGameplayEffectHandle Handle);
+
+	FActiveGameplayEffectHandle Handle;
+};
\ No newline at end of file
diff --git a/Engine/Source/Runtime/GameplayAbilities/Public/AbilitySystemComponent.h b/Engine/Source/Runtime/GameplayAbilities/Public/AbilitySystemComponent.h
--- a/Engine/Source/Runtime/GameplayAbilities/Public/AbilitySystemComponent.h
+++ b/Engine/Source/Runtime/GameplayAbilities/Public/AbilitySystemComponent.h
@@ -5,10 +5,8 @@
 #include "Abilities/GameplayAbilityTypes.h"
 #include "AbilitySystemComponent.generated.h"
 
-
 GAMEPLAYABILITIES_API DECLARE_LOG_CATEGORY_EXTERN(LogAbilitySystemComponent, Log, All);
 
-
 /** Used to register callbacks to confirm/cancel input */
 DECLARE_DYNAMIC_MULTICAST_DELEGATE(FAbilityConfirmOrCancel);
 
@@ -214,9 +212,13 @@ class GAMEPLAYABILITIES_API UAbilitySystemComponent : public UActorComponent
 	/** Allow events to be registered for specific gameplay tags being added or removed */
 	FOnGameplayEffectTagCountChanged& RegisterGameplayTagEvent(FGameplayTag Tag);
 
+	FOnGameplayAttributeChange& RegisterGameplayAttributeEvent(FGameplayAttribute Attribute);
+
 	// --------------------------------------------
 	// Possibly useful but not primary API functions:
 	// --------------------------------------------
+	
+	FOnActiveGameplayEffectRemoved* OnGameplayEffectRemovedDelegate(FActiveGameplayEffectHandle Handle);
 
 	FActiveGameplayEffectHandle ApplyGameplayEffectToTarget(UGameplayEffect *GameplayEffect, UAbilitySystemComponent *Target, float Level = FGameplayEffectLevelSpec::INVALID_LEVEL, FModifierQualifier BaseQualifier = FModifierQualifier());
 
diff --git a/Engine/Source/Runtime/GameplayAbilities/Public/AttributeSet.h b/Engine/Source/Runtime/GameplayAbilities/Public/AttributeSet.h
--- a/Engine/Source/Runtime/GameplayAbilities/Public/AttributeSet.h
+++ b/Engine/Source/Runtime/GameplayAbilities/Public/AttributeSet.h
@@ -84,10 +84,6 @@ public:
 	void ClampFromMetaDataTable(const UDataTable *DataTable);
 
 	virtual void PrintDebug();
-
-private:
-
-	void OnAttributeChange(UProperty *Property);
 };
 
 USTRUCT()
diff --git a/Engine/Source/Runtime/GameplayAbilities/Public/GameplayEffect.h b/Engine/Source/Runtime/GameplayAbilities/Public/GameplayEffect.h
--- a/Engine/Source/Runtime/GameplayAbilities/Public/GameplayEffect.h
+++ b/Engine/Source/Runtime/GameplayAbilities/Public/GameplayEffect.h
@@ -4,7 +4,6 @@
 
 #include "GameplayTagContainer.h"
 #include "GameplayTagAssetInterface.h"
-#include "AttributeSet.h"
 #include "TimerManager.h"
 #include "GameplayEffectTypes.h"
 #include "GameplayEffect.generated.h"
@@ -1012,9 +1011,6 @@ struct FGameplayEffectSpec
 	}
 
 	void PrintAll() const;
-
-	// Callbacks
-	FOnAttributeGameplayEffectSpecExecuted	OnExecute;
 };
 
 /**
@@ -1062,6 +1058,8 @@ struct FActiveGameplayEffect : public FFastArraySerializerItem
 	UPROPERTY(NotReplicated)
 	float StartWorldTime;
 
+	FOnActiveGameplayEffectRemoved	OnRemovedDelegate;
+
 	FTimerHandle PeriodHandle;
 
 	FTimerHandle DurationHandle;
@@ -1142,7 +1140,9 @@ struct FActiveGameplayEffectsContainer : public FFastArraySerializer
 
 	UAbilitySystemComponent *Owner;
 	
-	FActiveGameplayEffect & CreateNewActiveGameplayEffect(const FGameplayEffectSpec &Spec, int32 InPredictionKey);
+	FActiveGameplayEffect& CreateNewActiveGameplayEffect(const FGameplayEffectSpec &Spec, int32 InPredictionKey);
+
+	FActiveGameplayEffect* GetActiveGameplayEffect(const FActiveGameplayEffectHandle Handle);
 
 	// returns true if none of the active effects provide immunity to Spec
 	// returns false if one (or more) of the active effects provides immunity to Spec
@@ -1162,6 +1162,8 @@ struct FActiveGameplayEffectsContainer : public FFastArraySerializer
 
 	float GetGameplayEffectMagnitude(FActiveGameplayEffectHandle Handle, FGameplayAttribute Attribute) const;
 
+
+
 	// returns true if the handle points to an effect in this container that is not a stacking effect or an effect in this container that does stack and is applied by the current stacking rules
 	// returns false if the handle points to an effect that is not in this container or is not applied because of the current stacking rules
 	bool IsGameplayEffectActive(FActiveGameplayEffectHandle Handle) const;
@@ -1212,19 +1214,25 @@ struct FActiveGameplayEffectsContainer : public FFastArraySerializer
 	
 	FOnGameplayEffectTagCountChanged& RegisterGameplayTagEvent(FGameplayTag Tag);
 
+	FOnGameplayAttributeChange& RegisterGameplayAttributeEvent(FGameplayAttribute Attribute);
+
 private:
 
 	FTimerHandle StackHandle;
 
-	FAggregatorRef & FindOrCreateAttributeAggregator(FGameplayAttribute Attribute);
+	FAggregatorRef& FindOrCreateAttributeAggregator(FGameplayAttribute Attribute);
 
 	FActiveGameplayEffectHandle LastAssignedHandle;
 
 	TMap<FGameplayAttribute, FAggregatorRef> OngoingPropertyEffects;
 
+	TMap<FGameplayAttribute, FOnGameplayAttributeChange> AttributeChangeDelegates;
+
 	TMap<FGameplayTag, FOnGameplayEffectTagCountChanged> GameplayTagEventMap;
 	TMap<FGameplayTag, int32> GameplayTagCountMap;
 
+	void InternalUpdateNumericalAttribute(FGameplayAttribute Attribute, float NewValue, const FGameplayEffectModCallbackData* ModData);
+
 	bool IsNetAuthority() const;
 
 	/** Called internally to actually remove a GameplayEffect */
diff --git a/Engine/Source/Runtime/GameplayAbilities/Public/GameplayEffectTypes.h b/Engine/Source/Runtime/GameplayAbilities/Public/GameplayEffectTypes.h
--- a/Engine/Source/Runtime/GameplayAbilities/Public/GameplayEffectTypes.h
+++ b/Engine/Source/Runtime/GameplayAbilities/Public/GameplayEffectTypes.h
@@ -2,6 +2,7 @@
 
 #pragma once
 
+#include "AttributeSet.h"
 #include "GameplayEffectTypes.generated.h"
 
 #define SKILL_SYSTEM_AGGREGATOR_DEBUG 1
@@ -14,6 +15,9 @@
 
 class UAbilitySystemComponent;
 
+struct FGameplayEffectSpec;
+struct FGameplayEffectModCallbackData;
+
 UENUM(BlueprintType)
 namespace EGameplayModOp
 {
@@ -132,7 +136,7 @@ namespace EGameplayEffectStackingPolicy
  *	We may need to store some positional information as well.
  *
  */
-USTRUCT()
+USTRUCT(BlueprintType)
 struct FGameplayEffectInstigatorContext
 {
 	GENERATED_USTRUCT_BODY()
@@ -237,9 +241,10 @@ FString EGameplayEffectCopyPolicyToString(int32 Type);
 FString EGameplayEffectStackingPolicyToString(int32 Type);
 
 
-DECLARE_DELEGATE_OneParam(FOnGameplayAttributeEffectExecuted, struct FGameplayModifierEvaluatedData &);
-
-DECLARE_MULTICAST_DELEGATE_ThreeParams(FOnAttributeGameplayEffectSpecExecuted, const FGameplayAttribute &, const struct FGameplayEffectSpec &, struct FGameplayModifierEvaluatedData &);
+DECLARE_DELEGATE_OneParam(FOnGameplayAttributeEffectExecuted, struct FGameplayModifierEvaluatedData&);
 
 DECLARE_MULTICAST_DELEGATE_TwoParams(FOnGameplayEffectTagCountChanged, const FGameplayTag, int32 );
 
+DECLARE_MULTICAST_DELEGATE(FOnActiveGameplayEffectRemoved);
+
+DECLARE_MULTICAST_DELEGATE_TwoParams(FOnGameplayAttributeChange, float ,const FGameplayEffectModCallbackData*);
\ No newline at end of file
