missed pragma optimize

diff --git a/Engine/Source/Runtime/GameplayAbilities/Private/Abilities/Tasks/AbilityTask_SpawnActor.cpp b/Engine/Source/Runtime/GameplayAbilities/Private/Abilities/Tasks/AbilityTask_SpawnActor.cpp
--- a/Engine/Source/Runtime/GameplayAbilities/Private/Abilities/Tasks/AbilityTask_SpawnActor.cpp
+++ b/Engine/Source/Runtime/GameplayAbilities/Private/Abilities/Tasks/AbilityTask_SpawnActor.cpp
@@ -4,7 +4,6 @@
 #include "AbilitySystemComponent.h"
 #include "Abilities/Tasks/AbilityTask_SpawnActor.h"
 
-#pragma optimize( "", off )
 
 UAbilityTask_SpawnActor::UAbilityTask_SpawnActor(const class FPostConstructInitializeProperties& PCIP)
 	: Super(PCIP)
