GameplayAbilities Minor updates

diff --git a/Engine/Source/Runtime/GameplayAbilities/Private/Abilities/GameplayAbility.cpp b/Engine/Source/Runtime/GameplayAbilities/Private/Abilities/GameplayAbility.cpp
--- a/Engine/Source/Runtime/GameplayAbilities/Private/Abilities/GameplayAbility.cpp
+++ b/Engine/Source/Runtime/GameplayAbilities/Private/Abilities/GameplayAbility.cpp
@@ -34,6 +34,20 @@ UGameplayAbility::UGameplayAbility(const class FPostConstructInitializePropertie
 		UFunction* PredictiveActivateFunction = GetClass()->FindFunctionByName(FuncName);
 		HasBlueprintPredictiveActivate = PredictiveActivateFunction && PredictiveActivateFunction->GetOuter()->IsA(UBlueprintGeneratedClass::StaticClass());
 	}
+
+
+	
+#if WITH_EDITOR
+	/** Autoregister abilities with the blueprint debugger in the editor.*/
+	if (!HasAnyFlags(RF_ClassDefaultObject))
+	{
+		UBlueprint* BP = Cast<UBlueprint>(GetClass()->ClassGeneratedBy);
+		if (BP && (BP->GetWorldBeingDebugged() == nullptr || BP->GetWorldBeingDebugged() == GetWorld()))
+		{
+			BP->SetObjectBeingDebugged(this);
+		}
+	}
+#endif
 }
 
 void UGameplayAbility::PostNetInit()
diff --git a/Engine/Source/Runtime/GameplayAbilities/Private/Abilities/GameplayAbilityTargetActor_SingleLineTrace.cpp b/Engine/Source/Runtime/GameplayAbilities/Private/Abilities/GameplayAbilityTargetActor_SingleLineTrace.cpp
--- a/Engine/Source/Runtime/GameplayAbilities/Private/Abilities/GameplayAbilityTargetActor_SingleLineTrace.cpp
+++ b/Engine/Source/Runtime/GameplayAbilities/Private/Abilities/GameplayAbilityTargetActor_SingleLineTrace.cpp
@@ -89,6 +89,7 @@ void AGameplayAbilityTargetActor_SingleLineTrace::Confirm()
 
 void AGameplayAbilityTargetActor_SingleLineTrace::Cancel()
 {
+	CanceledDelegate.Broadcast(FGameplayAbilityTargetDataHandle());
 	Destroy();
 }
 
diff --git a/Engine/Source/Runtime/GameplayAbilities/Public/Abilities/Tasks/AbilityTask_SpawnActor.h b/Engine/Source/Runtime/GameplayAbilities/Public/Abilities/Tasks/AbilityTask_SpawnActor.h
--- a/Engine/Source/Runtime/GameplayAbilities/Public/Abilities/Tasks/AbilityTask_SpawnActor.h
+++ b/Engine/Source/Runtime/GameplayAbilities/Public/Abilities/Tasks/AbilityTask_SpawnActor.h
@@ -28,6 +28,7 @@ class UAbilityTask_SpawnActor: public UAbilityTask
 	UPROPERTY(BlueprintAssignable)
 	FSpawnActorDelegate	Success;
 	
+	/** Spawn new Actor on the network authority (server) */
 	UFUNCTION(BlueprintCallable, meta=(HidePin = "WorldContextObject", DefaultToSelf = "WorldContextObject", BlueprintInternalUseOnly = "true"), Category = "Abilities")
 	static UAbilityTask_SpawnActor* SpawnActor(UObject* WorldContextObject, TSubclassOf<AActor> Class);
 
diff --git a/Engine/Source/Runtime/GameplayAbilities/Public/Abilities/Tasks/AbilityTask_WaitAbilityActivate.h b/Engine/Source/Runtime/GameplayAbilities/Public/Abilities/Tasks/AbilityTask_WaitAbilityActivate.h
--- a/Engine/Source/Runtime/GameplayAbilities/Public/Abilities/Tasks/AbilityTask_WaitAbilityActivate.h
+++ b/Engine/Source/Runtime/GameplayAbilities/Public/Abilities/Tasks/AbilityTask_WaitAbilityActivate.h
@@ -25,7 +25,7 @@ class UAbilityTask_WaitAbilityActivate : public UAbilityTask
 	UFUNCTION()
 	void OnAbilityActivate(UGameplayAbility *ActivatedAbility);
 
-	/** Wait until an overlap occurs. This will need to be better fleshed out so we can specify game specific collision requirements */
+	/** Wait until a new ability (of the same or different type) is activated. */
 	UFUNCTION(BlueprintCallable, meta = (HidePin = "WorldContextObject", DefaultToSelf = "WorldContextObject", BlueprintInternalUseOnly = "TRUE"))
 	static UAbilityTask_WaitAbilityActivate* WaitForAbilityActivate(UObject* WorldContextObject, FGameplayTag WithTag, FGameplayTag WithoutTag);
 
diff --git a/Engine/Source/Runtime/GameplayAbilities/Public/Abilities/Tasks/AbilityTask_WaitAbilityCommit.h b/Engine/Source/Runtime/GameplayAbilities/Public/Abilities/Tasks/AbilityTask_WaitAbilityCommit.h
--- a/Engine/Source/Runtime/GameplayAbilities/Public/Abilities/Tasks/AbilityTask_WaitAbilityCommit.h
+++ b/Engine/Source/Runtime/GameplayAbilities/Public/Abilities/Tasks/AbilityTask_WaitAbilityCommit.h
@@ -27,8 +27,8 @@ class UAbilityTask_WaitAbilityCommit : public UAbilityTask
 	UFUNCTION()
 	void OnAbilityCommit(UGameplayAbility *ActivatedAbility);
 
-	/** Wait until an overlap occurs. This will need to be better fleshed out so we can specify game specific collision requirements */
-	UFUNCTION(BlueprintCallable, meta = (HidePin = "WorldContextObject", DefaultToSelf = "WorldContextObject", BlueprintInternalUseOnly = "TRUE"))
+	/** Wait until a new ability (of the same or different type) is commited. Used to gracefully interrupt abilities in specific ways. */
+	UFUNCTION(BlueprintCallable, meta = (HidePin = "WorldContextObject", DefaultToSelf = "WorldContextObject", BlueprintInternalUseOnly = "TRUE", FriendlyName = "Wait For New Ability Commit"))
 	static UAbilityTask_WaitAbilityCommit* WaitForAbilityCommit(UObject* WorldContextObject, FGameplayTag WithTag, FGameplayTag WithoutTage);	
 
 	FGameplayTag WithTag;
diff --git a/Engine/Source/Runtime/GameplayAbilities/Public/Abilities/Tasks/AbilityTask_WaitAttributeChange.h b/Engine/Source/Runtime/GameplayAbilities/Public/Abilities/Tasks/AbilityTask_WaitAttributeChange.h
--- a/Engine/Source/Runtime/GameplayAbilities/Public/Abilities/Tasks/AbilityTask_WaitAttributeChange.h
+++ b/Engine/Source/Runtime/GameplayAbilities/Public/Abilities/Tasks/AbilityTask_WaitAttributeChange.h
@@ -25,7 +25,7 @@ class UAbilityTask_WaitAttributeChange : public UAbilityTask
 		
 	void OnAttributeChange(float NewValue, const FGameplayEffectModCallbackData*);
 
-	/** Wait until an overlap occurs. This will need to be better fleshed out so we can specify game specific collision requirements */
+	/** Wait until an attribute changes. */
 	UFUNCTION(BlueprintCallable, meta = (HidePin = "WorldContextObject", DefaultToSelf = "WorldContextObject", BlueprintInternalUseOnly = "TRUE"))
 	static UAbilityTask_WaitAttributeChange* WaitForAttributeChange(UObject* WorldContextObject, FGameplayAttribute Attribute, FGameplayTag WithTag, FGameplayTag WithoutTag);
 
diff --git a/Engine/Source/Runtime/GameplayAbilities/Public/Abilities/Tasks/AbilityTask_WaitTargetData.h b/Engine/Source/Runtime/GameplayAbilities/Public/Abilities/Tasks/AbilityTask_WaitTargetData.h
--- a/Engine/Source/Runtime/GameplayAbilities/Public/Abilities/Tasks/AbilityTask_WaitTargetData.h
+++ b/Engine/Source/Runtime/GameplayAbilities/Public/Abilities/Tasks/AbilityTask_WaitTargetData.h
@@ -31,15 +31,7 @@ class UAbilityTask_WaitTargetData: public UAbilityTask
 	UFUNCTION()
 	void OnTargetDataCancelledCallback(FGameplayAbilityTargetDataHandle Data);
 
-	/**
-	 *	Requirements for using Begin/Finish SpawningActor functionality:
-	 *		-Have a parameters named 'Class' in your Proxy factor function (E.g., WaitTargetdata)
-	 *		-Have a function named BeginSpawningActor w/ the same Class parameter
-	 *			-This function should spawn the actor with SpawnActorDeferred and return true/false if it spawned something.
-	 *		-Have a function named FinishSpawningActor w/ an AActor* of the class you spawned
-	 *			-This function *must* call ExecuteConstruction + PostActorConstruction
-	 */
-	
+	/** Spawns Targeting actor and waits for it to return valid data or to be cancelled. */
 	UFUNCTION(BlueprintCallable, meta=(HidePin = "WorldContextObject", DefaultToSelf = "WorldContextObject", BlueprintInternalUseOnly = "true"), Category = "Abilities")
 	static UAbilityTask_WaitTargetData* WaitTargetData(UObject* WorldContextObject, TSubclassOf<AGameplayAbilityTargetActor> Class);
 
@@ -50,6 +42,7 @@ class UAbilityTask_WaitTargetData: public UAbilityTask
 	void FinishSpawningActor(UObject* WorldContextObject, AGameplayAbilityTargetActor* SpawnedActor);
 
 
+
 protected:
 
 	TSubclassOf<AGameplayAbilityTargetActor> TargetClass;
@@ -57,4 +50,14 @@ protected:
 	void Cleanup();
 
 	TWeakObjectPtr<AActor>	MySpawnedTargetActor;
-};
\ No newline at end of file
+};
+
+
+/**
+*	Requirements for using Begin/Finish SpawningActor functionality:
+*		-Have a parameters named 'Class' in your Proxy factor function (E.g., WaitTargetdata)
+*		-Have a function named BeginSpawningActor w/ the same Class parameter
+*			-This function should spawn the actor with SpawnActorDeferred and return true/false if it spawned something.
+*		-Have a function named FinishSpawningActor w/ an AActor* of the class you spawned
+*			-This function *must* call ExecuteConstruction + PostActorConstruction
+*/
\ No newline at end of file
