GameplayAbilities - update montage and movement mode tasks
-They wont show up in non GameplayAbilities blueprints
-Done the new way w/ out functions in AbilitySystemBlueprintLibrary

diff --git a/Engine/Source/Editor/GameplayAbilitiesEditor/Classes/K2Node_PlayMontageAndWait.h b/Engine/Source/Editor/GameplayAbilitiesEditor/Classes/K2Node_PlayMontageAndWait.h
deleted file mode 100644
diff --git a/Engine/Source/Editor/GameplayAbilitiesEditor/Classes/K2Node_WaitMovementModeChange.h b/Engine/Source/Editor/GameplayAbilitiesEditor/Classes/K2Node_WaitMovementModeChange.h
deleted file mode 100644
diff --git a/Engine/Source/Editor/GameplayAbilitiesEditor/Private/K2Node_PlayMontageAndWait.cpp b/Engine/Source/Editor/GameplayAbilitiesEditor/Private/K2Node_PlayMontageAndWait.cpp
deleted file mode 100644
diff --git a/Engine/Source/Editor/GameplayAbilitiesEditor/Private/K2Node_WaitMovementModeChange.cpp b/Engine/Source/Editor/GameplayAbilitiesEditor/Private/K2Node_WaitMovementModeChange.cpp
deleted file mode 100644
diff --git a/Engine/Source/Runtime/GameplayAbilities/Private/Abilities/Tasks/AbilityTask_PlayMontageAndWait.cpp b/Engine/Source/Runtime/GameplayAbilities/Private/Abilities/Tasks/AbilityTask_PlayMontageAndWait.cpp
--- a/Engine/Source/Runtime/GameplayAbilities/Private/Abilities/Tasks/AbilityTask_PlayMontageAndWait.cpp
+++ b/Engine/Source/Runtime/GameplayAbilities/Private/Abilities/Tasks/AbilityTask_PlayMontageAndWait.cpp
@@ -18,4 +18,32 @@ void UAbilityTask_PlayMontageAndWait::OnMontageEnded(UAnimMontage* Montage, bool
 	{
 		OnComplete.Broadcast();
 	}
+}
+
+UAbilityTask_PlayMontageAndWait* UAbilityTask_PlayMontageAndWait::CreatePlayMontageAndWaitProxy(UObject* WorldContextObject, UAnimMontage *MontageToPlay)
+{
+	check(WorldContextObject);
+
+	UAbilityTask_PlayMontageAndWait* MyObj = NewObject<UAbilityTask_PlayMontageAndWait>();
+	UGameplayAbility* ThisAbility = CastChecked<UGameplayAbility>(WorldContextObject);
+	MyObj->Ability = ThisAbility;
+	MyObj->MontageToPlay = MontageToPlay;
+
+	return MyObj;
+}
+
+void UAbilityTask_PlayMontageAndWait::Activate()
+{
+	if (Ability.IsValid())
+	{
+		AActor * ActorOwner = Cast<AActor>(Ability->GetOuter());
+				
+		TSharedPtr<FGameplayAbilityActorInfo> ActorInfo = TSharedPtr<FGameplayAbilityActorInfo>(UAbilitySystemGlobals::Get().AllocAbilityActorInfo(ActorOwner));
+		
+		FOnMontageEnded EndDelegate;
+		EndDelegate.BindUObject(this, &UAbilityTask_PlayMontageAndWait::OnMontageEnded);
+
+		ActorInfo->AnimInstance->Montage_Play(MontageToPlay, 1.f);
+		ActorInfo->AnimInstance->Montage_SetEndDelegate(EndDelegate);
+	}
 }
\ No newline at end of file
diff --git a/Engine/Source/Runtime/GameplayAbilities/Private/Abilities/Tasks/AbilityTask_WaitMovementModeChange.cpp b/Engine/Source/Runtime/GameplayAbilities/Private/Abilities/Tasks/AbilityTask_WaitMovementModeChange.cpp
--- a/Engine/Source/Runtime/GameplayAbilities/Private/Abilities/Tasks/AbilityTask_WaitMovementModeChange.cpp
+++ b/Engine/Source/Runtime/GameplayAbilities/Private/Abilities/Tasks/AbilityTask_WaitMovementModeChange.cpp
@@ -8,6 +8,38 @@ UAbilityTask_WaitMovementModeChange::UAbilityTask_WaitMovementModeChange(const c
 	RequiredMode = MOVE_None;
 }
 
+UAbilityTask_WaitMovementModeChange* UAbilityTask_WaitMovementModeChange::CreateWaitMovementModeChange(class UObject* WorldContextObject, EMovementMode NewMode)
+{
+	check(WorldContextObject);
+
+	UAbilityTask_WaitMovementModeChange* MyObj = NewObject<UAbilityTask_WaitMovementModeChange>();
+	UGameplayAbility* ThisAbility = CastChecked<UGameplayAbility>(WorldContextObject);
+	MyObj->Ability = ThisAbility;
+	MyObj->RequiredMode = NewMode;
+
+	return MyObj;
+}
+
+void UAbilityTask_WaitMovementModeChange::Activate()
+{
+	if (Ability.IsValid())
+	{
+		// Fixme: getting from the activating ability to the owning character is pretty awkward. We need
+		// a more standard way of doing this along with with a way for ability tasks to specify their requirements 
+		// (e.g., requires character / pawn, etc)
+
+		AActor * ActorOwner = Cast<AActor>(Ability->GetOuter());
+
+		TSharedPtr<FGameplayAbilityActorInfo> ActorInfo = TSharedPtr<FGameplayAbilityActorInfo>(UAbilitySystemGlobals::Get().AllocAbilityActorInfo(ActorOwner));
+
+		ACharacter * Character = Cast<ACharacter>(ActorInfo->Actor.Get());
+		if (Character)
+		{
+			Character->MovementModeChangedDelegate.AddDynamic(this, &UAbilityTask_WaitMovementModeChange::OnMovementModeChange);
+		}
+	}
+}
+
 void UAbilityTask_WaitMovementModeChange::OnMovementModeChange(ACharacter * Character, EMovementMode PrevMovementMode, uint8 PreviousCustomMode)
 {
 	if (Character)
diff --git a/Engine/Source/Runtime/GameplayAbilities/Private/AbilitySystemBlueprintLibrary.cpp b/Engine/Source/Runtime/GameplayAbilities/Private/AbilitySystemBlueprintLibrary.cpp
--- a/Engine/Source/Runtime/GameplayAbilities/Private/AbilitySystemBlueprintLibrary.cpp
+++ b/Engine/Source/Runtime/GameplayAbilities/Private/AbilitySystemBlueprintLibrary.cpp
@@ -61,61 +61,6 @@ FHitResult UAbilitySystemBlueprintLibrary::GetHitResultFromTargetData(FGameplayA
 
 // -------------------------------------------------------------------------------------
 
-UAbilityTask_PlayMontageAndWait* UAbilitySystemBlueprintLibrary::CreatePlayMontageAndWaitProxy(class UObject* WorldContextObject, UAnimMontage *MontageToPlay)
-{
-	check(WorldContextObject);
-	UGameplayAbility* Ability = CastChecked<UGameplayAbility>(WorldContextObject);
-	if (Ability)
-	{
-		AActor * ActorOwner = Cast<AActor>(Ability->GetOuter());
-
-		UAbilityTask_PlayMontageAndWait * MyObj = NULL;
-		MyObj = NewObject<UAbilityTask_PlayMontageAndWait>();
-
-		FOnMontageEnded EndDelegate;
-		EndDelegate.BindUObject(MyObj, &UAbilityTask_PlayMontageAndWait::OnMontageEnded);
-
-		FGameplayAbilityActorInfo	ActorInfo;
-		ActorInfo.InitFromActor(ActorOwner);
-
-		ActorInfo.AnimInstance->Montage_Play(MontageToPlay, 1.f);
-		ActorInfo.AnimInstance->Montage_SetEndDelegate(EndDelegate);
-
-		return MyObj;
-	}
-
-	return NULL;	
-}
-
-
-UAbilityTask_WaitMovementModeChange* UAbilitySystemBlueprintLibrary::CreateWaitMovementModeChange(class UObject* WorldContextObject, EMovementMode NewMode)
-{
-	check(WorldContextObject);
-	UGameplayAbility* Ability = CastChecked<UGameplayAbility>(WorldContextObject);
-	if (Ability)
-	{
-		AActor * ActorOwner = Cast<AActor>(Ability->GetOuter());
-
-		UAbilityTask_WaitMovementModeChange * MyObj = NULL;
-		MyObj = NewObject<UAbilityTask_WaitMovementModeChange>();
-		MyObj->RequiredMode = NewMode;
-
-		FGameplayAbilityActorInfo	ActorInfo;
-		ActorInfo.InitFromActor(ActorOwner);
-
-		ACharacter * Character = Cast<ACharacter>(ActorInfo.Actor.Get());
-		if (Character)
-		{
-			Character->MovementModeChangedDelegate.AddDynamic(MyObj, &UAbilityTask_WaitMovementModeChange::OnMovementModeChange);
-		}
-		return MyObj;
-	}
-
-	return NULL;
-}
-
-// -------------------------------------------------------------------------------------
-
 
 bool UAbilitySystemBlueprintLibrary::IsInstigatorLocallyControlled(FGameplayCueParameters Parameters)
 {
diff --git a/Engine/Source/Runtime/GameplayAbilities/Public/Abilities/Tasks/AbilityTask_PlayMontageAndWait.h b/Engine/Source/Runtime/GameplayAbilities/Public/Abilities/Tasks/AbilityTask_PlayMontageAndWait.h
--- a/Engine/Source/Runtime/GameplayAbilities/Public/Abilities/Tasks/AbilityTask_PlayMontageAndWait.h
+++ b/Engine/Source/Runtime/GameplayAbilities/Public/Abilities/Tasks/AbilityTask_PlayMontageAndWait.h
@@ -19,6 +19,14 @@ class UAbilityTask_PlayMontageAndWait : public UAbilityTask
 	UFUNCTION()
 	void OnMontageEnded(UAnimMontage* Montage, bool bInterrupted);
 
+
+	UFUNCTION(BlueprintCallable, meta = (HidePin = "WorldContextObject", DefaultToSelf = "WorldContextObject", BlueprintInternalUseOnly = "TRUE"))
+	static UAbilityTask_PlayMontageAndWait* CreatePlayMontageAndWaitProxy(UObject* WorldContextObject, UAnimMontage *MontageToPlay);
+
+	virtual void Activate() override;
+
 public:
+
+	UAnimMontage* MontageToPlay;
 	
 };
\ No newline at end of file
diff --git a/Engine/Source/Runtime/GameplayAbilities/Public/Abilities/Tasks/AbilityTask_WaitMovementModeChange.h b/Engine/Source/Runtime/GameplayAbilities/Public/Abilities/Tasks/AbilityTask_WaitMovementModeChange.h
--- a/Engine/Source/Runtime/GameplayAbilities/Public/Abilities/Tasks/AbilityTask_WaitMovementModeChange.h
+++ b/Engine/Source/Runtime/GameplayAbilities/Public/Abilities/Tasks/AbilityTask_WaitMovementModeChange.h
@@ -10,6 +10,8 @@ class UAbilityTask_WaitMovementModeChange : public UAbilityTask
 {
 	GENERATED_UCLASS_BODY()
 
+public:
+
 	UPROPERTY(BlueprintAssignable)
 	FMovementModeChangedDelegate	OnChange;
 
@@ -18,6 +20,8 @@ class UAbilityTask_WaitMovementModeChange : public UAbilityTask
 
 	EMovementMode	RequiredMode;
 
-public:
-	
+	UFUNCTION(BlueprintCallable, meta = (HidePin = "WorldContextObject", DefaultToSelf = "WorldContextObject", BlueprintInternalUseOnly = "TRUE"))
+	static UAbilityTask_WaitMovementModeChange* CreateWaitMovementModeChange(UObject* WorldContextObject, EMovementMode NewMode);
+
+	virtual void Activate() override;
 };
\ No newline at end of file
diff --git a/Engine/Source/Runtime/GameplayAbilities/Public/AbilitySystemBlueprintLibrary.h b/Engine/Source/Runtime/GameplayAbilities/Public/AbilitySystemBlueprintLibrary.h
--- a/Engine/Source/Runtime/GameplayAbilities/Public/AbilitySystemBlueprintLibrary.h
+++ b/Engine/Source/Runtime/GameplayAbilities/Public/AbilitySystemBlueprintLibrary.h
@@ -48,15 +48,4 @@ class UAbilitySystemBlueprintLibrary : public UBlueprintFunctionLibrary
 
 	// -------------------------------------------------------------------------------
 
-	UFUNCTION(BlueprintCallable, meta=(HidePin="WorldContextObject", DefaultToSelf="WorldContextObject", BlueprintInternalUseOnly = "TRUE"))
-	static UAbilityTask_PlayMontageAndWait* CreatePlayMontageAndWaitProxy(UObject* WorldContextObject, UAnimMontage *MontageToPlay);
-
-	UFUNCTION(BlueprintCallable, meta = (HidePin = "WorldContextObject", DefaultToSelf = "WorldContextObject", BlueprintInternalUseOnly = "TRUE"))
-	static UAbilityTask_WaitMovementModeChange* CreateWaitMovementModeChange(UObject* WorldContextObject, EMovementMode NewMode);
-
-	// -------------------------------------------------------------------------------
-
-
-
-
 };
\ No newline at end of file
