GameplayAbility
-Removed PredictiveActivate functions, replaced with AbilityTask/Delegate to callback when an ability is confirmed

diff --git a/Engine/Source/Runtime/GameplayAbilities/Private/Abilities/GameplayAbility.cpp b/Engine/Source/Runtime/GameplayAbilities/Private/Abilities/GameplayAbility.cpp
--- a/Engine/Source/Runtime/GameplayAbilities/Private/Abilities/GameplayAbility.cpp
+++ b/Engine/Source/Runtime/GameplayAbilities/Private/Abilities/GameplayAbility.cpp
@@ -29,13 +29,6 @@ UGameplayAbility::UGameplayAbility(const class FPostConstructInitializePropertie
 		UFunction* ActivateFunction = GetClass()->FindFunctionByName(FuncName);
 		HasBlueprintActivate = ActivateFunction && ActivateFunction->GetOuter()->IsA(UBlueprintGeneratedClass::StaticClass());
 	}
-	{
-		static FName FuncName = FName(TEXT("K2_PredictiveActivateAbility"));
-		UFunction* PredictiveActivateFunction = GetClass()->FindFunctionByName(FuncName);
-		HasBlueprintPredictiveActivate = PredictiveActivateFunction && PredictiveActivateFunction->GetOuter()->IsA(UBlueprintGeneratedClass::StaticClass());
-	}
-
-
 	
 #if WITH_EDITOR
 	/** Autoregister abilities with the blueprint debugger in the editor.*/
@@ -231,7 +224,7 @@ bool UGameplayAbility::TryActivateAbility(const FGameplayAbilityActorInfo* Actor
 			if(GetReplicationPolicy() == EGameplayAbilityReplicationPolicy::ReplicateNone)
 			{
 				UGameplayAbility * Ability = ActorInfo->AbilitySystemComponent->CreateNewInstanceOfAbility(this);
-				Ability->CallPredictiveActivateAbility(ActorInfo, ActivationInfo);
+				Ability->CallActivateAbility(ActorInfo, ActivationInfo);
 				if (OutInstancedAbility)
 				{
 					*OutInstancedAbility = Ability;
@@ -244,7 +237,7 @@ bool UGameplayAbility::TryActivateAbility(const FGameplayAbilityActorInfo* Actor
 		}
 		else
 		{
-			CallPredictiveActivateAbility(ActorInfo, ActivationInfo);
+			CallActivateAbility(ActorInfo, ActivationInfo);
 		}
 		
 		ServerTryActivateAbility(ActorInfo, ActivationInfo);
@@ -255,18 +248,9 @@ bool UGameplayAbility::TryActivateAbility(const FGameplayAbilityActorInfo* Actor
 
 void UGameplayAbility::EndAbility(const FGameplayAbilityActorInfo* ActorInfo)
 {
-	// MarkPending kill if we are instance per execution
-	if (InstancingPolicy == EGameplayAbilityInstancingPolicy::InstancedPerExecution)
-	{
-		// If not replicated, or if authority (clients should not delete subobjects on their own)
-		if (ReplicationPolicy == EGameplayAbilityReplicationPolicy::ReplicateNone || ActorInfo->IsNetAuthority())
-		{
-			MarkPendingKill();
-		}
-	}
-
-	// -Remove from owning AbilitySystemComponent
-	// -generic way of releasing all callbacks
+	ActorInfo->AbilitySystemComponent->NotifyAbilityEnded(this);
+	
+	// TODO: generic way of releasing all callbacks?
 }
 
 void UGameplayAbility::ActivateAbility(const FGameplayAbilityActorInfo* ActorInfo, const FGameplayAbilityActivationInfo ActivationInfo)
@@ -310,23 +294,12 @@ void UGameplayAbility::CallActivateAbility(const FGameplayAbilityActorInfo* Acto
 	ActivateAbility(ActorInfo, ActivationInfo);
 }
 
-void UGameplayAbility::CallPredictiveActivateAbility(const FGameplayAbilityActorInfo* ActorInfo, const FGameplayAbilityActivationInfo ActivationInfo)
+void UGameplayAbility::ConfirmActivateSucceed()
 {
-	PreActivate(ActorInfo, ActivationInfo);
-	PredictiveActivateAbility(ActorInfo, ActivationInfo);
-}
-
-void UGameplayAbility::PredictiveActivateAbility(const FGameplayAbilityActorInfo* ActorInfo, const FGameplayAbilityActivationInfo ActivationInfo)
-{
-	SetCurrentInfo(ActorInfo, ActivationInfo);
-	
-	check(ActivationInfo.ActivationMode == EGameplayAbilityActivationMode::Predicting);
+	CurrentActivationInfo.SetActivationConfirmed();
 
-	if (HasBlueprintPredictiveActivate)
-	{
-		K2_PredictiveActivateAbility();
-	}
-	
+	OnConfirmDelegate.Broadcast(this);
+	OnConfirmDelegate.Clear();
 }
 
 void UGameplayAbility::ServerTryActivateAbility(const FGameplayAbilityActorInfo* ActorInfo, const FGameplayAbilityActivationInfo ActivationInfo)
diff --git a/Engine/Source/Runtime/GameplayAbilities/Private/Abilities/GameplayAbility_CharacterJump.cpp b/Engine/Source/Runtime/GameplayAbilities/Private/Abilities/GameplayAbility_CharacterJump.cpp
--- a/Engine/Source/Runtime/GameplayAbilities/Private/Abilities/GameplayAbility_CharacterJump.cpp
+++ b/Engine/Source/Runtime/GameplayAbilities/Private/Abilities/GameplayAbility_CharacterJump.cpp
@@ -30,11 +30,6 @@ void UGameplayAbility_CharacterJump::ActivateAbility(const FGameplayAbilityActor
 	}
 }
 
-void UGameplayAbility_CharacterJump::PredictiveActivateAbility(const FGameplayAbilityActorInfo* ActorInfo, const FGameplayAbilityActivationInfo ActivationInfo)
-{
-	CallActivateAbility(ActorInfo, ActivationInfo);
-}
-
 void UGameplayAbility_CharacterJump::InputReleased(int32 InputID, const FGameplayAbilityActorInfo* ActorInfo)
 {
 	// This is the 'top' of the ability execution code path. We don't have an ActivationInfo yet.
diff --git a/Engine/Source/Runtime/GameplayAbilities/Private/AbilitySystemComponent_Abilities.cpp b/Engine/Source/Runtime/GameplayAbilities/Private/AbilitySystemComponent_Abilities.cpp
--- a/Engine/Source/Runtime/GameplayAbilities/Private/AbilitySystemComponent_Abilities.cpp
+++ b/Engine/Source/Runtime/GameplayAbilities/Private/AbilitySystemComponent_Abilities.cpp
@@ -79,6 +79,29 @@ UGameplayAbility * UAbilitySystemComponent::CreateNewInstanceOfAbility(UGameplay
 	return AbilityInstance;
 }
 
+void UAbilitySystemComponent::NotifyAbilityEnded(UGameplayAbility* Ability)
+{
+	check(Ability);
+	
+	/** If this is instanced per execution, mark pending kill and remove it from our instanced lists if we are the authority */
+	if (Ability->GetInstancingPolicy() == EGameplayAbilityInstancingPolicy::InstancedPerExecution)
+	{
+		if (Ability->GetReplicationPolicy() != EGameplayAbilityReplicationPolicy::ReplicateNone)
+		{
+			if (GetOwnerRole() == ROLE_Authority)
+			{
+				ReplicatedInstancedAbilities.Remove(Ability);
+				Ability->MarkPendingKill();
+			}
+		}
+		else
+		{
+			NonReplicatedInstancedAbilities.Remove(Ability);
+			Ability->MarkPendingKill();
+		}
+	}
+}
+
 bool UAbilitySystemComponent::ActivateAbility(TWeakObjectPtr<UGameplayAbility> Ability)
 {
 	check(AbilityActorInfo.IsValid());
@@ -223,11 +246,45 @@ void UAbilitySystemComponent::ClientActivateAbilitySucceed_Implementation(UGamep
 	AActor * OwnerActor = GetOwner();
 	check(OwnerActor);
 
-	// Fixme: We need a better way to link up/reconcile preditive replicated abilities.
-	if (AbilityToActivate->GetInstancingPolicy() != EGameplayAbilityInstancingPolicy::InstancedPerExecution)
+	// Fixme: We need a better way to link up/reconcile preditive replicated abilities. It would be ideal if we could predictively spawn an
+	// ability and then replace/link it with the server spawned one once the server has confirmed it.
+	// 
+
+	FGameplayAbilityActivationInfo	ActivationInfo(EGameplayAbilityActivationMode::Confirmed, PredictionKey);
+
+
+	if (AbilityToActivate->NetExecutionPolicy == EGameplayAbilityNetExecutionPolicy::Predictive)
 	{
-		FGameplayAbilityActivationInfo	ActivationInfo(EGameplayAbilityActivationMode::Confirmed, PredictionKey);
-		AbilityToActivate->CallActivateAbility(AbilityActorInfo.Get(), ActivationInfo);
+		// Find the one we predictively spawned, tell them we are confirmed
+		bool found = false;
+		for (UGameplayAbility* LocalAbility : NonReplicatedInstancedAbilities)				// Fixme: this has to be updated once predictive abilities can replicate
+		{
+			if (LocalAbility->GetCurrentActivationInfo().PredictionKey == PredictionKey)
+			{
+				LocalAbility->ConfirmActivateSucceed();
+				found = true;
+				break;
+			}
+		}
+
+		if (!found)
+		{
+			ABILITY_LOG(Warning, TEXT("Ability %s was confirmed by server but no longer exists on client (replication key: %d"), *AbilityToActivate->GetName(), PredictionKey);
+		}
+	}
+	else
+	{
+		// We haven't already executed this ability at all, so kick it off.
+		if (AbilityToActivate->GetInstancingPolicy() == EGameplayAbilityInstancingPolicy::InstancedPerExecution)
+		{
+			// Need to instantiate this in order to execute
+			UGameplayAbility* InstancedAbility = CreateNewInstanceOfAbility(AbilityToActivate);
+			InstancedAbility->CallActivateAbility(AbilityActorInfo.Get(), ActivationInfo);
+		}
+		else
+		{
+			AbilityToActivate->CallActivateAbility(AbilityActorInfo.Get(), ActivationInfo);
+		}
 	}
 }
 
diff --git a/Engine/Source/Runtime/GameplayAbilities/Public/Abilities/GameplayAbility.h b/Engine/Source/Runtime/GameplayAbilities/Public/Abilities/GameplayAbility.h
--- a/Engine/Source/Runtime/GameplayAbilities/Public/Abilities/GameplayAbility.h
+++ b/Engine/Source/Runtime/GameplayAbilities/Public/Abilities/GameplayAbility.h
@@ -179,6 +179,8 @@ public:
 	UPROPERTY(EditDefaultsOnly, Category = Tags)
 	FGameplayTagContainer AbilityTags;
 
+	FGenericAbilityDelegate	OnConfirmDelegate;
+
 protected:
 		
 	// --------------------------------------
@@ -209,6 +211,9 @@ protected:
 
 	void CallActivateAbility(const FGameplayAbilityActorInfo* ActorInfo, const FGameplayAbilityActivationInfo ActivationInfo);
 
+	/** Called on a predictive ability when the server confirms its execution */
+	void ConfirmActivateSucceed();
+
 	// --------------------------------------
 	//	CommitAbility
 	// --------------------------------------
@@ -248,19 +253,6 @@ protected:
 	/** Destroys intanced-per-execution abilities. Instance-per-actor abilities should 'reset'. Non instance abilities - what can we do? */
 	virtual void CancelAbility(const FGameplayAbilityActorInfo* ActorInfo, const FGameplayAbilityActivationInfo ActivationInfo) { }
 
-	// -------------------------------------
-	//	PredictiveActivate
-	// -------------------------------------
-
-	UFUNCTION(BlueprintImplementableEvent, Category = Ability, FriendlyName="PredictiveActivateAbility")
-	virtual void K2_PredictiveActivateAbility();
-
-	virtual void PredictiveActivateAbility(const FGameplayAbilityActorInfo* ActorInfo, const FGameplayAbilityActivationInfo ActivationInfo);
-
-	void CallPredictiveActivateAbility(const FGameplayAbilityActorInfo* ActorInfo, const FGameplayAbilityActivationInfo ActivationInfo);
-
-	bool HasBlueprintPredictiveActivate;
-
 	// -------------------------------------
 	//	EndAbility
 	// -------------------------------------
diff --git a/Engine/Source/Runtime/GameplayAbilities/Public/Abilities/GameplayAbility_CharacterJump.h b/Engine/Source/Runtime/GameplayAbilities/Public/Abilities/GameplayAbility_CharacterJump.h
--- a/Engine/Source/Runtime/GameplayAbilities/Public/Abilities/GameplayAbility_CharacterJump.h
+++ b/Engine/Source/Runtime/GameplayAbilities/Public/Abilities/GameplayAbility_CharacterJump.h
@@ -21,7 +21,5 @@ public:
 
 	virtual void InputReleased(int32 InputID, const FGameplayAbilityActorInfo* ActorInfo) override;
 
-	virtual void PredictiveActivateAbility(const FGameplayAbilityActorInfo* ActorInfo, const FGameplayAbilityActivationInfo ActivationInfo) override;
-
 	virtual void CancelAbility(const FGameplayAbilityActorInfo* ActorInfo, const FGameplayAbilityActivationInfo ActivationInfo) override;
 };
\ No newline at end of file
diff --git a/Engine/Source/Runtime/GameplayAbilities/Public/Abilities/Tasks/AbilityTask.h b/Engine/Source/Runtime/GameplayAbilities/Public/Abilities/Tasks/AbilityTask.h
--- a/Engine/Source/Runtime/GameplayAbilities/Public/Abilities/Tasks/AbilityTask.h
+++ b/Engine/Source/Runtime/GameplayAbilities/Public/Abilities/Tasks/AbilityTask.h
@@ -55,6 +55,8 @@ class GAMEPLAYABILITIES_API UAbilityTask : public UObject
 {
 	GENERATED_UCLASS_BODY()
 
+	DECLARE_DYNAMIC_MULTICAST_DELEGATE(FGenericAbilityTaskDelegate);
+
 	// Called to trigger the actual task once the delegates have been set up
 	UFUNCTION(BlueprintCallable, meta = (BlueprintInternalUseOnly = "true"), Category = "Abilities")
 	virtual void Activate();
diff --git a/Engine/Source/Runtime/GameplayAbilities/Public/AbilitySystemComponent.h b/Engine/Source/Runtime/GameplayAbilities/Public/AbilitySystemComponent.h
--- a/Engine/Source/Runtime/GameplayAbilities/Public/AbilitySystemComponent.h
+++ b/Engine/Source/Runtime/GameplayAbilities/Public/AbilitySystemComponent.h
@@ -309,6 +309,8 @@ class GAMEPLAYABILITIES_API UAbilitySystemComponent : public UActorComponent
 	UPROPERTY(Replicated)
 	TArray<UGameplayAbility*>	ReplicatedInstancedAbilities;
 
+	void NotifyAbilityEnded(UGameplayAbility* Ability);
+
 	/**
 	 *	The abilities we can activate. 
 	 *		-This will include CDOs for non instanced abilities and per-execution instanced abilities. 
@@ -436,20 +438,12 @@ class GAMEPLAYABILITIES_API UAbilitySystemComponent : public UActorComponent
 
 private:
 
-	// --------------------------------------------
-	// Important internal functions
-	// --------------------------------------------
-
 	void ExecutePeriodicEffect(FActiveGameplayEffectHandle	Handle);
 
 	void ExecuteGameplayEffect(const FGameplayEffectSpec &Spec, const FModifierQualifier &QualifierContext);
 
 	void CheckDurationExpired(FActiveGameplayEffectHandle Handle);
 
-	// --------------------------------------------
-	// Internal Utility / helper functions
-	// --------------------------------------------
-
 	bool AreGameplayEffectApplicationRequirementsSatisfied(const class UGameplayEffect* EffectToAdd, FGameplayEffectInstigatorContext& Instigator) const;
 
 	bool IsOwnerActorAuthoritative() const;
