const-ify UAttributeSet* getters
-We very rarely want C++ to set GameplayAttribute directlyl; should always go through GameplayEffects

diff --git a/Engine/Source/Runtime/GameplayAbilities/Private/AbilitySystemComponent.cpp b/Engine/Source/Runtime/GameplayAbilities/Private/AbilitySystemComponent.cpp
--- a/Engine/Source/Runtime/GameplayAbilities/Private/AbilitySystemComponent.cpp
+++ b/Engine/Source/Runtime/GameplayAbilities/Private/AbilitySystemComponent.cpp
@@ -52,37 +52,31 @@ UAbilitySystemComponent::~UAbilitySystemComponent()
 	ActiveGameplayEffects.PreDestroy();
 }
 
-UAttributeSet* UAbilitySystemComponent::InitStats(TSubclassOf<class UAttributeSet> Attributes, const UDataTable* DataTable)
+const UAttributeSet* UAbilitySystemComponent::InitStats(TSubclassOf<class UAttributeSet> Attributes, const UDataTable* DataTable)
 {
-	UAttributeSet* AttributeObj = NULL;
+	const UAttributeSet* AttributeObj = NULL;
 	if (Attributes)
 	{
 		AttributeObj = GetOrCreateAttributeSubobject(Attributes);
 		if (AttributeObj && DataTable)
 		{
-			AttributeObj->InitFromMetaDataTable(DataTable);
+			// This const_cast is OK - this is one of the few places we want to directly modify our AttributeSet properties rather
+			// than go through a gameplay effect
+			const_cast<UAttributeSet*>(AttributeObj)->InitFromMetaDataTable(DataTable);
 		}
 	}
 	return AttributeObj;
 }
 
-void UAbilitySystemComponent::ModifyStats(TSubclassOf<class UAttributeSet> AttributeClass, FDataTableRowHandle ModifierHandle)
+void UAbilitySystemComponent::K2_InitStats(TSubclassOf<class UAttributeSet> Attributes, const UDataTable* DataTable)
 {
-	UAttributeSet *Attributes = Cast<UAttributeSet>(GetAttributeSubobject(AttributeClass));
-	if (Attributes)
-	{
-		FAttributeModifierTest *ModifierData = ModifierHandle.GetRow<FAttributeModifierTest>();
-		if (ModifierData)
-		{
-			ModifierData->ApplyModifier(Attributes);
-		}
-	}
+	InitStats(Attributes, DataTable);
 }
 
-UAttributeSet*	UAbilitySystemComponent::GetOrCreateAttributeSubobject(const TSubclassOf<UAttributeSet> AttributeClass)
+const UAttributeSet*	UAbilitySystemComponent::GetOrCreateAttributeSubobject(const TSubclassOf<UAttributeSet> AttributeClass)
 {
 	AActor *OwningActor = GetOwner();
-	UAttributeSet *MyAttributes  = NULL;
+	const UAttributeSet *MyAttributes  = NULL;
 	if (OwningActor && AttributeClass)
 	{
 		MyAttributes = GetAttributeSubobject(AttributeClass);
@@ -96,16 +90,16 @@ UAttributeSet*	UAbilitySystemComponent::GetOrCreateAttributeSubobject(const TSub
 	return MyAttributes;
 }
 
-UAttributeSet* UAbilitySystemComponent::GetAttributeSubobjectChecked(const TSubclassOf<UAttributeSet> AttributeClass) const
+const UAttributeSet* UAbilitySystemComponent::GetAttributeSubobjectChecked(const TSubclassOf<UAttributeSet> AttributeClass) const
 {
-	UAttributeSet *Set = GetAttributeSubobject(AttributeClass);
+	const UAttributeSet *Set = GetAttributeSubobject(AttributeClass);
 	check(Set);
 	return Set;
 }
 
-UAttributeSet* UAbilitySystemComponent::GetAttributeSubobject(const TSubclassOf<UAttributeSet> AttributeClass) const
+const UAttributeSet* UAbilitySystemComponent::GetAttributeSubobject(const TSubclassOf<UAttributeSet> AttributeClass) const
 {
-	for (UAttributeSet * Set : SpawnedAttributes)
+	for (const UAttributeSet* Set : SpawnedAttributes)
 	{
 		if (Set && Set->IsA(AttributeClass))
 		{
@@ -124,7 +118,7 @@ void UAbilitySystemComponent::OnRegister()
 	{
 		if (DefaultStartingData[i].Attributes && DefaultStartingData[i].DefaultStartingTable)
 		{
-			UAttributeSet* Attributes = GetOrCreateAttributeSubobject(DefaultStartingData[i].Attributes);
+			UAttributeSet* Attributes = const_cast<UAttributeSet*>(GetOrCreateAttributeSubobject(DefaultStartingData[i].Attributes));
 			Attributes->InitFromMetaDataTable(DefaultStartingData[i].DefaultStartingTable);
 		}
 	}
@@ -163,13 +157,13 @@ bool UAbilitySystemComponent::IsOwnerActorAuthoritative() const
 
 void UAbilitySystemComponent::SetNumericAttribute(const FGameplayAttribute &Attribute, float NewFloatValue)
 {
-	UAttributeSet * AttributeSet = GetAttributeSubobjectChecked(Attribute.GetAttributeSetClass());
-	Attribute.SetNumericValueChecked(NewFloatValue, AttributeSet);
+	const UAttributeSet* AttributeSet = GetAttributeSubobjectChecked(Attribute.GetAttributeSetClass());
+	Attribute.SetNumericValueChecked(NewFloatValue, const_cast<UAttributeSet*>(AttributeSet));
 }
 
 float UAbilitySystemComponent::GetNumericAttribute(const FGameplayAttribute &Attribute)
 {
-	UAttributeSet * AttributeSet = GetAttributeSubobjectChecked(Attribute.GetAttributeSetClass());
+	const UAttributeSet* AttributeSet = GetAttributeSubobjectChecked(Attribute.GetAttributeSetClass());
 	return Attribute.GetNumericValueChecked(AttributeSet);
 }
 
@@ -512,11 +506,11 @@ bool UAbilitySystemComponent::ReplicateSubobjects(class UActorChannel *Channel,
 {
 	bool WroteSomething = Super::ReplicateSubobjects(Channel, Bunch, RepFlags);
 
-	for (UAttributeSet* Set : SpawnedAttributes)
+	for (const UAttributeSet* Set : SpawnedAttributes)
 	{
 		if (Set)
 		{
-			WroteSomething |= Channel->ReplicateSubobject(Set, *Bunch, *RepFlags);
+			WroteSomething |= Channel->ReplicateSubobject(const_cast<UAttributeSet*>(Set), *Bunch, *RepFlags);
 		}
 	}
 
@@ -533,11 +527,11 @@ bool UAbilitySystemComponent::ReplicateSubobjects(class UActorChannel *Channel,
 
 void UAbilitySystemComponent::GetSubobjectsWithStableNamesForNetworking(TArray<UObject*>& Objs)
 {
-	for (UAttributeSet* Set : SpawnedAttributes)
+	for (const UAttributeSet* Set : SpawnedAttributes)
 	{
 		if (Set && Set->IsNameStableForNetworking())
 		{
-			Objs.Add(Set);
+			Objs.Add(const_cast<UAttributeSet*>(Set));
 		}
 	}
 }
diff --git a/Engine/Source/Runtime/GameplayAbilities/Private/AbilitySystemDebugHUD.cpp b/Engine/Source/Runtime/GameplayAbilities/Private/AbilitySystemDebugHUD.cpp
--- a/Engine/Source/Runtime/GameplayAbilities/Private/AbilitySystemDebugHUD.cpp
+++ b/Engine/Source/Runtime/GameplayAbilities/Private/AbilitySystemDebugHUD.cpp
@@ -91,7 +91,7 @@ void AAbilitySystemDebugHUD::DrawDebugAbilitySystemComponent(UAbilitySystemCompo
 	DrawWithBackground(Font, String, Color, EAlignHorizontal::Left, X, EAlignVertical::Top, Y);
 
 
-	for (UAttributeSet * Set : Component->SpawnedAttributes)
+	for (const UAttributeSet * Set : Component->SpawnedAttributes)
 	{
 		if (!Set)
 			continue;
@@ -108,7 +108,7 @@ void AAbilitySystemDebugHUD::DrawDebugAbilitySystemComponent(UAbilitySystemCompo
 			UProperty *Prop = *PropertyIt;			
 
 			FString ValueString;
-			void *PropertyValue = Prop->ContainerPtrToValuePtr<void>(Set);
+			const void *PropertyValue = Prop->ContainerPtrToValuePtr<void>(Set);
 			Prop->ExportTextItem(ValueString, PropertyValue, NULL, NULL, 0);
 
 			String = FString::Printf(TEXT("%s: %s"), *Prop->GetName(), *ValueString);
diff --git a/Engine/Source/Runtime/GameplayAbilities/Private/AbilitySystemTestAttributeSet.cpp b/Engine/Source/Runtime/GameplayAbilities/Private/AbilitySystemTestAttributeSet.cpp
--- a/Engine/Source/Runtime/GameplayAbilities/Private/AbilitySystemTestAttributeSet.cpp
+++ b/Engine/Source/Runtime/GameplayAbilities/Private/AbilitySystemTestAttributeSet.cpp
@@ -56,7 +56,7 @@ void UAbilitySystemTestAttributeSet::PreAttributeModify(struct FGameplayEffectMo
 		if (Data.EvaluatedData.Magnitude > 0.f)
 		{
 			// Check the source - does he have Crit?
-			UAbilitySystemTestAttributeSet * SourceAttributes = Data.EffectSpec.InstigatorContext.GetOriginalInstigatorAbilitySystemComponent()->GetSet<UAbilitySystemTestAttributeSet>();
+			const UAbilitySystemTestAttributeSet * SourceAttributes = Data.EffectSpec.InstigatorContext.GetOriginalInstigatorAbilitySystemComponent()->GetSet<UAbilitySystemTestAttributeSet>();
 			if (SourceAttributes && SourceAttributes->CritChance > 0.f)
 			{
 				if (FMath::FRand() <= SourceAttributes->CritChance)
diff --git a/Engine/Source/Runtime/GameplayAbilities/Private/AttributeSet.cpp b/Engine/Source/Runtime/GameplayAbilities/Private/AttributeSet.cpp
--- a/Engine/Source/Runtime/GameplayAbilities/Private/AttributeSet.cpp
+++ b/Engine/Source/Runtime/GameplayAbilities/Private/AttributeSet.cpp
@@ -13,10 +13,10 @@ void FGameplayAttribute::SetNumericValueChecked(const float NewValue, class UAtt
 	NumericProperty->SetFloatingPointPropertyValue(ValuePtr, NewValue);
 }
 
-float FGameplayAttribute::GetNumericValueChecked(class UAttributeSet* Src) const
+float FGameplayAttribute::GetNumericValueChecked(const UAttributeSet* Src) const
 {
-	UNumericProperty *NumericProperty = CastChecked<UNumericProperty>(Attribute);
-	void * ValuePtr = NumericProperty->ContainerPtrToValuePtr<void>(Src);
+	UNumericProperty* NumericProperty = CastChecked<UNumericProperty>(Attribute);
+	const void* ValuePtr = NumericProperty->ContainerPtrToValuePtr<void>(Src);
 	return NumericProperty->GetFloatingPointPropertyValue(ValuePtr);
 }
 
@@ -32,7 +32,7 @@ void UAttributeSet::InitFromMetaDataTable(const UDataTable* DataTable)
 
 	for( TFieldIterator<UProperty> It(GetClass(), EFieldIteratorFlags::IncludeSuper) ; It ; ++It )
 	{
-		UProperty *Property = *It;
+		UProperty* Property = *It;
 		UNumericProperty *NumericProperty = Cast<UNumericProperty>(Property);
 		if (NumericProperty)
 		{
diff --git a/Engine/Source/Runtime/GameplayAbilities/Private/GameplayEffect.cpp b/Engine/Source/Runtime/GameplayAbilities/Private/GameplayEffect.cpp
--- a/Engine/Source/Runtime/GameplayAbilities/Private/GameplayEffect.cpp
+++ b/Engine/Source/Runtime/GameplayAbilities/Private/GameplayEffect.cpp
@@ -1284,7 +1284,7 @@ void FActiveGameplayEffectsContainer::ExecuteActiveEffectsFrom(const FGameplayEf
 	{
 		if (Mod.Info.ModifierType == EGameplayMod::Attribute)
 		{
-			UAttributeSet * AttributeSet = Owner->GetAttributeSubobject(Mod.Info.Attribute.GetAttributeSetClass());
+			UAttributeSet* AttributeSet = const_cast<UAttributeSet*>(Owner->GetAttributeSubobject(Mod.Info.Attribute.GetAttributeSetClass()));
 			if (AttributeSet == NULL)
 			{
 				// Our owner doesn't have this attribute, so we can't do anything
@@ -1982,7 +1982,7 @@ bool FActiveGameplayEffectsContainer::CanApplyAttributeModifiers(const UGameplay
 		// It only makes sense to check additive operators
 		if (Mod.Info.ModifierOp == EGameplayModOp::Additive)
 		{
-			UAttributeSet * Set = Owner->GetAttributeSubobject(Mod.Info.Attribute.GetAttributeSetClass());
+			const UAttributeSet* Set = Owner->GetAttributeSubobject(Mod.Info.Attribute.GetAttributeSetClass());
 			float CurrentValue = Mod.Info.Attribute.GetNumericValueChecked(Set);
 			float CostValue = Mod.Aggregator.Get()->Evaluate().Magnitude;
 
diff --git a/Engine/Source/Runtime/GameplayAbilities/Public/AbilitySystemComponent.h b/Engine/Source/Runtime/GameplayAbilities/Public/AbilitySystemComponent.h
--- a/Engine/Source/Runtime/GameplayAbilities/Public/AbilitySystemComponent.h
+++ b/Engine/Source/Runtime/GameplayAbilities/Public/AbilitySystemComponent.h
@@ -67,44 +67,44 @@ class GAMEPLAYABILITIES_API UAbilitySystemComponent : public UActorComponent
 
 	/** Finds existing AttributeSet */
 	template <class T >
-	T*	GetSet()	// FIXME: This should be const by default
+	const T*	GetSet()
 	{
 		return (T*)GetAttributeSubobject(T::StaticClass());
 	}
 
 	template <class T >
-	T*	GetSetChecked()
+	const T*	GetSetChecked()
 	{
 		return (T*)GetAttributeSubobjectChecked(T::StaticClass());
 	}
 
 	/** Adds a new AttributeSet (initialized to default values) */
 	template <class T >
-	T*  AddSet()
+	const T*  AddSet()
 	{
 		return (T*)GetOrCreateAttributeSubobject(T::StaticClass());
 	}
 
 	/** Adds a new AttributeSet that is a DSO (created by called in their CStor) */
 	template <class T>
-	T*	AddDefaultSubobjectSet(TSubobjectPtr<T> Subobject)
+	const T*	AddDefaultSubobjectSet(TSubobjectPtr<T> Subobject)
 	{
-		T* Set = Subobject.Get();
+		const T* Set = Subobject.Get();
 		SpawnedAttributes.Add(Set);
 		return Set;
 	}
 
-	UFUNCTION(BlueprintCallable, Category="Skills")
-	UAttributeSet * InitStats(TSubclassOf<class UAttributeSet> Attributes, const UDataTable* DataTable);
+	const UAttributeSet* InitStats(TSubclassOf<class UAttributeSet> Attributes, const UDataTable* DataTable);
 
-	UFUNCTION(BlueprintCallable, Category="Skills")
-	void ModifyStats(TSubclassOf<class UAttributeSet> Attributes, FDataTableRowHandle ModifierHandle);
+	UFUNCTION(BlueprintCallable, Category="Skills", meta=(FriendlyName="InitStats"))
+	void K2_InitStats(TSubclassOf<class UAttributeSet> Attributes, const UDataTable* DataTable);
+		
 
 	UPROPERTY(EditAnywhere, Category="AttributeTest")
 	TArray<FAttributeDefaults>	DefaultStartingData;
 
 	UPROPERTY(Replicated)
-	TArray<UAttributeSet*>	SpawnedAttributes;
+	TArray<const UAttributeSet*>	SpawnedAttributes;
 
 	void SetNumericAttribute(const FGameplayAttribute &Attribute, float NewFloatValue);
 	float GetNumericAttribute(const FGameplayAttribute &Attribute);
@@ -482,9 +482,9 @@ protected:
 
 	virtual void OnRegister() override;
 
-	UAttributeSet*	GetAttributeSubobject(const TSubclassOf<UAttributeSet> AttributeClass) const;
-	UAttributeSet*	GetAttributeSubobjectChecked(const TSubclassOf<UAttributeSet> AttributeClass) const;
-	UAttributeSet*	GetOrCreateAttributeSubobject(const TSubclassOf<UAttributeSet> AttributeClass);
+	const UAttributeSet*	GetAttributeSubobject(const TSubclassOf<UAttributeSet> AttributeClass) const;
+	const UAttributeSet*	GetAttributeSubobjectChecked(const TSubclassOf<UAttributeSet> AttributeClass) const;
+	const UAttributeSet*	GetOrCreateAttributeSubobject(const TSubclassOf<UAttributeSet> AttributeClass);
 
 	friend struct FActiveGameplayEffectsContainer;
 	friend struct FActiveGameplayEffect;
diff --git a/Engine/Source/Runtime/GameplayAbilities/Public/AbilitySystemTestAttributeSet.h b/Engine/Source/Runtime/GameplayAbilities/Public/AbilitySystemTestAttributeSet.h
--- a/Engine/Source/Runtime/GameplayAbilities/Public/AbilitySystemTestAttributeSet.h
+++ b/Engine/Source/Runtime/GameplayAbilities/Public/AbilitySystemTestAttributeSet.h
@@ -9,56 +9,65 @@ class GAMEPLAYABILITIES_API UAbilitySystemTestAttributeSet : public UAttributeSe
 {
 	GENERATED_UCLASS_BODY()
 
+
+
+	/**
+	 *	NOTE ON MUTABLE:
+	 *	This is only done so that UAbilitySystemTestAttributeSet can be initialized directly in GameplayEffectsTest.cpp without doing a const_cast in 100+ places.
+	 *	Mutable is not required and should never be used on normal attribute sets.
+	 */
+
+
 	UPROPERTY(EditAnywhere, BlueprintReadWrite, Replicated, Category = "AttributeTest", meta = (HideFromModifiers))			// You can't make a GameplayEffect modify Health Directly (go through)
-	float	MaxHealth;
+	mutable float	MaxHealth;
 
 	UPROPERTY(EditAnywhere, BlueprintReadWrite, Replicated, Category = "AttributeTest", meta = (HideFromModifiers))			// You can't make a GameplayEffect modify Health Directly (go through)
-	float	Health;
+	mutable float	Health;
 
 	UPROPERTY(EditAnywhere, BlueprintReadWrite, Replicated, Category = "AttributeTest")
-	float	Mana;
+	mutable float	Mana;
 
 	UPROPERTY(EditAnywhere, BlueprintReadWrite, Replicated, Category = "AttributeTest")
-	float	MaxMana;
+	mutable float	MaxMana;
 
 	/** This Damage is just used for applying negative health mods. Its not a 'persistent' attribute. */
 	UPROPERTY(EditAnywhere, BlueprintReadWrite, Category = "AttributeTest", meta = (HideFromLevelInfos))		// You can't make a GameplayEffect 'powered' by Damage (Its transient)
-	float	Damage;
+	mutable float	Damage;
 
 	/** This Attribute is the actual spell damage for this actor. It will power spell based GameplayEffects */
 	UPROPERTY(EditAnywhere, BlueprintReadWrite, Replicated, Category = "AttributeTest")
-	float	SpellDamage;
+	mutable float	SpellDamage;
 
 	/** This Attribute is the actual physical damage for this actor. It will power physical based GameplayEffects */
 	UPROPERTY(EditAnywhere, BlueprintReadWrite, Replicated, Category = "AttributeTest")
-	float	PhysicalDamage;
+	mutable float	PhysicalDamage;
 
 	UPROPERTY(EditAnywhere, BlueprintReadWrite, Replicated, Category = "AttributeTest")
-	float	CritChance;
+	mutable float	CritChance;
 
 	UPROPERTY(EditAnywhere, BlueprintReadWrite, Replicated, Category = "AttributeTest")
-	float	CritMultiplier;
+	mutable float	CritMultiplier;
 
 	UPROPERTY(EditAnywhere, BlueprintReadWrite, Replicated, Category = "AttributeTest")
-	float	ArmorDamageReduction;
+	mutable float	ArmorDamageReduction;
 
 	UPROPERTY(EditAnywhere, BlueprintReadWrite, Replicated, Category = "AttributeTest")
-	float	DodgeChance;
+	mutable float	DodgeChance;
 
 	UPROPERTY(EditAnywhere, BlueprintReadWrite, Replicated, Category = "AttributeTest")
-	float	LifeSteal;
+	mutable float	LifeSteal;
 
 	UPROPERTY(EditAnywhere, BlueprintReadWrite, Replicated, Category = "AttributeTest")
-	float	Strength;
+	mutable float	Strength;
 
 	UPROPERTY(EditAnywhere, BlueprintReadWrite, Category = "AttributeTest")
-	float	StackingAttribute1;
+	mutable float	StackingAttribute1;
 
 	UPROPERTY(EditAnywhere, BlueprintReadWrite, Category = "AttributeTest")
-	float	StackingAttribute2;
+	mutable float	StackingAttribute2;
 
 	UPROPERTY(EditAnywhere, BlueprintReadWrite, Category = "AttributeTest")
-	float	NoStackAttribute;
+	mutable float	NoStackAttribute;
 
 	virtual void PreAttributeModify(struct FGameplayEffectModCallbackData &Data) override;
 	virtual void PostAttributeModify(const struct FGameplayEffectModCallbackData &Data) override;
diff --git a/Engine/Source/Runtime/GameplayAbilities/Public/AttributeSet.h b/Engine/Source/Runtime/GameplayAbilities/Public/AttributeSet.h
--- a/Engine/Source/Runtime/GameplayAbilities/Public/AttributeSet.h
+++ b/Engine/Source/Runtime/GameplayAbilities/Public/AttributeSet.h
@@ -37,7 +37,7 @@ struct GAMEPLAYABILITIES_API FGameplayAttribute
 
 	void SetNumericValueChecked(const float NewValue, class UAttributeSet* Dest) const;
 
-	float GetNumericValueChecked(class UAttributeSet* Src) const;
+	float GetNumericValueChecked(const UAttributeSet* Src) const;
 	
 	/** Equality/Inequality operators */
 	bool operator==(const FGameplayAttribute& Other) const;
