Few more updates

diff --git a/Engine/Source/Runtime/GameplayAbilities/Public/Abilities/GameplayAbility.h b/Engine/Source/Runtime/GameplayAbilities/Public/Abilities/GameplayAbility.h
--- a/Engine/Source/Runtime/GameplayAbilities/Public/Abilities/GameplayAbility.h
+++ b/Engine/Source/Runtime/GameplayAbilities/Public/Abilities/GameplayAbility.h
@@ -66,6 +66,8 @@
 /***
  
 
+ TODO: Update: this is not all accurate anymore now that PredictiveACtivateAbility is gone. Please update comments!
+
 
 This is the basic flow of function calls in a networked environment:
 -Client sees if he thinks he can activate the ability.
