AbilityTask_WaitConfirm

diff --git a/Engine/Source/Runtime/GameplayAbilities/Private/Abilities/Tasks/AbilityTask_WaitConfirm.cpp b/Engine/Source/Runtime/GameplayAbilities/Private/Abilities/Tasks/AbilityTask_WaitConfirm.cpp
new file mode 100644
--- /dev/null
+++ b/Engine/Source/Runtime/GameplayAbilities/Private/Abilities/Tasks/AbilityTask_WaitConfirm.cpp
@@ -0,0 +1,45 @@
+
+#include "AbilitySystemPrivatePCH.h"
+#include "Abilities/Tasks/AbilityTask_WaitConfirm.h"
+
+UAbilityTask_WaitConfirm::UAbilityTask_WaitConfirm(const class FPostConstructInitializeProperties& PCIP)
+	: Super(PCIP)
+{
+
+}
+
+void UAbilityTask_WaitConfirm::OnConfirmCallback(UGameplayAbility* InAbility)
+{
+	OnConfirm.Broadcast();
+
+	// We are done. Kill us so we don't keep getting broadcast messages
+	MarkPendingKill();
+}
+
+UAbilityTask_WaitConfirm* UAbilityTask_WaitConfirm::WaitConfirm(class UObject* WorldContextObject)
+{
+	check(WorldContextObject);
+
+	UAbilityTask_WaitConfirm * MyObj = NewObject<UAbilityTask_WaitConfirm>();
+	UGameplayAbility* ThisAbility = CastChecked<UGameplayAbility>(WorldContextObject);
+	MyObj->InitTask(ThisAbility);	
+
+	return MyObj;
+}
+
+void UAbilityTask_WaitConfirm::Activate()
+{
+	if (Ability.IsValid())
+	{
+		if (Ability->GetCurrentActivationInfo().ActivationMode == EGameplayAbilityActivationMode::Predicting)
+		{
+			// Register delegate so that when we are confirmed by the server, we call OnConfirmCallback
+			Ability->OnConfirmDelegate.AddUObject(this, &UAbilityTask_WaitConfirm::OnConfirmCallback);
+		}
+		else
+		{
+			// We have already been confirmed, just call the OnConfirm callback
+			OnConfirmCallback(Ability.Get());
+		}
+	}
+}
\ No newline at end of file
diff --git a/Engine/Source/Runtime/GameplayAbilities/Public/Abilities/Tasks/AbilityTask_WaitConfirm.h b/Engine/Source/Runtime/GameplayAbilities/Public/Abilities/Tasks/AbilityTask_WaitConfirm.h
new file mode 100644
--- /dev/null
+++ b/Engine/Source/Runtime/GameplayAbilities/Public/Abilities/Tasks/AbilityTask_WaitConfirm.h
@@ -0,0 +1,26 @@
+// Copyright 1998-2014 Epic Games, Inc. All Rights Reserved.
+#pragma once
+#include "AbilityTask.h"
+#include "Abilities/GameplayAbilityTypes.h"
+#include "AbilityTask_WaitConfirm.generated.h"
+
+UCLASS(MinimalAPI)
+class UAbilityTask_WaitConfirm : public UAbilityTask
+{
+	GENERATED_UCLASS_BODY()
+
+	UPROPERTY(BlueprintAssignable)
+	FGenericAbilityTaskDelegate	OnConfirm;
+
+	UFUNCTION()
+	void OnConfirmCallback(UGameplayAbility* Ability);
+
+	virtual void Activate() override;
+
+	/** Wait until the server confirms the use of this ability. This is used to gate predictive portions of the ability */
+	UFUNCTION(BlueprintCallable, Category=Abilities, meta = (HidePin = "WorldContextObject", DefaultToSelf = "WorldContextObject", BlueprintInternalUseOnly = "TRUE"))
+	static UAbilityTask_WaitConfirm* WaitConfirm(UObject* WorldContextObject);
+
+public:
+	
+};
\ No newline at end of file
