Fix ordering issues in AbilityActorInfo initialization

diff --git a/Engine/Source/Runtime/GameplayAbilities/Private/Abilities/GameplayAbilitySet.cpp b/Engine/Source/Runtime/GameplayAbilities/Private/Abilities/GameplayAbilitySet.cpp
--- a/Engine/Source/Runtime/GameplayAbilities/Private/Abilities/GameplayAbilitySet.cpp
+++ b/Engine/Source/Runtime/GameplayAbilities/Private/Abilities/GameplayAbilitySet.cpp
@@ -19,9 +19,6 @@ void InputPressed(TWeakObjectPtr<UGameplayAbility> Ability, int32 InputID, TWeak
 	UAbilitySystemComponent * AbilitySystemComponent = ActorInfo.AbilitySystemComponent.Get();
 	if (AbilitySystemComponent)
 	{
-		// FIXME: We have to reinit the ability component actor info to get a valid player controller
-		AbilitySystemComponent->InitAbilityActorInfo();
-
 		if (Ability.IsValid())
 		{
 			Ability.Get()->InputPressed(InputID, AbilitySystemComponent->AbilityActorInfo.Get());
