Cleaned-up version of traceline ability with proper network replication/simulation on clients and server, including predictive behavior. Ability system modified to implement this properly instead of hacking it in.

diff --git a/Engine/Source/Runtime/GameplayAbilities/Private/Abilities/GameplayAbilityTargetActor.cpp b/Engine/Source/Runtime/GameplayAbilities/Private/Abilities/GameplayAbilityTargetActor.cpp
--- a/Engine/Source/Runtime/GameplayAbilities/Private/Abilities/GameplayAbilityTargetActor.cpp
+++ b/Engine/Source/Runtime/GameplayAbilities/Private/Abilities/GameplayAbilityTargetActor.cpp
@@ -28,4 +28,10 @@ void AGameplayAbilityTargetActor::ConfirmTargeting()
 void AGameplayAbilityTargetActor::CancelTargeting()
 {
 	CanceledDelegate.Broadcast(FGameplayAbilityTargetDataHandle());
-}
\ No newline at end of file
+}
+
+bool AGameplayAbilityTargetActor::IsNetRelevantFor(class APlayerController* RealViewer, AActor* Viewer, const FVector& SrcLocation)
+{
+	//The player who created the ability doesn't need to be updated about it - there should be local prediction in place.
+	return (RealViewer == MasterPC) ? false : Super::IsNetRelevantFor(RealViewer, Viewer, SrcLocation);
+}
diff --git a/Engine/Source/Runtime/GameplayAbilities/Private/Abilities/GameplayAbilityTargetActor_SingleLineTrace.cpp b/Engine/Source/Runtime/GameplayAbilities/Private/Abilities/GameplayAbilityTargetActor_SingleLineTrace.cpp
--- a/Engine/Source/Runtime/GameplayAbilities/Private/Abilities/GameplayAbilityTargetActor_SingleLineTrace.cpp
+++ b/Engine/Source/Runtime/GameplayAbilities/Private/Abilities/GameplayAbilityTargetActor_SingleLineTrace.cpp
@@ -21,13 +21,16 @@ AGameplayAbilityTargetActor_SingleLineTrace::AGameplayAbilityTargetActor_SingleL
 	bBindToConfirmCancelInputs = true;
 }
 
-FGameplayAbilityTargetDataHandle AGameplayAbilityTargetActor_SingleLineTrace::StaticGetTargetData(UWorld * World, const FGameplayAbilityActorInfo* ActorInfo, FGameplayAbilityActivationInfo ActivationInfo)
+void AGameplayAbilityTargetActor_SingleLineTrace::GetLifetimeReplicatedProps(TArray< FLifetimeProperty > & OutLifetimeProps) const
 {
-	// very temp - do a mostly hardcoded trace from the source actor
+	Super::GetLifetimeReplicatedProps(OutLifetimeProps);
+	DOREPLIFETIME(AGameplayAbilityTargetActor_SingleLineTrace, bDebug);
+	DOREPLIFETIME(AGameplayAbilityTargetActor_SingleLineTrace, SourceActor);
+}
 
-	AActor * SourceActor = ActorInfo->Actor.Get();
-	check(SourceActor);
 
+FHitResult AGameplayAbilityTargetActor_SingleLineTrace::PerformTrace(AActor *SourceActor)
+{
 	static const FName LineTraceSingleName(TEXT("AGameplayAbilityTargetActor_SingleLineTrace"));
 	bool bTraceComplex = false;
 	TArray<AActor*> ActorsToIgnore;
@@ -39,21 +42,23 @@ FGameplayAbilityTargetDataHandle AGameplayAbilityTargetActor_SingleLineTrace::St
 	Params.bTraceAsyncScene = true;
 	Params.AddIgnoredActors(ActorsToIgnore);
 
-	FVector Start = SourceActor->GetActorLocation();
-	FVector End = Start + (SourceActor->GetActorForwardVector() * 3000.f);
+	FVector TraceStart = SourceActor->GetActorLocation();
+	FVector TraceEnd = TraceStart + (SourceActor->GetActorForwardVector() * 3000.f);
 
 	// ------------------------------------------------------
 
-	FGameplayAbilityTargetData_SingleTargetHit* ReturnData = new FGameplayAbilityTargetData_SingleTargetHit();
-
-	SourceActor->GetWorld()->LineTraceSingle(ReturnData->HitResult, Start, End, ECC_WorldStatic, Params);
+	FHitResult ReturnHitResult;
+	SourceActor->GetWorld()->LineTraceSingle(ReturnHitResult, TraceStart, TraceEnd, ECC_WorldStatic, Params);
+	return ReturnHitResult;
+}
 
-	if (bDebug)
-	{
-		DrawDebugLine(World, Start, End, FLinearColor::Green, false);
-		DrawDebugSphere(World, ReturnData->HitResult.Location, 16, 10, FLinearColor::Green, false);
-	}
+FGameplayAbilityTargetDataHandle AGameplayAbilityTargetActor_SingleLineTrace::StaticGetTargetData(UWorld * World, const FGameplayAbilityActorInfo* ActorInfo, FGameplayAbilityActivationInfo ActivationInfo)
+{
+	SourceActor = ActorInfo->Actor.Get();
+	check(SourceActor.Get());
 
+	FGameplayAbilityTargetData_SingleTargetHit* ReturnData = new FGameplayAbilityTargetData_SingleTargetHit();
+	ReturnData->HitResult = PerformTrace(SourceActor.Get());
 	return FGameplayAbilityTargetDataHandle(ReturnData);
 }
 
@@ -95,12 +100,31 @@ void AGameplayAbilityTargetActor_SingleLineTrace::Cancel()
 
 void AGameplayAbilityTargetActor_SingleLineTrace::Tick(float DeltaSeconds)
 {
-	Super::Tick(DeltaSeconds);
+	//Super::Tick(DeltaSeconds);
+	FGameplayAbilityTargetDataHandle Handle;
+	FHitResult HitResult;
 
 	/** Temp: Do a trace wiuth bDebug=true to draw a cheap "preview" */
 	if (Ability.IsValid())
 	{
-		StaticGetTargetData(Ability->GetWorld(), Ability->GetCurrentActorInfo(), Ability->GetCurrentActivationInfo());
+		Super::Tick(DeltaSeconds);
+		Handle = StaticGetTargetData(Ability->GetWorld(), Ability->GetCurrentActorInfo(), Ability->GetCurrentActivationInfo());
+		HitResult = *Handle.Data->GetHitResult();
+	}
+	else
+	{
+		Super::Tick(DeltaSeconds);
+	}
+
+	// very temp - do a mostly hardcoded trace from the source actor
+	if (bDebug && SourceActor.Get())
+	{
+		if (!Ability.IsValid())
+		{
+			HitResult = PerformTrace(SourceActor.Get());
+		}
+		DrawDebugLine(GetWorld(), SourceActor->GetActorLocation(), HitResult.Location, FLinearColor::Green, false);
+		DrawDebugSphere(GetWorld(), HitResult.Location, 16, 10, FLinearColor::Green, false);
 	}
 }
 
@@ -115,3 +139,4 @@ void AGameplayAbilityTargetActor_SingleLineTrace::ConfirmTargeting()
 
 	Destroy();
 }
+
diff --git a/Engine/Source/Runtime/GameplayAbilities/Private/Abilities/Tasks/AbilityTask_WaitTargetData.cpp b/Engine/Source/Runtime/GameplayAbilities/Private/Abilities/Tasks/AbilityTask_WaitTargetData.cpp
--- a/Engine/Source/Runtime/GameplayAbilities/Private/Abilities/Tasks/AbilityTask_WaitTargetData.cpp
+++ b/Engine/Source/Runtime/GameplayAbilities/Private/Abilities/Tasks/AbilityTask_WaitTargetData.cpp
@@ -27,7 +27,12 @@ bool UAbilityTask_WaitTargetData::BeginSpawningActor(UObject* WorldContextObject
 
 	if (Ability.IsValid())
 	{
-		if (Ability.Get()->GetCurrentActorInfo()->IsLocallyControlled())
+		bool bSpawnLocalVersion = false;
+
+		bSpawnLocalVersion |= (Ability.Get()->GetCurrentActorInfo()->IsLocallyControlled() && (Ability.Get()->GetNetExecutionPolicy() != EGameplayAbilityNetExecutionPolicy::Server));
+		bSpawnLocalVersion |= (Ability.Get()->GetNetExecutionPolicy() == EGameplayAbilityNetExecutionPolicy::Predictive);
+
+		if (bSpawnLocalVersion)
 		{
 			// Locally controlled - spawn the targeting actor.
 			AGameplayAbilityTargetActor* CDO = CastChecked<AGameplayAbilityTargetActor>(TargetClass->GetDefaultObject());
@@ -50,6 +55,12 @@ bool UAbilityTask_WaitTargetData::BeginSpawningActor(UObject* WorldContextObject
 				SpawnedActor->CanceledDelegate.AddUObject(this, &UAbilityTask_WaitTargetData::OnTargetDataCancelledCallback);
 
 				MySpawnedTargetActor = SpawnedActor;
+
+				AGameplayAbilityTargetActor* TargetActor = CastChecked<AGameplayAbilityTargetActor>(SpawnedActor);
+				if (TargetActor)
+				{
+					TargetActor->MasterPC = Ability.Get()->GetCurrentActorInfo()->PlayerController.Get();
+				}
 			}
 		}
 		else
diff --git a/Engine/Source/Runtime/GameplayAbilities/Public/Abilities/GameplayAbility.h b/Engine/Source/Runtime/GameplayAbilities/Public/Abilities/GameplayAbility.h
--- a/Engine/Source/Runtime/GameplayAbilities/Public/Abilities/GameplayAbility.h
+++ b/Engine/Source/Runtime/GameplayAbilities/Public/Abilities/GameplayAbility.h
@@ -152,6 +152,11 @@ public:
 		return ReplicationPolicy;
 	}
 
+	EGameplayAbilityNetExecutionPolicy::Type GetNetExecutionPolicy() const
+	{
+		return NetExecutionPolicy;
+	}
+
 	/** Gets the current actor info bound to this ability - can only be called on instanced abilities. */
 	const FGameplayAbilityActorInfo* GetCurrentActorInfo()
 	{
diff --git a/Engine/Source/Runtime/GameplayAbilities/Public/Abilities/GameplayAbilityTargetActor.h b/Engine/Source/Runtime/GameplayAbilities/Public/Abilities/GameplayAbilityTargetActor.h
--- a/Engine/Source/Runtime/GameplayAbilities/Public/Abilities/GameplayAbilityTargetActor.h
+++ b/Engine/Source/Runtime/GameplayAbilities/Public/Abilities/GameplayAbilityTargetActor.h
@@ -34,4 +34,9 @@ public:
 	
 	FAbilityTargetData	TargetDataReadyDelegate;
 	FAbilityTargetData	CanceledDelegate;
+
+	virtual bool IsNetRelevantFor(class APlayerController* RealViewer, AActor* Viewer, const FVector& SrcLocation);
+
+	UPROPERTY(BlueprintReadOnly, Category = Ability)
+	TWeakObjectPtr<APlayerController> MasterPC;
 };
\ No newline at end of file
diff --git a/Engine/Source/Runtime/GameplayAbilities/Public/Abilities/GameplayAbilityTargetActor_SingleLineTrace.h b/Engine/Source/Runtime/GameplayAbilities/Public/Abilities/GameplayAbilityTargetActor_SingleLineTrace.h
--- a/Engine/Source/Runtime/GameplayAbilities/Public/Abilities/GameplayAbilityTargetActor_SingleLineTrace.h
+++ b/Engine/Source/Runtime/GameplayAbilities/Public/Abilities/GameplayAbilityTargetActor_SingleLineTrace.h
@@ -11,7 +11,6 @@ class GAMEPLAYABILITIES_API AGameplayAbilityTargetActor_SingleLineTrace : public
 	GENERATED_UCLASS_BODY()
 
 public:
-	
 	virtual FGameplayAbilityTargetDataHandle StaticGetTargetData(UWorld * World, const FGameplayAbilityActorInfo* ActorInfo, FGameplayAbilityActivationInfo ActivationInfo) override;
 
 	virtual void StartTargeting(UGameplayAbility* Ability);
@@ -28,8 +27,16 @@ public:
 	
 	TWeakObjectPtr<UGameplayAbility> Ability;
 	
+	UPROPERTY(BlueprintReadOnly, Replicated, Category = Targeting)
 	bool bDebug;
 
+	//For remote clients, the actor who starts the trace (also gives us trace direction)
+	UPROPERTY(BlueprintReadOnly, Replicated, Category = Targeting)
+	TWeakObjectPtr<AActor> SourceActor;		//Maybe add a socket to this?
+
 	UPROPERTY(BlueprintReadOnly, meta=(ExposeOnSpawn=true), Category=Projectile)
 	bool bBindToConfirmCancelInputs;
+
+protected:
+	FHitResult PerformTrace(AActor *SourceActor);
 };
\ No newline at end of file
