GameplayAbilities: Blind fixing non-unity build

diff --git a/Engine/Source/Runtime/GameplayAbilities/Private/Abilities/GameplayAbilityTargetActor_SingleLineTrace.cpp b/Engine/Source/Runtime/GameplayAbilities/Private/Abilities/GameplayAbilityTargetActor_SingleLineTrace.cpp
--- a/Engine/Source/Runtime/GameplayAbilities/Private/Abilities/GameplayAbilityTargetActor_SingleLineTrace.cpp
+++ b/Engine/Source/Runtime/GameplayAbilities/Private/Abilities/GameplayAbilityTargetActor_SingleLineTrace.cpp
@@ -4,6 +4,7 @@
 #include "AbilitySystemComponent.h"
 #include "GameplayAbilityTargetActor_SingleLineTrace.h"
 #include "Engine/World.h"
+#include "Runtime/Engine/Public/Net/UnrealNetwork.h"
 
 // --------------------------------------------------------------------------------------------------------------------------------------------------------
 //
