New ability stuff for ticking movement, first pass.

diff --git a/Engine/Source/Runtime/GameplayAbilities/Private/Abilities/GameplayAbilityTargetActor_SingleLineTrace.cpp b/Engine/Source/Runtime/GameplayAbilities/Private/Abilities/GameplayAbilityTargetActor_SingleLineTrace.cpp
--- a/Engine/Source/Runtime/GameplayAbilities/Private/Abilities/GameplayAbilityTargetActor_SingleLineTrace.cpp
+++ b/Engine/Source/Runtime/GameplayAbilities/Private/Abilities/GameplayAbilityTargetActor_SingleLineTrace.cpp
@@ -118,14 +118,18 @@ void AGameplayAbilityTargetActor_SingleLineTrace::Tick(float DeltaSeconds)
 	}
 
 	// very temp - do a mostly hardcoded trace from the source actor
-	if (bDebug && SourceActor.Get())
+	if (SourceActor.Get())
 	{
 		if (!Ability.IsValid())
 		{
 			HitResult = PerformTrace(SourceActor.Get());
 		}
-		DrawDebugLine(GetWorld(), SourceActor->GetActorLocation(), HitResult.Location, FLinearColor::Green, false);
-		DrawDebugSphere(GetWorld(), HitResult.Location, 16, 10, FLinearColor::Green, false);
+		if (bDebug)
+		{
+			DrawDebugLine(GetWorld(), SourceActor->GetActorLocation(), HitResult.Location, FLinearColor::Green, false);
+			DrawDebugSphere(GetWorld(), HitResult.Location, 16, 10, FLinearColor::Green, false);
+		}
+		SetActorLocationAndRotation(HitResult.Location, SourceActor.Get()->GetActorRotation());
 	}
 }
 
diff --git a/Engine/Source/Runtime/GameplayAbilities/Private/Abilities/Tasks/AbilityTask_MoveToLocation.cpp b/Engine/Source/Runtime/GameplayAbilities/Private/Abilities/Tasks/AbilityTask_MoveToLocation.cpp
new file mode 100644
--- /dev/null
+++ b/Engine/Source/Runtime/GameplayAbilities/Private/Abilities/Tasks/AbilityTask_MoveToLocation.cpp
@@ -0,0 +1,96 @@
+
+#include "AbilitySystemPrivatePCH.h"
+#include "Abilities/Tasks/AbilityTask_MoveToLocation.h"
+
+UAbilityTask_MoveToLocation::UAbilityTask_MoveToLocation(const class FPostConstructInitializeProperties& PCIP)
+: Super(PCIP)
+{
+
+}
+
+
+
+void UAbilityTask_MoveToLocation::OnOverlapCallback(class AActor* OtherActor, class UPrimitiveComponent* OtherComp, int32 OtherBodyIndex, bool bFromSweep, const FHitResult & SweepResult)
+{
+	/*
+	if (OtherActor)
+	{
+		// Construct TargetData
+		FGameplayAbilityTargetData_SingleTargetHit * TargetData = new FGameplayAbilityTargetData_SingleTargetHit(Hit);
+
+		// Give it a handle and return
+		FGameplayAbilityTargetDataHandle	Handle;
+		Handle.Data = TSharedPtr<FGameplayAbilityTargetData>(TargetData);
+		OnOverlap.Broadcast(Handle);
+	}
+	*/
+}
+
+
+/**
+*	Need:
+*	-Easy way to specify which primitive components should be used for this overlap test
+*	-Easy way to specify which types of actors/collision overlaps that we care about/want to block on
+*/
+
+void UAbilityTask_MoveToLocation::InterpolatePosition()
+{
+	if (Ability.IsValid())
+	{
+		if (AActor* ActorOwner = Cast<AActor>(Ability->GetOuter()))
+		{
+			float CurrentTime = ActorOwner->GetWorld()->GetTimeSeconds();
+
+			if (CurrentTime >= TimeMoveWillEnd)
+			{
+				ActorOwner->GetWorld()->GetTimerManager().ClearTimer(this, &UAbilityTask_MoveToLocation::InterpolatePosition);
+				ActorOwner->SetActorLocation(TargetLocation);
+				OnTargetLocationReached.Broadcast();
+			}
+			else
+			{
+				float MoveFraction = (CurrentTime - TimeMoveStarted) / DurationOfMovement;
+				ActorOwner->SetActorLocation(FMath::Lerp<FVector, float>(StartLocation, TargetLocation, MoveFraction));
+			}
+		}
+	}
+}
+
+UAbilityTask_MoveToLocation* UAbilityTask_MoveToLocation::MoveToLocation(class UObject* WorldContextObject, FVector Location, float Duration)
+{
+	check(WorldContextObject);
+	UGameplayAbility* ThisAbility = CastChecked<UGameplayAbility>(WorldContextObject);
+	check(ThisAbility);
+	AActor* ActorOwner = Cast<AActor>(ThisAbility->GetOuter());
+	check(ActorOwner);
+
+	UAbilityTask_MoveToLocation* MyObj = NewObject<UAbilityTask_MoveToLocation>();
+	MyObj->Ability = ThisAbility;
+
+	MyObj->StartLocation = ActorOwner->GetActorLocation();
+	MyObj->TargetLocation = Location;
+	MyObj->DurationOfMovement = FMath::Max(Duration, 0.001f);		//Avoid negative or divide-by-zero cases
+	MyObj->TimeMoveStarted = WorldContextObject->GetWorld()->GetTimeSeconds();
+	MyObj->TimeMoveWillEnd = MyObj->TimeMoveStarted + MyObj->DurationOfMovement;
+	WorldContextObject->GetWorld()->GetTimerManager().SetTimer(MyObj, &UAbilityTask_MoveToLocation::InterpolatePosition, 0.001f, true);
+
+	return MyObj;
+}
+
+void UAbilityTask_MoveToLocation::Activate()
+{
+	if (Ability.IsValid())
+	{
+		AActor* ActorOwner = Cast<AActor>(Ability->GetOuter());
+
+		// TEMP - we are just using root component's collision. A real system will need more data to specify which component to use
+		UPrimitiveComponent * PrimComponent = Cast<UPrimitiveComponent>(ActorOwner->GetRootComponent());
+		if (!PrimComponent)
+		{
+			PrimComponent = ActorOwner->FindComponentByClass<UPrimitiveComponent>();
+		}
+		check(PrimComponent);
+
+		//PrimComponent->OnComponentBeginOverlap.AddDynamic(this, &UAbilityTask_WaitOverlap::OnOverlapCallback);
+	}
+}
diff --git a/Engine/Source/Runtime/GameplayAbilities/Public/Abilities/Tasks/AbilityTask_MoveToLocation.h b/Engine/Source/Runtime/GameplayAbilities/Public/Abilities/Tasks/AbilityTask_MoveToLocation.h
new file mode 100644
--- /dev/null
+++ b/Engine/Source/Runtime/GameplayAbilities/Public/Abilities/Tasks/AbilityTask_MoveToLocation.h
@@ -0,0 +1,40 @@
+// Copyright 1998-2014 Epic Games, Inc. All Rights Reserved.
+#pragma once
+#include "AbilityTask.h"
+#include "Abilities/GameplayAbilityTypes.h"
+#include "AbilityTask_MoveToLocation.generated.h"
+
+DECLARE_DYNAMIC_MULTICAST_DELEGATE(FMoveToLocationDelegate);
+
+class AActor;
+class UPrimitiveComponent;
+
+/**
+*		Move to a location, ignoring clipping, over a given length of time. Callback whenever we pass through an actor and when we reach our target location.
+*/
+UCLASS(MinimalAPI)
+class UAbilityTask_MoveToLocation : public UAbilityTask
+{
+	GENERATED_UCLASS_BODY()
+
+	UPROPERTY(BlueprintAssignable)
+	FMoveToLocationDelegate		OnTargetLocationReached;
+
+	void InterpolatePosition();
+
+	UFUNCTION()
+		void OnOverlapCallback(AActor* OtherActor, UPrimitiveComponent* OtherComp, int32 OtherBodyIndex, bool bFromSweep, const FHitResult & SweepResult);
+
+	virtual void Activate() override;
+
+	/** Wait until an overlap occurs. This will need to be better fleshed out so we can specify game specific collision requirements */
+	UFUNCTION(BlueprintCallable, Category = Abilities, meta = (HidePin = "WorldContextObject", DefaultToSelf = "WorldContextObject", BlueprintInternalUseOnly = "TRUE"))
+		static UAbilityTask_MoveToLocation* MoveToLocation(UObject* WorldContextObject, FVector Location, float Duration);
+
+protected:
+	FVector StartLocation;
+	FVector TargetLocation;
+	float DurationOfMovement;
+	float TimeMoveStarted;
+	float TimeMoveWillEnd;
+};
\ No newline at end of file
