Added a curve to movement task. There seems to be something off with the aiming task.

diff --git a/Engine/Source/Runtime/GameplayAbilities/Private/Abilities/Tasks/AbilityTask_MoveToLocation.cpp b/Engine/Source/Runtime/GameplayAbilities/Private/Abilities/Tasks/AbilityTask_MoveToLocation.cpp
--- a/Engine/Source/Runtime/GameplayAbilities/Private/Abilities/Tasks/AbilityTask_MoveToLocation.cpp
+++ b/Engine/Source/Runtime/GameplayAbilities/Private/Abilities/Tasks/AbilityTask_MoveToLocation.cpp
@@ -8,31 +8,6 @@ UAbilityTask_MoveToLocation::UAbilityTask_MoveToLocation(const class FPostConstr
 
 }
 
-
-
-void UAbilityTask_MoveToLocation::OnOverlapCallback(class AActor* OtherActor, class UPrimitiveComponent* OtherComp, int32 OtherBodyIndex, bool bFromSweep, const FHitResult & SweepResult)
-{
-	/*
-	if (OtherActor)
-	{
-		// Construct TargetData
-		FGameplayAbilityTargetData_SingleTargetHit * TargetData = new FGameplayAbilityTargetData_SingleTargetHit(Hit);
-
-		// Give it a handle and return
-		FGameplayAbilityTargetDataHandle	Handle;
-		Handle.Data = TSharedPtr<FGameplayAbilityTargetData>(TargetData);
-		OnOverlap.Broadcast(Handle);
-	}
-	*/
-}
-
-
-/**
-*	Need:
-*	-Easy way to specify which primitive components should be used for this overlap test
-*	-Easy way to specify which types of actors/collision overlaps that we care about/want to block on
-*/
-
 void UAbilityTask_MoveToLocation::InterpolatePosition()
 {
 	if (Ability.IsValid())
@@ -50,13 +25,17 @@ void UAbilityTask_MoveToLocation::InterpolatePosition()
 			else
 			{
 				float MoveFraction = (CurrentTime - TimeMoveStarted) / DurationOfMovement;
-				ActorOwner->SetActorLocation(FMath::Lerp<FVector, float>(StartLocation, TargetLocation, MoveFraction));
+				if (LerpCurve.IsValid())
+				{
+					MoveFraction = LerpCurve.Get()->GetFloatValue(MoveFraction);
+				}
+				ActorOwner->SetActorLocation(FMath::Lerp<FVector, float>(StartLocation, TargetLocation, MoveFraction));				
 			}
 		}
 	}
 }
 
-UAbilityTask_MoveToLocation* UAbilityTask_MoveToLocation::MoveToLocation(class UObject* WorldContextObject, FVector Location, float Duration)
+UAbilityTask_MoveToLocation* UAbilityTask_MoveToLocation::MoveToLocation(class UObject* WorldContextObject, FVector Location, float Duration, UCurveFloat* OptionalInterpolationCurve)
 {
 	check(WorldContextObject);
 	UGameplayAbility* ThisAbility = CastChecked<UGameplayAbility>(WorldContextObject);
@@ -72,25 +51,8 @@ UAbilityTask_MoveToLocation* UAbilityTask_MoveToLocation::MoveToLocation(class U
 	MyObj->DurationOfMovement = FMath::Max(Duration, 0.001f);		//Avoid negative or divide-by-zero cases
 	MyObj->TimeMoveStarted = WorldContextObject->GetWorld()->GetTimeSeconds();
 	MyObj->TimeMoveWillEnd = MyObj->TimeMoveStarted + MyObj->DurationOfMovement;
+	MyObj->LerpCurve = OptionalInterpolationCurve;
 	WorldContextObject->GetWorld()->GetTimerManager().SetTimer(MyObj, &UAbilityTask_MoveToLocation::InterpolatePosition, 0.001f, true);
 
 	return MyObj;
 }
-
-void UAbilityTask_MoveToLocation::Activate()
-{
-	if (Ability.IsValid())
-	{
-		AActor* ActorOwner = Cast<AActor>(Ability->GetOuter());
-
-		// TEMP - we are just using root component's collision. A real system will need more data to specify which component to use
-		UPrimitiveComponent * PrimComponent = Cast<UPrimitiveComponent>(ActorOwner->GetRootComponent());
-		if (!PrimComponent)
-		{
-			PrimComponent = ActorOwner->FindComponentByClass<UPrimitiveComponent>();
-		}
-		check(PrimComponent);
-
-		//PrimComponent->OnComponentBeginOverlap.AddDynamic(this, &UAbilityTask_WaitOverlap::OnOverlapCallback);
-	}
-}
diff --git a/Engine/Source/Runtime/GameplayAbilities/Public/Abilities/Tasks/AbilityTask_MoveToLocation.h b/Engine/Source/Runtime/GameplayAbilities/Public/Abilities/Tasks/AbilityTask_MoveToLocation.h
--- a/Engine/Source/Runtime/GameplayAbilities/Public/Abilities/Tasks/AbilityTask_MoveToLocation.h
+++ b/Engine/Source/Runtime/GameplayAbilities/Public/Abilities/Tasks/AbilityTask_MoveToLocation.h
@@ -22,14 +22,9 @@ class UAbilityTask_MoveToLocation : public UAbilityTask
 
 	void InterpolatePosition();
 
-	UFUNCTION()
-		void OnOverlapCallback(AActor* OtherActor, UPrimitiveComponent* OtherComp, int32 OtherBodyIndex, bool bFromSweep, const FHitResult & SweepResult);
-
-	virtual void Activate() override;
-
-	/** Wait until an overlap occurs. This will need to be better fleshed out so we can specify game specific collision requirements */
+	/** Move to the specified location, using the curve (range 0 - 1) or linearly if no curve is specified */
 	UFUNCTION(BlueprintCallable, Category = Abilities, meta = (HidePin = "WorldContextObject", DefaultToSelf = "WorldContextObject", BlueprintInternalUseOnly = "TRUE"))
-		static UAbilityTask_MoveToLocation* MoveToLocation(UObject* WorldContextObject, FVector Location, float Duration);
+		static UAbilityTask_MoveToLocation* MoveToLocation(UObject* WorldContextObject, FVector Location, float Duration, UCurveFloat* OptionalInterpolationCurve);
 
 protected:
 	FVector StartLocation;
@@ -37,4 +32,5 @@ protected:
 	float DurationOfMovement;
 	float TimeMoveStarted;
 	float TimeMoveWillEnd;
+	TWeakObjectPtr<UCurveFloat> LerpCurve;
 };
\ No newline at end of file
