Cleanup and a test upgrade.

diff --git a/Engine/Source/Runtime/GameplayAbilities/Private/Abilities/GameplayAbilityTargetActor_SingleLineTrace.cpp b/Engine/Source/Runtime/GameplayAbilities/Private/Abilities/GameplayAbilityTargetActor_SingleLineTrace.cpp
--- a/Engine/Source/Runtime/GameplayAbilities/Private/Abilities/GameplayAbilityTargetActor_SingleLineTrace.cpp
+++ b/Engine/Source/Runtime/GameplayAbilities/Private/Abilities/GameplayAbilityTargetActor_SingleLineTrace.cpp
@@ -101,21 +101,16 @@ void AGameplayAbilityTargetActor_SingleLineTrace::Cancel()
 
 void AGameplayAbilityTargetActor_SingleLineTrace::Tick(float DeltaSeconds)
 {
-	//Super::Tick(DeltaSeconds);
+	Super::Tick(DeltaSeconds);
 	FGameplayAbilityTargetDataHandle Handle;
 	FHitResult HitResult;
 
 	/** Temp: Do a trace wiuth bDebug=true to draw a cheap "preview" */
 	if (Ability.IsValid())
 	{
-		Super::Tick(DeltaSeconds);
 		Handle = StaticGetTargetData(Ability->GetWorld(), Ability->GetCurrentActorInfo(), Ability->GetCurrentActivationInfo());
 		HitResult = *Handle.Data->GetHitResult();
 	}
-	else
-	{
-		Super::Tick(DeltaSeconds);
-	}
 
 	// very temp - do a mostly hardcoded trace from the source actor
 	if (SourceActor.Get())
@@ -126,7 +121,7 @@ void AGameplayAbilityTargetActor_SingleLineTrace::Tick(float DeltaSeconds)
 		}
 		if (bDebug)
 		{
-			DrawDebugLine(GetWorld(), SourceActor->GetActorLocation(), HitResult.Location, FLinearColor::Green, false);
+			DrawDebugLine(GetWorld(), SourceActor->GetActorLocation(), HitResult.ImpactPoint, FLinearColor::Green, false);
 			DrawDebugSphere(GetWorld(), HitResult.Location, 16, 10, FLinearColor::Green, false);
 		}
 		SetActorLocationAndRotation(HitResult.Location, SourceActor.Get()->GetActorRotation());
