Targeting updates:
-Fix server-spawning logic in ::BeginSpawningActor
-consolidated MFD_3DReticle into TestDaggerReplication for now.
-Added ::OnReplicatedTargetDataReceived for GameplayTargetActors to transform or reject target data server side before accepting data from client.
-Misc cleanup.

diff --git a/Engine/Source/Runtime/GameplayAbilities/Private/Abilities/GameplayAbilityTargetActor.cpp b/Engine/Source/Runtime/GameplayAbilities/Private/Abilities/GameplayAbilityTargetActor.cpp
--- a/Engine/Source/Runtime/GameplayAbilities/Private/Abilities/GameplayAbilityTargetActor.cpp
+++ b/Engine/Source/Runtime/GameplayAbilities/Private/Abilities/GameplayAbilityTargetActor.cpp
@@ -33,5 +33,25 @@ void AGameplayAbilityTargetActor::CancelTargeting()
 bool AGameplayAbilityTargetActor::IsNetRelevantFor(class APlayerController* RealViewer, AActor* Viewer, const FVector& SrcLocation)
 {
 	//The player who created the ability doesn't need to be updated about it - there should be local prediction in place.
-	return (RealViewer == MasterPC) ? false : Super::IsNetRelevantFor(RealViewer, Viewer, SrcLocation);
+	if (RealViewer == MasterPC)
+	{
+		return false;
+	}
+
+	return Super::IsNetRelevantFor(RealViewer, Viewer, SrcLocation);
+}
+
+bool AGameplayAbilityTargetActor::GetReplicates() const
+{
+	return bReplicates;
+}
+
+FGameplayAbilityTargetDataHandle AGameplayAbilityTargetActor::StaticGetTargetData(UWorld * World, const FGameplayAbilityActorInfo* ActorInfo, FGameplayAbilityActivationInfo ActivationInfo) const
+{
+	return FGameplayAbilityTargetDataHandle();
+}
+
+bool AGameplayAbilityTargetActor::OnReplicatedTargetDataReceived(FGameplayAbilityTargetDataHandle& Data) const
+{
+	return true;
 }
diff --git a/Engine/Source/Runtime/GameplayAbilities/Private/Abilities/GameplayAbilityTargetActor_SingleLineTrace.cpp b/Engine/Source/Runtime/GameplayAbilities/Private/Abilities/GameplayAbilityTargetActor_SingleLineTrace.cpp
--- a/Engine/Source/Runtime/GameplayAbilities/Private/Abilities/GameplayAbilityTargetActor_SingleLineTrace.cpp
+++ b/Engine/Source/Runtime/GameplayAbilities/Private/Abilities/GameplayAbilityTargetActor_SingleLineTrace.cpp
@@ -30,42 +30,43 @@ void AGameplayAbilityTargetActor_SingleLineTrace::GetLifetimeReplicatedProps(TAr
 }
 
 
-FHitResult AGameplayAbilityTargetActor_SingleLineTrace::PerformTrace(AActor *SourceActor)
+FHitResult AGameplayAbilityTargetActor_SingleLineTrace::PerformTrace(AActor *InSourceActor) const
 {
 	static const FName LineTraceSingleName(TEXT("AGameplayAbilityTargetActor_SingleLineTrace"));
 	bool bTraceComplex = false;
 	TArray<AActor*> ActorsToIgnore;
 
-	ActorsToIgnore.Add(SourceActor);
+	ActorsToIgnore.Add(InSourceActor);
 
 	FCollisionQueryParams Params(LineTraceSingleName, bTraceComplex);
 	Params.bReturnPhysicalMaterial = true;
 	Params.bTraceAsyncScene = true;
 	Params.AddIgnoredActors(ActorsToIgnore);
 
-	FVector TraceStart = SourceActor->GetActorLocation();
-	FVector TraceEnd = TraceStart + (SourceActor->GetActorForwardVector() * 3000.f);
+	FVector TraceStart = InSourceActor->GetActorLocation();
+	FVector TraceEnd = TraceStart + (InSourceActor->GetActorForwardVector() * 3000.f);
 
 	// ------------------------------------------------------
 
 	FHitResult ReturnHitResult;
-	SourceActor->GetWorld()->LineTraceSingle(ReturnHitResult, TraceStart, TraceEnd, ECC_WorldStatic, Params);
+	InSourceActor->GetWorld()->LineTraceSingle(ReturnHitResult, TraceStart, TraceEnd, ECC_WorldStatic, Params);
 	return ReturnHitResult;
 }
 
-FGameplayAbilityTargetDataHandle AGameplayAbilityTargetActor_SingleLineTrace::StaticGetTargetData(UWorld * World, const FGameplayAbilityActorInfo* ActorInfo, FGameplayAbilityActivationInfo ActivationInfo)
+FGameplayAbilityTargetDataHandle AGameplayAbilityTargetActor_SingleLineTrace::StaticGetTargetData(UWorld * World, const FGameplayAbilityActorInfo* ActorInfo, FGameplayAbilityActivationInfo ActivationInfo) const
 {
-	SourceActor = ActorInfo->Actor.Get();
-	check(SourceActor.Get());
+	AActor* StaticSourceActor = ActorInfo->Actor.Get();
+	check(StaticSourceActor);
 
 	FGameplayAbilityTargetData_SingleTargetHit* ReturnData = new FGameplayAbilityTargetData_SingleTargetHit();
-	ReturnData->HitResult = PerformTrace(SourceActor.Get());
+	ReturnData->HitResult = PerformTrace(StaticSourceActor);
 	return FGameplayAbilityTargetDataHandle(ReturnData);
 }
 
 void AGameplayAbilityTargetActor_SingleLineTrace::StartTargeting(UGameplayAbility* InAbility)
 {
 	Ability = InAbility;
+	SourceActor = InAbility->GetCurrentActorInfo()->Actor.Get();
 
 	// We can bind directly to our ASC's confirm/cancel events, or wait to be told from an outside source to confirm or cancel
 	if (bBindToConfirmCancelInputs)
@@ -105,26 +106,17 @@ void AGameplayAbilityTargetActor_SingleLineTrace::Tick(float DeltaSeconds)
 	FGameplayAbilityTargetDataHandle Handle;
 	FHitResult HitResult;
 
-	/** Temp: Do a trace wiuth bDebug=true to draw a cheap "preview" */
-	if (Ability.IsValid())
-	{
-		Handle = StaticGetTargetData(Ability->GetWorld(), Ability->GetCurrentActorInfo(), Ability->GetCurrentActivationInfo());
-		HitResult = *Handle.Data->GetHitResult();
-	}
-
 	// very temp - do a mostly hardcoded trace from the source actor
-	if (SourceActor.Get())
+	if (SourceActor)
 	{
-		if (!Ability.IsValid())
-		{
-			HitResult = PerformTrace(SourceActor.Get());
-		}
+		HitResult = PerformTrace(SourceActor);
+		
 		if (bDebug)
 		{
 			DrawDebugLine(GetWorld(), SourceActor->GetActorLocation(), HitResult.ImpactPoint, FLinearColor::Green, false);
 			DrawDebugSphere(GetWorld(), HitResult.Location, 16, 10, FLinearColor::Green, false);
 		}
-		SetActorLocationAndRotation(HitResult.Location, SourceActor.Get()->GetActorRotation());
+		SetActorLocationAndRotation(HitResult.Location, SourceActor->GetActorRotation());
 	}
 }
 
diff --git a/Engine/Source/Runtime/GameplayAbilities/Private/Abilities/Tasks/AbilityTask_WaitTargetData.cpp b/Engine/Source/Runtime/GameplayAbilities/Private/Abilities/Tasks/AbilityTask_WaitTargetData.cpp
--- a/Engine/Source/Runtime/GameplayAbilities/Private/Abilities/Tasks/AbilityTask_WaitTargetData.cpp
+++ b/Engine/Source/Runtime/GameplayAbilities/Private/Abilities/Tasks/AbilityTask_WaitTargetData.cpp
@@ -25,21 +25,28 @@ bool UAbilityTask_WaitTargetData::BeginSpawningActor(UObject* WorldContextObject
 {
 	SpawnedActor = nullptr;
 
-	if (Ability.IsValid())
+	UGameplayAbility* MyAbility = Ability.Get();
+	if (MyAbility)
 	{
-		bool bSpawnLocalVersion = false;
+		const AGameplayAbilityTargetActor* CDO = CastChecked<AGameplayAbilityTargetActor>(TargetClass->GetDefaultObject());
+		MyTargetActor = CDO;
 
-		bSpawnLocalVersion |= (Ability.Get()->GetCurrentActorInfo()->IsLocallyControlled() && (Ability.Get()->GetNetExecutionPolicy() != EGameplayAbilityNetExecutionPolicy::Server));
-		bSpawnLocalVersion |= (Ability.Get()->GetNetExecutionPolicy() == EGameplayAbilityNetExecutionPolicy::Predictive);
+		bool Replicates = CDO->GetReplicates();
+		bool StaticFunc = CDO->StaticTargetFunction;
 
-		if (bSpawnLocalVersion)
+		if (Replicates && StaticFunc)
 		{
-			// Locally controlled - spawn the targeting actor.
-			AGameplayAbilityTargetActor* CDO = CastChecked<AGameplayAbilityTargetActor>(TargetClass->GetDefaultObject());
-			if (CDO->StaticTargetFunction)
+			ABILITY_LOG(Fatal, TEXT("AbilityTargetActor class %s can't be Replicating and Static"), *TargetClass->GetName());
+			Replicates = false;
+		}
+
+		// Spawn the actor if this is a locally controlled ability (always) or if this is a replicated targeting mode.
+		if (Replicates || MyAbility->GetCurrentActorInfo()->IsLocallyControlled())
+		{
+			if (StaticFunc)
 			{
 				// This is just a static function that should instantly give us back target data
-				FGameplayAbilityTargetDataHandle Data = CDO->StaticGetTargetData(Ability->GetWorld(), Ability->GetCurrentActorInfo(), Ability->GetCurrentActivationInfo());
+				FGameplayAbilityTargetDataHandle Data = CDO->StaticGetTargetData(MyAbility->GetWorld(), MyAbility->GetCurrentActorInfo(), MyAbility->GetCurrentActivationInfo());
 				OnTargetDataReadyCallback(Data);
 			}
 			else
@@ -54,12 +61,12 @@ bool UAbilityTask_WaitTargetData::BeginSpawningActor(UObject* WorldContextObject
 				SpawnedActor->TargetDataReadyDelegate.AddUObject(this, &UAbilityTask_WaitTargetData::OnTargetDataReadyCallback);
 				SpawnedActor->CanceledDelegate.AddUObject(this, &UAbilityTask_WaitTargetData::OnTargetDataCancelledCallback);
 
-				MySpawnedTargetActor = SpawnedActor;
+				MyTargetActor = SpawnedActor;
 
 				AGameplayAbilityTargetActor* TargetActor = CastChecked<AGameplayAbilityTargetActor>(SpawnedActor);
 				if (TargetActor)
 				{
-					TargetActor->MasterPC = Ability.Get()->GetCurrentActorInfo()->PlayerController.Get();
+					TargetActor->MasterPC = MyAbility->GetCurrentActorInfo()->PlayerController.Get();
 				}
 			}
 		}
@@ -87,7 +94,7 @@ void UAbilityTask_WaitTargetData::FinishSpawningActor(UObject* WorldContextObjec
 {
 	if (SpawnedActor)
 	{
-		check(MySpawnedTargetActor == SpawnedActor);
+		check(MyTargetActor == SpawnedActor);
 
 		FTransform SpawnTransform = AbilitySystemComponent->GetOwner()->GetTransform();
 
@@ -109,7 +116,23 @@ void UAbilityTask_WaitTargetData::OnTargetDataReplicatedCallback(FGameplayAbilit
 {
 	check(AbilitySystemComponent.IsValid());
 
-	ValidData.Broadcast(Data);
+	/** 
+	 *  Call into the TargetActor to sanitize/verify the data. If this returns false, we are rejecting
+	 *	the replicated target data and will treat this as a cancel.
+	 *	
+	 *	This can also be used for bandwidth optimizations. OnReplicatedTargetDataReceived could do an actual
+	 *	trace/check/whatever server side and use that data. So rather than having the client send that data
+	 *	explicitly, the client is basically just sending a 'confirm' and the server is now going to do the work
+	 *	in OnReplicatedTargetDataReceived.
+	 */
+	if (MyTargetActor.IsValid() && !MyTargetActor->OnReplicatedTargetDataReceived(Data))
+	{
+		Cancelled.Broadcast(Data);
+	}
+	else
+	{
+		ValidData.Broadcast(Data);
+	}
 
 	Cleanup();
 }
@@ -154,9 +177,9 @@ void UAbilityTask_WaitTargetData::Cleanup()
 	AbilitySystemComponent->ReplicatedTargetDataDelegate.RemoveUObject(this, &UAbilityTask_WaitTargetData::OnTargetDataReplicatedCallback);
 	AbilitySystemComponent->ReplicatedTargetDataCancelledDelegate.RemoveDynamic(this, &UAbilityTask_WaitTargetData::OnTargetDataReplicatedCancelledCallback);
 
-	if (MySpawnedTargetActor.IsValid())
+	if (MyTargetActor.IsValid() && !MyTargetActor->HasAnyFlags(RF_ClassDefaultObject))
 	{
-		MySpawnedTargetActor->Destroy();
+		MyTargetActor->Destroy();
 	}
 
 	MarkPendingKill();
diff --git a/Engine/Source/Runtime/GameplayAbilities/Public/Abilities/GameplayAbilityTargetActor.h b/Engine/Source/Runtime/GameplayAbilities/Public/Abilities/GameplayAbilityTargetActor.h
--- a/Engine/Source/Runtime/GameplayAbilities/Public/Abilities/GameplayAbilityTargetActor.h
+++ b/Engine/Source/Runtime/GameplayAbilities/Public/Abilities/GameplayAbilityTargetActor.h
@@ -4,6 +4,7 @@
 #include "GameplayAbility.h"
 #include "GameplayAbilityTargetActor.generated.h"
 
+
 UCLASS(Blueprintable, abstract)
 class GAMEPLAYABILITIES_API AGameplayAbilityTargetActor : public AActor
 {
@@ -11,17 +12,13 @@ class GAMEPLAYABILITIES_API AGameplayAbilityTargetActor : public AActor
 
 public:
 
-	/** Native classes can set this to call the non-blueprintable  */
+	/** Native classes can set this to call StaticGetTargetData to instantly get TargetData, instead of instantiating the GameplayTargetActor and calling StartTargeting. */
 	UPROPERTY()
-	bool	StaticTargetFunction;
-
-	virtual FGameplayAbilityTargetDataHandle StaticGetTargetData(UWorld * World, const FGameplayAbilityActorInfo* ActorInfo, FGameplayAbilityActivationInfo ActivationInfo) { return FGameplayAbilityTargetDataHandle(); }
-
-	/** Otherwise, the actor will be instantiated and call the blueprintable function GetTargetData */
-	UFUNCTION(BlueprintImplementableEvent, Category = "Targeting")
-	FGameplayAbilityTargetDataHandle StaticGetTargetData(UWorld * World, FGameplayAbilityActorInfo ActorInfo, FGameplayAbilityActivationInfo ActivationInfo);
+	bool StaticTargetFunction;
 
+	virtual FGameplayAbilityTargetDataHandle StaticGetTargetData(UWorld * World, const FGameplayAbilityActorInfo* ActorInfo, FGameplayAbilityActivationInfo ActivationInfo) const;
 	
+	/** Initialize and begin targeting logic  */
 	virtual void StartTargeting(UGameplayAbility* Ability);
 
 	/** Outside code is saying 'stop and just give me what you have' */
@@ -30,13 +27,19 @@ public:
 	/** Outside code is saying 'stop everything and just forget about it' */
 	virtual void CancelTargeting();
 
+	/** Replicated target data was received from a client. Possibly sanitize/verify. return true if data is good and we should broadcast it as valid data. */
+	virtual bool OnReplicatedTargetDataReceived(FGameplayAbilityTargetDataHandle& Data) const;
+
+	/** Accessor for checking, before instantiating, if this TargetActor will replicate. */
+	bool GetReplicates() const;
+
 	// ------------------------------
 	
 	FAbilityTargetData	TargetDataReadyDelegate;
 	FAbilityTargetData	CanceledDelegate;
 
-	virtual bool IsNetRelevantFor(class APlayerController* RealViewer, AActor* Viewer, const FVector& SrcLocation);
+	virtual bool IsNetRelevantFor(class APlayerController* RealViewer, AActor* Viewer, const FVector& SrcLocation) override;
 
-	UPROPERTY(BlueprintReadOnly, Category = Ability)
-	TWeakObjectPtr<APlayerController> MasterPC;
+	UPROPERTY(BlueprintReadOnly, Category="Targeting")
+	APlayerController* MasterPC;
 };
\ No newline at end of file
diff --git a/Engine/Source/Runtime/GameplayAbilities/Public/Abilities/GameplayAbilityTargetActor_SingleLineTrace.h b/Engine/Source/Runtime/GameplayAbilities/Public/Abilities/GameplayAbilityTargetActor_SingleLineTrace.h
--- a/Engine/Source/Runtime/GameplayAbilities/Public/Abilities/GameplayAbilityTargetActor_SingleLineTrace.h
+++ b/Engine/Source/Runtime/GameplayAbilities/Public/Abilities/GameplayAbilityTargetActor_SingleLineTrace.h
@@ -11,7 +11,7 @@ class GAMEPLAYABILITIES_API AGameplayAbilityTargetActor_SingleLineTrace : public
 	GENERATED_UCLASS_BODY()
 
 public:
-	virtual FGameplayAbilityTargetDataHandle StaticGetTargetData(UWorld * World, const FGameplayAbilityActorInfo* ActorInfo, FGameplayAbilityActivationInfo ActivationInfo) override;
+	virtual FGameplayAbilityTargetDataHandle StaticGetTargetData(UWorld * World, const FGameplayAbilityActorInfo* ActorInfo, FGameplayAbilityActivationInfo ActivationInfo) const override;
 
 	virtual void StartTargeting(UGameplayAbility* Ability);
 
@@ -32,11 +32,11 @@ public:
 
 	//For remote clients, the actor who starts the trace (also gives us trace direction)
 	UPROPERTY(BlueprintReadOnly, Replicated, Category = Targeting)
-	TWeakObjectPtr<AActor> SourceActor;		//Maybe add a socket to this?
+	AActor* SourceActor;		//Maybe add a socket to this?
 
 	UPROPERTY(BlueprintReadOnly, meta=(ExposeOnSpawn=true), Category=Projectile)
 	bool bBindToConfirmCancelInputs;
 
 protected:
-	FHitResult PerformTrace(AActor *SourceActor);
+	FHitResult PerformTrace(AActor *SourceActor) const;
 };
\ No newline at end of file
diff --git a/Engine/Source/Runtime/GameplayAbilities/Public/Abilities/Tasks/AbilityTask_WaitTargetData.h b/Engine/Source/Runtime/GameplayAbilities/Public/Abilities/Tasks/AbilityTask_WaitTargetData.h
--- a/Engine/Source/Runtime/GameplayAbilities/Public/Abilities/Tasks/AbilityTask_WaitTargetData.h
+++ b/Engine/Source/Runtime/GameplayAbilities/Public/Abilities/Tasks/AbilityTask_WaitTargetData.h
@@ -42,14 +42,14 @@ class UAbilityTask_WaitTargetData: public UAbilityTask
 	void FinishSpawningActor(UObject* WorldContextObject, AGameplayAbilityTargetActor* SpawnedActor);
 
 
-
 protected:
 
 	TSubclassOf<AGameplayAbilityTargetActor> TargetClass;
 
 	void Cleanup();
 
-	TWeakObjectPtr<AActor>	MySpawnedTargetActor;
+	/** The TargetActor that we spawned, or the class CDO if this is a static targeting task */
+	TWeakObjectPtr<AGameplayAbilityTargetActor>	MyTargetActor;
 };
 
 
