Minor cleanup - remove redundant functions

diff --git a/Engine/Source/Runtime/GameplayAbilities/Private/Abilities/GameplayAbilityTargetActor.cpp b/Engine/Source/Runtime/GameplayAbilities/Private/Abilities/GameplayAbilityTargetActor.cpp
--- a/Engine/Source/Runtime/GameplayAbilities/Private/Abilities/GameplayAbilityTargetActor.cpp
+++ b/Engine/Source/Runtime/GameplayAbilities/Private/Abilities/GameplayAbilityTargetActor.cpp
@@ -22,12 +22,14 @@ void AGameplayAbilityTargetActor::StartTargeting(UGameplayAbility* Ability)
 void AGameplayAbilityTargetActor::ConfirmTargeting()
 {
 	TargetDataReadyDelegate.Broadcast(FGameplayAbilityTargetDataHandle());
+	Destroy();
 }
 
 /** Outside code is saying 'stop everything and just forget about it' */
 void AGameplayAbilityTargetActor::CancelTargeting()
 {
 	CanceledDelegate.Broadcast(FGameplayAbilityTargetDataHandle());
+	Destroy();
 }
 
 bool AGameplayAbilityTargetActor::IsNetRelevantFor(class APlayerController* RealViewer, AActor* Viewer, const FVector& SrcLocation)
diff --git a/Engine/Source/Runtime/GameplayAbilities/Private/Abilities/GameplayAbilityTargetActor_SingleLineTrace.cpp b/Engine/Source/Runtime/GameplayAbilities/Private/Abilities/GameplayAbilityTargetActor_SingleLineTrace.cpp
--- a/Engine/Source/Runtime/GameplayAbilities/Private/Abilities/GameplayAbilityTargetActor_SingleLineTrace.cpp
+++ b/Engine/Source/Runtime/GameplayAbilities/Private/Abilities/GameplayAbilityTargetActor_SingleLineTrace.cpp
@@ -74,32 +74,14 @@ void AGameplayAbilityTargetActor_SingleLineTrace::StartTargeting(UGameplayAbilit
 		UAbilitySystemComponent* ASC = Ability->GetCurrentActorInfo()->AbilitySystemComponent.Get();
 		if (ASC)
 		{
-			ASC->ConfirmCallbacks.AddDynamic(this, &AGameplayAbilityTargetActor_SingleLineTrace::Confirm);
-			ASC->CancelCallbacks.AddDynamic(this, &AGameplayAbilityTargetActor_SingleLineTrace::Cancel);
+			ASC->ConfirmCallbacks.AddDynamic(this, &AGameplayAbilityTargetActor_SingleLineTrace::ConfirmTargeting);
+			ASC->CancelCallbacks.AddDynamic(this, &AGameplayAbilityTargetActor_SingleLineTrace::CancelTargeting);
 		}
 	}
 	
 	bDebug = true;
 }
 
-void AGameplayAbilityTargetActor_SingleLineTrace::Confirm()
-{
-	if (Ability.IsValid())
-	{
-		bDebug = false;
-		FGameplayAbilityTargetDataHandle Handle = StaticGetTargetData(Ability->GetWorld(), Ability->GetCurrentActorInfo(), Ability->GetCurrentActivationInfo());
-		TargetDataReadyDelegate.Broadcast(Handle);
-	}
-
-	Destroy();
-}
-
-void AGameplayAbilityTargetActor_SingleLineTrace::Cancel()
-{
-	CanceledDelegate.Broadcast(FGameplayAbilityTargetDataHandle());
-	Destroy();
-}
-
 void AGameplayAbilityTargetActor_SingleLineTrace::Tick(float DeltaSeconds)
 {
 	Super::Tick(DeltaSeconds);
@@ -132,3 +114,4 @@ void AGameplayAbilityTargetActor_SingleLineTrace::ConfirmTargeting()
 	Destroy();
 }
 
+
diff --git a/Engine/Source/Runtime/GameplayAbilities/Public/Abilities/GameplayAbilityTargetActor.h b/Engine/Source/Runtime/GameplayAbilities/Public/Abilities/GameplayAbilityTargetActor.h
--- a/Engine/Source/Runtime/GameplayAbilities/Public/Abilities/GameplayAbilityTargetActor.h
+++ b/Engine/Source/Runtime/GameplayAbilities/Public/Abilities/GameplayAbilityTargetActor.h
@@ -22,9 +22,11 @@ public:
 	virtual void StartTargeting(UGameplayAbility* Ability);
 
 	/** Outside code is saying 'stop and just give me what you have' */
+	UFUNCTION()
 	virtual void ConfirmTargeting();
 
 	/** Outside code is saying 'stop everything and just forget about it' */
+	UFUNCTION()
 	virtual void CancelTargeting();
 
 	/** Replicated target data was received from a client. Possibly sanitize/verify. return true if data is good and we should broadcast it as valid data. */
diff --git a/Engine/Source/Runtime/GameplayAbilities/Public/Abilities/GameplayAbilityTargetActor_SingleLineTrace.h b/Engine/Source/Runtime/GameplayAbilities/Public/Abilities/GameplayAbilityTargetActor_SingleLineTrace.h
--- a/Engine/Source/Runtime/GameplayAbilities/Public/Abilities/GameplayAbilityTargetActor_SingleLineTrace.h
+++ b/Engine/Source/Runtime/GameplayAbilities/Public/Abilities/GameplayAbilityTargetActor_SingleLineTrace.h
@@ -14,14 +14,8 @@ public:
 	virtual FGameplayAbilityTargetDataHandle StaticGetTargetData(UWorld * World, const FGameplayAbilityActorInfo* ActorInfo, FGameplayAbilityActivationInfo ActivationInfo) const override;
 
 	virtual void StartTargeting(UGameplayAbility* Ability);
-
-	UFUNCTION()
-	void Confirm();
-
-	UFUNCTION()
-	void Cancel();
-
-	virtual void ConfirmTargeting();	
+	
+	virtual void ConfirmTargeting() override;	
 
 	virtual void Tick(float DeltaSeconds) override;
 	
