Target fixes
-Don't register confim/cancel callbacks on server with replicated target actors (only client should bind to these delegates and then tell server about target data)
-Register target data callbacks on server for replicated target actors

diff --git a/Engine/Source/Runtime/GameplayAbilities/Private/Abilities/GameplayAbilityTargetActor_SingleLineTrace.cpp b/Engine/Source/Runtime/GameplayAbilities/Private/Abilities/GameplayAbilityTargetActor_SingleLineTrace.cpp
--- a/Engine/Source/Runtime/GameplayAbilities/Private/Abilities/GameplayAbilityTargetActor_SingleLineTrace.cpp
+++ b/Engine/Source/Runtime/GameplayAbilities/Private/Abilities/GameplayAbilityTargetActor_SingleLineTrace.cpp
@@ -69,7 +69,7 @@ void AGameplayAbilityTargetActor_SingleLineTrace::StartTargeting(UGameplayAbilit
 	SourceActor = InAbility->GetCurrentActorInfo()->Actor.Get();
 
 	// We can bind directly to our ASC's confirm/cancel events, or wait to be told from an outside source to confirm or cancel
-	if (bBindToConfirmCancelInputs)
+	if (bBindToConfirmCancelInputs && (MasterPC && MasterPC->IsLocalController()))
 	{
 		UAbilitySystemComponent* ASC = Ability->GetCurrentActorInfo()->AbilitySystemComponent.Get();
 		if (ASC)
diff --git a/Engine/Source/Runtime/GameplayAbilities/Private/Abilities/Tasks/AbilityTask_WaitTargetData.cpp b/Engine/Source/Runtime/GameplayAbilities/Private/Abilities/Tasks/AbilityTask_WaitTargetData.cpp
--- a/Engine/Source/Runtime/GameplayAbilities/Private/Abilities/Tasks/AbilityTask_WaitTargetData.cpp
+++ b/Engine/Source/Runtime/GameplayAbilities/Private/Abilities/Tasks/AbilityTask_WaitTargetData.cpp
@@ -33,15 +33,18 @@ bool UAbilityTask_WaitTargetData::BeginSpawningActor(UObject* WorldContextObject
 
 		bool Replicates = CDO->GetReplicates();
 		bool StaticFunc = CDO->StaticTargetFunction;
+		bool IsLocallyControlled = MyAbility->GetCurrentActorInfo()->IsLocallyControlled();
 
 		if (Replicates && StaticFunc)
 		{
+			// We can't replicate a staticFunc target actor, since we are just calling a static function and not spawning an actor at all!
 			ABILITY_LOG(Fatal, TEXT("AbilityTargetActor class %s can't be Replicating and Static"), *TargetClass->GetName());
 			Replicates = false;
 		}
 
-		// Spawn the actor if this is a locally controlled ability (always) or if this is a replicated targeting mode.
-		if (Replicates || MyAbility->GetCurrentActorInfo()->IsLocallyControlled())
+		// Spawn the actor if this is a locally controlled ability (always) or if this is a replicating targeting mode.
+		// (E.g., server will spawn this target actor to replicate to all non owning clients)
+		if (Replicates || IsLocallyControlled)
 		{
 			if (StaticFunc)
 			{
@@ -70,11 +73,12 @@ bool UAbilityTask_WaitTargetData::BeginSpawningActor(UObject* WorldContextObject
 				}
 			}
 		}
-		else
-		{
-			// If not locally controlled (server for remote client), see if TargetData was already sent
-			// else register callback for when it does get here
+		
 
+		// If not locally controlled (server for remote client), see if TargetData was already sent
+		// else register callback for when it does get here.
+		if (!IsLocallyControlled)
+		{
 			if (AbilitySystemComponent->ReplicatedTargetData.IsValid())
 			{
 				ValidData.Broadcast(AbilitySystemComponent->ReplicatedTargetData);
