Minor cleanup

diff --git a/Engine/Source/Runtime/GameplayAbilities/Private/Abilities/GameplayAbilityTargetActor_SingleLineTrace.cpp b/Engine/Source/Runtime/GameplayAbilities/Private/Abilities/GameplayAbilityTargetActor_SingleLineTrace.cpp
--- a/Engine/Source/Runtime/GameplayAbilities/Private/Abilities/GameplayAbilityTargetActor_SingleLineTrace.cpp
+++ b/Engine/Source/Runtime/GameplayAbilities/Private/Abilities/GameplayAbilityTargetActor_SingleLineTrace.cpp
@@ -57,10 +57,8 @@ FGameplayAbilityTargetDataHandle AGameplayAbilityTargetActor_SingleLineTrace::St
 {
 	AActor* StaticSourceActor = ActorInfo->Actor.Get();
 	check(StaticSourceActor);
-
-	FGameplayAbilityTargetData_SingleTargetHit* ReturnData = new FGameplayAbilityTargetData_SingleTargetHit();
-	ReturnData->HitResult = PerformTrace(StaticSourceActor);
-	return FGameplayAbilityTargetDataHandle(ReturnData);
+	
+	return MakeTargetData(PerformTrace(StaticSourceActor));
 }
 
 void AGameplayAbilityTargetActor_SingleLineTrace::StartTargeting(UGameplayAbility* InAbility)
@@ -85,13 +83,11 @@ void AGameplayAbilityTargetActor_SingleLineTrace::StartTargeting(UGameplayAbilit
 void AGameplayAbilityTargetActor_SingleLineTrace::Tick(float DeltaSeconds)
 {
 	Super::Tick(DeltaSeconds);
-	FGameplayAbilityTargetDataHandle Handle;
-	FHitResult HitResult;
 
 	// very temp - do a mostly hardcoded trace from the source actor
 	if (SourceActor)
 	{
-		HitResult = PerformTrace(SourceActor);
+		FHitResult HitResult = PerformTrace(SourceActor);
 		
 		if (bDebug)
 		{
@@ -104,14 +100,20 @@ void AGameplayAbilityTargetActor_SingleLineTrace::Tick(float DeltaSeconds)
 
 void AGameplayAbilityTargetActor_SingleLineTrace::ConfirmTargeting()
 {
-	if (Ability.IsValid())
+	if (SourceActor)
 	{
 		bDebug = false;
-		FGameplayAbilityTargetDataHandle Handle = StaticGetTargetData(Ability->GetWorld(), Ability->GetCurrentActorInfo(), Ability->GetCurrentActivationInfo());
+		FGameplayAbilityTargetDataHandle Handle = MakeTargetData(PerformTrace(SourceActor));
 		TargetDataReadyDelegate.Broadcast(Handle);
 	}
 
 	Destroy();
 }
 
-
+FGameplayAbilityTargetDataHandle AGameplayAbilityTargetActor_SingleLineTrace::MakeTargetData(FHitResult HitResult) const
+{
+	/** Note this is cleaned up by the FGameplayAbilityTargetDataHandle (via an internal TSharedPtr) */
+	FGameplayAbilityTargetData_SingleTargetHit* ReturnData = new FGameplayAbilityTargetData_SingleTargetHit();
+	ReturnData->HitResult = HitResult;
+	return FGameplayAbilityTargetDataHandle(ReturnData);
+}
\ No newline at end of file
diff --git a/Engine/Source/Runtime/GameplayAbilities/Public/Abilities/GameplayAbilityTargetActor_SingleLineTrace.h b/Engine/Source/Runtime/GameplayAbilities/Public/Abilities/GameplayAbilityTargetActor_SingleLineTrace.h
--- a/Engine/Source/Runtime/GameplayAbilities/Public/Abilities/GameplayAbilityTargetActor_SingleLineTrace.h
+++ b/Engine/Source/Runtime/GameplayAbilities/Public/Abilities/GameplayAbilityTargetActor_SingleLineTrace.h
@@ -33,4 +33,6 @@ public:
 
 protected:
 	FHitResult PerformTrace(AActor *SourceActor) const;
+
+	FGameplayAbilityTargetDataHandle MakeTargetData(FHitResult HitResult) const;
 };
\ No newline at end of file
