Small updates
-Reconcile end points in traces that hit vs dont hit
-minor cleanup

diff --git a/Engine/Source/Runtime/GameplayAbilities/Private/Abilities/GameplayAbility.cpp b/Engine/Source/Runtime/GameplayAbilities/Private/Abilities/GameplayAbility.cpp
--- a/Engine/Source/Runtime/GameplayAbilities/Private/Abilities/GameplayAbility.cpp
+++ b/Engine/Source/Runtime/GameplayAbilities/Private/Abilities/GameplayAbility.cpp
@@ -553,7 +553,7 @@ void FGameplayAbilityActivationInfo::SetActivationConfirmed()
 
 void FGameplayAbilityTargetData::ApplyGameplayEffect(UGameplayEffect* GameplayEffect, const FGameplayAbilityActorInfo InstigatorInfo)
 {
-	// Improve relationship between InstigatorContext and 
+	// TODO: Improve relationship between InstigatorContext and FGameplayAbilityTargetData/FHitResult (or use something different between HitResult)
 	
 
 	FGameplayEffectSpec	SpecToApply(GameplayEffect,					// The UGameplayEffect data asset
diff --git a/Engine/Source/Runtime/GameplayAbilities/Private/Abilities/GameplayAbilityTargetActor_SingleLineTrace.cpp b/Engine/Source/Runtime/GameplayAbilities/Private/Abilities/GameplayAbilityTargetActor_SingleLineTrace.cpp
--- a/Engine/Source/Runtime/GameplayAbilities/Private/Abilities/GameplayAbilityTargetActor_SingleLineTrace.cpp
+++ b/Engine/Source/Runtime/GameplayAbilities/Private/Abilities/GameplayAbilityTargetActor_SingleLineTrace.cpp
@@ -88,13 +88,15 @@ void AGameplayAbilityTargetActor_SingleLineTrace::Tick(float DeltaSeconds)
 	if (SourceActor)
 	{
 		FHitResult HitResult = PerformTrace(SourceActor);
+
+		FVector EndPoint = HitResult.Component.IsValid() ? HitResult.ImpactPoint : HitResult.TraceEnd;
 		
 		if (bDebug)
 		{
-			DrawDebugLine(GetWorld(), SourceActor->GetActorLocation(), HitResult.ImpactPoint, FLinearColor::Green, false);
-			DrawDebugSphere(GetWorld(), HitResult.Location, 16, 10, FLinearColor::Green, false);
+			DrawDebugLine(GetWorld(), SourceActor->GetActorLocation(), EndPoint, FLinearColor::Green, false);
+			DrawDebugSphere(GetWorld(), EndPoint, 16, 10, FLinearColor::Green, false);
 		}
-		SetActorLocationAndRotation(HitResult.Location, SourceActor->GetActorRotation());
+		SetActorLocationAndRotation(EndPoint, SourceActor->GetActorRotation());
 	}
 }
 
diff --git a/Engine/Source/Runtime/GameplayAbilities/Private/AbilitySystemBlueprintLibrary.cpp b/Engine/Source/Runtime/GameplayAbilities/Private/AbilitySystemBlueprintLibrary.cpp
--- a/Engine/Source/Runtime/GameplayAbilities/Private/AbilitySystemBlueprintLibrary.cpp
+++ b/Engine/Source/Runtime/GameplayAbilities/Private/AbilitySystemBlueprintLibrary.cpp
@@ -35,7 +35,7 @@ void UAbilitySystemBlueprintLibrary::ApplyGameplayEffectToTargetData(FGameplayAb
 FGameplayAbilityTargetDataHandle UAbilitySystemBlueprintLibrary::AbilityTargetDataFromHitResult(FHitResult HitResult)
 {
 	// Construct TargetData
-	FGameplayAbilityTargetData_SingleTargetHit * TargetData = new FGameplayAbilityTargetData_SingleTargetHit(HitResult);
+	FGameplayAbilityTargetData_SingleTargetHit* TargetData = new FGameplayAbilityTargetData_SingleTargetHit(HitResult);
 
 	// Give it a handle and return
 	FGameplayAbilityTargetDataHandle	Handle;
@@ -46,10 +46,10 @@ FGameplayAbilityTargetDataHandle UAbilitySystemBlueprintLibrary::AbilityTargetDa
 
 FHitResult UAbilitySystemBlueprintLibrary::GetHitResultFromTargetData(FGameplayAbilityTargetDataHandle TargetData)
 {
-	FGameplayAbilityTargetData * Data = TargetData.Data.Get();
+	FGameplayAbilityTargetData* Data = TargetData.Data.Get();
 	if (Data)
 	{
-		const FHitResult * HitResultPtr = Data->GetHitResult();
+		const FHitResult* HitResultPtr = Data->GetHitResult();
 		if (HitResultPtr)
 		{
 			return *HitResultPtr;
@@ -59,6 +59,28 @@ FHitResult UAbilitySystemBlueprintLibrary::GetHitResultFromTargetData(FGameplayA
 	return FHitResult();
 }
 
+FVector UAbilitySystemBlueprintLibrary::GetTargetDataEndPoint(FGameplayAbilityTargetDataHandle TargetData)
+{
+	FGameplayAbilityTargetData * Data = TargetData.Data.Get();
+	if (Data)
+	{
+		const FHitResult* HitResultPtr = Data->GetHitResult();
+		if (HitResultPtr)
+		{
+			if (HitResultPtr->Component.IsValid())
+			{
+				return HitResultPtr->ImpactPoint;
+			}
+			else
+			{
+				return HitResultPtr->TraceEnd;
+			}
+		}
+	}
+
+	return FVector();
+}
+
 // -------------------------------------------------------------------------------------
 
 
diff --git a/Engine/Source/Runtime/GameplayAbilities/Public/AbilitySystemBlueprintLibrary.h b/Engine/Source/Runtime/GameplayAbilities/Public/AbilitySystemBlueprintLibrary.h
--- a/Engine/Source/Runtime/GameplayAbilities/Public/AbilitySystemBlueprintLibrary.h
+++ b/Engine/Source/Runtime/GameplayAbilities/Public/AbilitySystemBlueprintLibrary.h
@@ -25,15 +25,18 @@ class UAbilitySystemBlueprintLibrary : public UBlueprintFunctionLibrary
 	//		TargetData
 	// -------------------------------------------------------------------------------
 
-	UFUNCTION(BlueprintCallable, Category="Ability|TargetData")
+	UFUNCTION(BlueprintCallable, Category = "Ability|TargetData")
 	static void ApplyGameplayEffectToTargetData(FGameplayAbilityTargetDataHandle Target, UGameplayEffect *GameplayEffect, const FGameplayAbilityActorInfo InstigatorInfo);
 
-	UFUNCTION(BlueprintCallable, Category="Ability|TargetData")
+	UFUNCTION(BlueprintPure, Category = "Ability|TargetData")
 	static FGameplayAbilityTargetDataHandle	AbilityTargetDataFromHitResult(FHitResult HitResult);
 
-	UFUNCTION(BlueprintCallable, Category="Ability|TargetData")
+	UFUNCTION(BlueprintPure, Category = "Ability|TargetData")
 	static FHitResult GetHitResultFromTargetData(FGameplayAbilityTargetDataHandle HitResult);
 
+	UFUNCTION(BlueprintPure, Category = "Ability|TargetData")
+	static FVector GetTargetDataEndPoint(FGameplayAbilityTargetDataHandle TargetData);
+
 	// -------------------------------------------------------------------------------
 	//		GameplayCue
 	// -------------------------------------------------------------------------------
