Cleanup - mostly categorization and comments

diff --git a/Engine/Source/Runtime/GameplayAbilities/Public/Abilities/Tasks/AbilityTask_MoveToLocation.h b/Engine/Source/Runtime/GameplayAbilities/Public/Abilities/Tasks/AbilityTask_MoveToLocation.h
--- a/Engine/Source/Runtime/GameplayAbilities/Public/Abilities/Tasks/AbilityTask_MoveToLocation.h
+++ b/Engine/Source/Runtime/GameplayAbilities/Public/Abilities/Tasks/AbilityTask_MoveToLocation.h
@@ -23,8 +23,8 @@ class UAbilityTask_MoveToLocation : public UAbilityTask
 	void InterpolatePosition();
 
 	/** Move to the specified location, using the curve (range 0 - 1) or linearly if no curve is specified */
-	UFUNCTION(BlueprintCallable, Category = Abilities, meta = (HidePin = "WorldContextObject", DefaultToSelf = "WorldContextObject", BlueprintInternalUseOnly = "TRUE"))
-		static UAbilityTask_MoveToLocation* MoveToLocation(UObject* WorldContextObject, FVector Location, float Duration, UCurveFloat* OptionalInterpolationCurve);
+	UFUNCTION(BlueprintCallable, Category = "Ability|Tasks", meta = (HidePin = "WorldContextObject", DefaultToSelf = "WorldContextObject", BlueprintInternalUseOnly = "TRUE"))
+	static UAbilityTask_MoveToLocation* MoveToLocation(UObject* WorldContextObject, FVector Location, float Duration, UCurveFloat* OptionalInterpolationCurve);
 
 protected:
 	FVector StartLocation;
diff --git a/Engine/Source/Runtime/GameplayAbilities/Public/Abilities/Tasks/AbilityTask_PlayMontageAndWait.h b/Engine/Source/Runtime/GameplayAbilities/Public/Abilities/Tasks/AbilityTask_PlayMontageAndWait.h
--- a/Engine/Source/Runtime/GameplayAbilities/Public/Abilities/Tasks/AbilityTask_PlayMontageAndWait.h
+++ b/Engine/Source/Runtime/GameplayAbilities/Public/Abilities/Tasks/AbilityTask_PlayMontageAndWait.h
@@ -20,7 +20,7 @@ class UAbilityTask_PlayMontageAndWait : public UAbilityTask
 	void OnMontageEnded(UAnimMontage* Montage, bool bInterrupted);
 
 
-	UFUNCTION(BlueprintCallable, Category=Abilities, meta = (FriendlyName="PlayMontageAndWait", HidePin = "WorldContextObject", DefaultToSelf = "WorldContextObject", BlueprintInternalUseOnly = "TRUE"))
+	UFUNCTION(BlueprintCallable, Category="Ability|Tasks", meta = (FriendlyName="PlayMontageAndWait", HidePin = "WorldContextObject", DefaultToSelf = "WorldContextObject", BlueprintInternalUseOnly = "TRUE"))
 	static UAbilityTask_PlayMontageAndWait* CreatePlayMontageAndWaitProxy(UObject* WorldContextObject, UAnimMontage *MontageToPlay);
 
 	virtual void Activate() override;
diff --git a/Engine/Source/Runtime/GameplayAbilities/Public/Abilities/Tasks/AbilityTask_SpawnActor.h b/Engine/Source/Runtime/GameplayAbilities/Public/Abilities/Tasks/AbilityTask_SpawnActor.h
--- a/Engine/Source/Runtime/GameplayAbilities/Public/Abilities/Tasks/AbilityTask_SpawnActor.h
+++ b/Engine/Source/Runtime/GameplayAbilities/Public/Abilities/Tasks/AbilityTask_SpawnActor.h
@@ -29,7 +29,7 @@ class UAbilityTask_SpawnActor: public UAbilityTask
 	FSpawnActorDelegate	Success;
 	
 	/** Spawn new Actor on the network authority (server) */
-	UFUNCTION(BlueprintCallable, meta=(HidePin = "WorldContextObject", DefaultToSelf = "WorldContextObject", BlueprintInternalUseOnly = "true"), Category = "Abilities")
+	UFUNCTION(BlueprintCallable, meta=(HidePin = "WorldContextObject", DefaultToSelf = "WorldContextObject", BlueprintInternalUseOnly = "true"), Category="Ability|Tasks")
 	static UAbilityTask_SpawnActor* SpawnActor(UObject* WorldContextObject, TSubclassOf<AActor> Class);
 
 	UFUNCTION(BlueprintCallable, meta = (HidePin = "WorldContextObject", DefaultToSelf = "WorldContextObject", BlueprintInternalUseOnly = "true"), Category = "Abilities")
diff --git a/Engine/Source/Runtime/GameplayAbilities/Public/Abilities/Tasks/AbilityTask_WaitAbilityActivate.h b/Engine/Source/Runtime/GameplayAbilities/Public/Abilities/Tasks/AbilityTask_WaitAbilityActivate.h
--- a/Engine/Source/Runtime/GameplayAbilities/Public/Abilities/Tasks/AbilityTask_WaitAbilityActivate.h
+++ b/Engine/Source/Runtime/GameplayAbilities/Public/Abilities/Tasks/AbilityTask_WaitAbilityActivate.h
@@ -26,7 +26,7 @@ class UAbilityTask_WaitAbilityActivate : public UAbilityTask
 	void OnAbilityActivate(UGameplayAbility *ActivatedAbility);
 
 	/** Wait until a new ability (of the same or different type) is activated. */
-	UFUNCTION(BlueprintCallable, Category=Abilities, meta = (HidePin = "WorldContextObject", DefaultToSelf = "WorldContextObject", BlueprintInternalUseOnly = "TRUE"))
+	UFUNCTION(BlueprintCallable, Category="Ability|Tasks", meta = (HidePin = "WorldContextObject", DefaultToSelf = "WorldContextObject", BlueprintInternalUseOnly = "TRUE"))
 	static UAbilityTask_WaitAbilityActivate* WaitForAbilityActivate(UObject* WorldContextObject, FGameplayTag WithTag, FGameplayTag WithoutTag);
 
 	FGameplayTag WithTag;
diff --git a/Engine/Source/Runtime/GameplayAbilities/Public/Abilities/Tasks/AbilityTask_WaitAbilityCommit.h b/Engine/Source/Runtime/GameplayAbilities/Public/Abilities/Tasks/AbilityTask_WaitAbilityCommit.h
--- a/Engine/Source/Runtime/GameplayAbilities/Public/Abilities/Tasks/AbilityTask_WaitAbilityCommit.h
+++ b/Engine/Source/Runtime/GameplayAbilities/Public/Abilities/Tasks/AbilityTask_WaitAbilityCommit.h
@@ -28,7 +28,7 @@ class UAbilityTask_WaitAbilityCommit : public UAbilityTask
 	void OnAbilityCommit(UGameplayAbility *ActivatedAbility);
 
 	/** Wait until a new ability (of the same or different type) is commited. Used to gracefully interrupt abilities in specific ways. */
-	UFUNCTION(BlueprintCallable, Category=Abilities, meta = (HidePin = "WorldContextObject", DefaultToSelf = "WorldContextObject", BlueprintInternalUseOnly = "TRUE", FriendlyName = "Wait For New Ability Commit"))
+	UFUNCTION(BlueprintCallable, Category="Ability|Tasks", meta = (HidePin = "WorldContextObject", DefaultToSelf = "WorldContextObject", BlueprintInternalUseOnly = "TRUE", FriendlyName = "Wait For New Ability Commit"))
 	static UAbilityTask_WaitAbilityCommit* WaitForAbilityCommit(UObject* WorldContextObject, FGameplayTag WithTag, FGameplayTag WithoutTage);	
 
 	FGameplayTag WithTag;
diff --git a/Engine/Source/Runtime/GameplayAbilities/Public/Abilities/Tasks/AbilityTask_WaitAttributeChange.h b/Engine/Source/Runtime/GameplayAbilities/Public/Abilities/Tasks/AbilityTask_WaitAttributeChange.h
--- a/Engine/Source/Runtime/GameplayAbilities/Public/Abilities/Tasks/AbilityTask_WaitAttributeChange.h
+++ b/Engine/Source/Runtime/GameplayAbilities/Public/Abilities/Tasks/AbilityTask_WaitAttributeChange.h
@@ -26,7 +26,7 @@ class UAbilityTask_WaitAttributeChange : public UAbilityTask
 	void OnAttributeChange(float NewValue, const FGameplayEffectModCallbackData*);
 
 	/** Wait until an attribute changes. */
-	UFUNCTION(BlueprintCallable, Category=Abilities, meta = (HidePin = "WorldContextObject", DefaultToSelf = "WorldContextObject", BlueprintInternalUseOnly = "TRUE"))
+	UFUNCTION(BlueprintCallable, Category="Ability|Tasks", meta = (HidePin = "WorldContextObject", DefaultToSelf = "WorldContextObject", BlueprintInternalUseOnly = "TRUE"))
 	static UAbilityTask_WaitAttributeChange* WaitForAttributeChange(UObject* WorldContextObject, FGameplayAttribute Attribute, FGameplayTag WithTag, FGameplayTag WithoutTag);
 
 	FGameplayTag WithTag;
diff --git a/Engine/Source/Runtime/GameplayAbilities/Public/Abilities/Tasks/AbilityTask_WaitConfirm.h b/Engine/Source/Runtime/GameplayAbilities/Public/Abilities/Tasks/AbilityTask_WaitConfirm.h
--- a/Engine/Source/Runtime/GameplayAbilities/Public/Abilities/Tasks/AbilityTask_WaitConfirm.h
+++ b/Engine/Source/Runtime/GameplayAbilities/Public/Abilities/Tasks/AbilityTask_WaitConfirm.h
@@ -18,7 +18,7 @@ class UAbilityTask_WaitConfirm : public UAbilityTask
 	virtual void Activate() override;
 
 	/** Wait until the server confirms the use of this ability. This is used to gate predictive portions of the ability */
-	UFUNCTION(BlueprintCallable, Category=Abilities, meta = (HidePin = "WorldContextObject", DefaultToSelf = "WorldContextObject", BlueprintInternalUseOnly = "TRUE"))
+	UFUNCTION(BlueprintCallable, Category="Ability|Tasks", meta = (HidePin = "WorldContextObject", DefaultToSelf = "WorldContextObject", BlueprintInternalUseOnly = "TRUE"))
 	static UAbilityTask_WaitConfirm* WaitConfirm(UObject* WorldContextObject);
 
 public:
diff --git a/Engine/Source/Runtime/GameplayAbilities/Public/Abilities/Tasks/AbilityTask_WaitConfirmCancel.h b/Engine/Source/Runtime/GameplayAbilities/Public/Abilities/Tasks/AbilityTask_WaitConfirmCancel.h
--- a/Engine/Source/Runtime/GameplayAbilities/Public/Abilities/Tasks/AbilityTask_WaitConfirmCancel.h
+++ b/Engine/Source/Runtime/GameplayAbilities/Public/Abilities/Tasks/AbilityTask_WaitConfirmCancel.h
@@ -24,7 +24,7 @@ class UAbilityTask_WaitConfirmCancel : public UAbilityTask
 	UFUNCTION()
 	void OnCancelCallback();
 
-	UFUNCTION(BlueprintCallable, meta=(HidePin = "WorldContextObject", DefaultToSelf = "WorldContextObject", BlueprintInternalUseOnly = "true", FriendlyName="Wait for Confirm Input"), Category="Abilities")
+	UFUNCTION(BlueprintCallable, meta=(HidePin = "WorldContextObject", DefaultToSelf = "WorldContextObject", BlueprintInternalUseOnly = "true", FriendlyName="Wait for Confirm Input"), Category="Ability|Tasks")
 	static UAbilityTask_WaitConfirmCancel* WaitConfirmCancel(UObject* WorldContextObject);	
 
 	virtual void Activate() override;	
diff --git a/Engine/Source/Runtime/GameplayAbilities/Public/Abilities/Tasks/AbilityTask_WaitGameplayEffectRemoved.h b/Engine/Source/Runtime/GameplayAbilities/Public/Abilities/Tasks/AbilityTask_WaitGameplayEffectRemoved.h
--- a/Engine/Source/Runtime/GameplayAbilities/Public/Abilities/Tasks/AbilityTask_WaitGameplayEffectRemoved.h
+++ b/Engine/Source/Runtime/GameplayAbilities/Public/Abilities/Tasks/AbilityTask_WaitGameplayEffectRemoved.h
@@ -28,7 +28,7 @@ class UAbilityTask_WaitGameplayEffectRemoved : public UAbilityTask
 	void OnGameplayEffectRemoved();
 
 	/** Wait until an overlap occurs. This will need to be better fleshed out so we can specify game specific collision requirements */
-	UFUNCTION(BlueprintCallable, Category=Abilities, meta = (HidePin = "WorldContextObject", DefaultToSelf = "WorldContextObject", BlueprintInternalUseOnly = "TRUE"))
+	UFUNCTION(BlueprintCallable, Category="Ability|Tasks", meta = (HidePin = "WorldContextObject", DefaultToSelf = "WorldContextObject", BlueprintInternalUseOnly = "TRUE"))
 	static UAbilityTask_WaitGameplayEffectRemoved* WaitForGameplayEffectRemoved(UObject* WorldContextObject, FActiveGameplayEffectHandle Handle);
 
 	FActiveGameplayEffectHandle Handle;
diff --git a/Engine/Source/Runtime/GameplayAbilities/Public/Abilities/Tasks/AbilityTask_WaitMovementModeChange.h b/Engine/Source/Runtime/GameplayAbilities/Public/Abilities/Tasks/AbilityTask_WaitMovementModeChange.h
--- a/Engine/Source/Runtime/GameplayAbilities/Public/Abilities/Tasks/AbilityTask_WaitMovementModeChange.h
+++ b/Engine/Source/Runtime/GameplayAbilities/Public/Abilities/Tasks/AbilityTask_WaitMovementModeChange.h
@@ -20,7 +20,7 @@ public:
 
 	EMovementMode	RequiredMode;
 
-	UFUNCTION(BlueprintCallable, Category=Abilities, meta = (HidePin = "WorldContextObject", DefaultToSelf = "WorldContextObject", BlueprintInternalUseOnly = "TRUE"))
+	UFUNCTION(BlueprintCallable, Category="Ability|Tasks", meta = (HidePin = "WorldContextObject", DefaultToSelf = "WorldContextObject", BlueprintInternalUseOnly = "TRUE"))
 	static UAbilityTask_WaitMovementModeChange* CreateWaitMovementModeChange(UObject* WorldContextObject, EMovementMode NewMode);
 
 	virtual void Activate() override;
diff --git a/Engine/Source/Runtime/GameplayAbilities/Public/Abilities/Tasks/AbilityTask_WaitOverlap.h b/Engine/Source/Runtime/GameplayAbilities/Public/Abilities/Tasks/AbilityTask_WaitOverlap.h
--- a/Engine/Source/Runtime/GameplayAbilities/Public/Abilities/Tasks/AbilityTask_WaitOverlap.h
+++ b/Engine/Source/Runtime/GameplayAbilities/Public/Abilities/Tasks/AbilityTask_WaitOverlap.h
@@ -31,7 +31,7 @@ class UAbilityTask_WaitOverlap : public UAbilityTask
 	virtual void Activate() override;
 
 	/** Wait until an overlap occurs. This will need to be better fleshed out so we can specify game specific collision requirements */
-	UFUNCTION(BlueprintCallable, Category=Abilities, meta = (HidePin = "WorldContextObject", DefaultToSelf = "WorldContextObject", BlueprintInternalUseOnly = "TRUE"))
+	UFUNCTION(BlueprintCallable, Category="Ability|Tasks", meta = (HidePin = "WorldContextObject", DefaultToSelf = "WorldContextObject", BlueprintInternalUseOnly = "TRUE"))
 	static UAbilityTask_WaitOverlap* WaitForOverlap(UObject* WorldContextObject);
 
 public:
diff --git a/Engine/Source/Runtime/GameplayAbilities/Public/Abilities/Tasks/AbilityTask_WaitTargetData.h b/Engine/Source/Runtime/GameplayAbilities/Public/Abilities/Tasks/AbilityTask_WaitTargetData.h
--- a/Engine/Source/Runtime/GameplayAbilities/Public/Abilities/Tasks/AbilityTask_WaitTargetData.h
+++ b/Engine/Source/Runtime/GameplayAbilities/Public/Abilities/Tasks/AbilityTask_WaitTargetData.h
@@ -32,7 +32,7 @@ class UAbilityTask_WaitTargetData: public UAbilityTask
 	void OnTargetDataCancelledCallback(FGameplayAbilityTargetDataHandle Data);
 
 	/** Spawns Targeting actor and waits for it to return valid data or to be cancelled. */
-	UFUNCTION(BlueprintCallable, meta=(HidePin = "WorldContextObject", DefaultToSelf = "WorldContextObject", BlueprintInternalUseOnly = "true", HideSpawnParms="Instigator"), Category = "Abilities")
+	UFUNCTION(BlueprintCallable, meta=(HidePin = "WorldContextObject", DefaultToSelf = "WorldContextObject", BlueprintInternalUseOnly = "true", HideSpawnParms="Instigator"), Category="Ability|Tasks")
 	static UAbilityTask_WaitTargetData* WaitTargetData(UObject* WorldContextObject, TSubclassOf<AGameplayAbilityTargetActor> Class);
 
 	UFUNCTION(BlueprintCallable, meta = (HidePin = "WorldContextObject", DefaultToSelf = "WorldContextObject", BlueprintInternalUseOnly = "true"), Category = "Abilities")
