Friendly name

diff --git a/Engine/Source/Runtime/GameplayAbilities/Public/Abilities/Tasks/AbilityTask_WaitMovementModeChange.h b/Engine/Source/Runtime/GameplayAbilities/Public/Abilities/Tasks/AbilityTask_WaitMovementModeChange.h
--- a/Engine/Source/Runtime/GameplayAbilities/Public/Abilities/Tasks/AbilityTask_WaitMovementModeChange.h
+++ b/Engine/Source/Runtime/GameplayAbilities/Public/Abilities/Tasks/AbilityTask_WaitMovementModeChange.h
@@ -20,7 +20,8 @@ public:
 
 	EMovementMode	RequiredMode;
 
-	UFUNCTION(BlueprintCallable, Category="Ability|Tasks", meta = (HidePin = "WorldContextObject", DefaultToSelf = "WorldContextObject", BlueprintInternalUseOnly = "TRUE"))
+	/** Wait until movement mode changes (E.g., landing) */
+	UFUNCTION(BlueprintCallable, Category="Ability|Tasks", meta = (HidePin = "WorldContextObject", DefaultToSelf = "WorldContextObject", BlueprintInternalUseOnly = "TRUE", FriendlyName="WaitMovementModeChange"))
 	static UAbilityTask_WaitMovementModeChange* CreateWaitMovementModeChange(UObject* WorldContextObject, EMovementMode NewMode);
 
 	virtual void Activate() override;
