Allow GameplayEffectTests to build a GameplayTag table in code.

This removes the need to include entries for gameplay effect automation testing in the tag tables used by various projects that use gameplay effects.

diff --git a/Engine/Source/Runtime/GameplayAbilities/Private/GameplayEffectTests.cpp b/Engine/Source/Runtime/GameplayAbilities/Private/GameplayEffectTests.cpp
--- a/Engine/Source/Runtime/GameplayAbilities/Private/GameplayEffectTests.cpp
+++ b/Engine/Source/Runtime/GameplayAbilities/Private/GameplayEffectTests.cpp
@@ -2378,6 +2378,46 @@ void ClearGlobalDataTable()
 	IGameplayAbilitiesModule::Get().GetAbilitySystemGlobals()->AutomationTestOnly_SetGlobalAttributeDataTable(NULL);
 }
 
+UDataTable* CreateGameplayDataTable()
+{
+	FString CSV(TEXT(",Tag,CategoryText,"));
+	CSV.Append(TEXT("\r\n0,Damage"));
+	CSV.Append(TEXT("\r\n1,Damage.Basic"));
+	CSV.Append(TEXT("\r\n2,Damage.Type1"));
+	CSV.Append(TEXT("\r\n3,Damage.Type2"));
+	CSV.Append(TEXT("\r\n4,Damage.Reduce"));
+	CSV.Append(TEXT("\r\n5,Damage.Buffable"));
+	CSV.Append(TEXT("\r\n6,Damage.Buff"));
+	CSV.Append(TEXT("\r\n7,Damage.Physical"));
+	CSV.Append(TEXT("\r\n8,Damage.Fire"));
+	CSV.Append(TEXT("\r\n9,Damage.Buffed.FireBuff"));
+	CSV.Append(TEXT("\r\n10,Damage.Mitigated.Armor"));
+	CSV.Append(TEXT("\r\n11,Lifesteal"));
+	CSV.Append(TEXT("\r\n12,Shield"));
+	CSV.Append(TEXT("\r\n13,Buff"));
+	CSV.Append(TEXT("\r\n14,Immune"));
+	CSV.Append(TEXT("\r\n15,FireDamage"));
+	CSV.Append(TEXT("\r\n16,ShieldAbsorb"));
+	CSV.Append(TEXT("\r\n17,Stackable"));
+	CSV.Append(TEXT("\r\n18,Stack"));
+	CSV.Append(TEXT("\r\n19,Stack.CappedNumber"));
+	CSV.Append(TEXT("\r\n20,Stack.DiminishingReturns"));
+	CSV.Append(TEXT("\r\n21,Protect.Damage"));
+	CSV.Append(TEXT("\r\n22,SpellDmg.Buff"));
+	CSV.Append(TEXT("\r\n23,GameplayCue.Burning"));
+
+	UDataTable * DataTable = Cast<UDataTable>(StaticConstructObject(UDataTable::StaticClass(), GetTransientPackage(), FName(TEXT("TempDataTable"))));
+	DataTable->RowStruct = FGameplayTagTableRow::StaticStruct();
+	DataTable->CreateTableFromCSVString(CSV);
+
+	FGameplayTagTableRow * Row = (FGameplayTagTableRow*)DataTable->RowMap["0"];
+	if (Row)
+	{
+		check(Row->Tag == TEXT("Damage"));
+	}
+	return DataTable;
+}
+
 bool GameplayEffectsTest_InstantDamage_ScalingExplicit(UWorld *World, FAutomationTestBase * Test)
 {
 	// This example uses explicit scaling in a GameplayEffect. We explicitly specify the curve table to use in the GameplayEffect
@@ -6007,6 +6047,11 @@ bool FGameplayEffectsTest::RunTest( const FString& Parameters )
 	UCurveTable *CurveTable = IGameplayAbilitiesModule::Get().GetAbilitySystemGlobals()->GetGlobalCurveTable();
 	UDataTable *DataTable = IGameplayAbilitiesModule::Get().GetAbilitySystemGlobals()->GetGlobalAttributeDataTable();
 
+	// setup required GameplayTags
+	UDataTable* TagTable = CreateGameplayDataTable();
+
+	IGameplayTagsModule::Get().GetGameplayTagsManager().PopulateTreeFromDataTable(TagTable);
+
 	UWorld *World = UWorld::CreateWorld(EWorldType::Game, false);
 	FWorldContext &WorldContext = GEngine->CreateNewWorldContext(EWorldType::Game);
 	WorldContext.SetCurrentWorld(World);
diff --git a/Engine/Source/Runtime/GameplayTags/Classes/GameplayTagsManager.h b/Engine/Source/Runtime/GameplayTags/Classes/GameplayTagsManager.h
--- a/Engine/Source/Runtime/GameplayTags/Classes/GameplayTagsManager.h
+++ b/Engine/Source/Runtime/GameplayTags/Classes/GameplayTagsManager.h
@@ -233,6 +233,7 @@ class GAMEPLAYTAGS_API UGameplayTagsManager : public UObject
 private:
 
 	friend class FGameplayTagTest;
+	friend class FGameplayEffectsTest;
 
 	/** Helper function to populate the tag tree from each table */
 	void PopulateTreeFromDataTable(class UDataTable* Table);
