Make CreateGameplayDataTable functions static in automation tests.

Fixes monolithic builds.

diff --git a/Engine/Source/Runtime/GameplayAbilities/Private/GameplayEffectTests.cpp b/Engine/Source/Runtime/GameplayAbilities/Private/GameplayEffectTests.cpp
--- a/Engine/Source/Runtime/GameplayAbilities/Private/GameplayEffectTests.cpp
+++ b/Engine/Source/Runtime/GameplayAbilities/Private/GameplayEffectTests.cpp
@@ -2378,7 +2378,7 @@ void ClearGlobalDataTable()
 	IGameplayAbilitiesModule::Get().GetAbilitySystemGlobals()->AutomationTestOnly_SetGlobalAttributeDataTable(NULL);
 }
 
-UDataTable* CreateGameplayDataTable()
+static UDataTable* CreateGameplayDataTable()
 {
 	FString CSV(TEXT(",Tag,CategoryText,"));
 	CSV.Append(TEXT("\r\n0,Damage"));
diff --git a/Engine/Source/Runtime/GameplayTags/Private/Tests/GameplayTagTests.cpp b/Engine/Source/Runtime/GameplayTags/Private/Tests/GameplayTagTests.cpp
--- a/Engine/Source/Runtime/GameplayTags/Private/Tests/GameplayTagTests.cpp
+++ b/Engine/Source/Runtime/GameplayTags/Private/Tests/GameplayTagTests.cpp
@@ -4,7 +4,7 @@
 
 IMPLEMENT_SIMPLE_AUTOMATION_TEST(FGameplayTagTest, "GameplayTags.GameplayTag", EAutomationTestFlags::ATF_Editor)
 
-UDataTable* CreateGameplayDataTable()
+static UDataTable* CreateGameplayDataTable()
 {
 	FString CSV(TEXT(",Tag,CategoryText,\r\n0,GameplayTagTest.Test1\r\n1,GameplayTagTest.Test2"));
 
