Add checks to prevent crashes when the modifier level isn't set on a gameplay effect.

diff --git a/Engine/Source/Runtime/GameplayAbilities/Private/AbilitySystemComponent.cpp b/Engine/Source/Runtime/GameplayAbilities/Private/AbilitySystemComponent.cpp
--- a/Engine/Source/Runtime/GameplayAbilities/Private/AbilitySystemComponent.cpp
+++ b/Engine/Source/Runtime/GameplayAbilities/Private/AbilitySystemComponent.cpp
@@ -443,7 +443,7 @@ void UAbilitySystemComponent::InvokeGameplayCueEvent(const FGameplayEffectSpec &
 	}
 
 	// FIXME: Replication of level not finished
-	float ExecuteLevel = Spec.ModifierLevel.Get()->IsValid() ? Spec.ModifierLevel.Get()->GetLevel() : 1.f;
+	float ExecuteLevel =  (Spec.ModifierLevel.IsValid() && Spec.ModifierLevel.Get()->IsValid()) ? Spec.ModifierLevel.Get()->GetLevel() : 1.f;
 
 	FGameplayCueParameters CueParameters;
 	CueParameters.InstigatorContext = Spec.InstigatorContext;
