Basic dagger-aiming code built in Orion. Structure for two-actor scheme (targeting separate from visualization) set up.

diff --git a/Engine/Source/Runtime/GameplayAbilities/Public/Abilities/GameplayAbilityTargetActor_SingleLineTrace.h b/Engine/Source/Runtime/GameplayAbilities/Public/Abilities/GameplayAbilityTargetActor_SingleLineTrace.h
--- a/Engine/Source/Runtime/GameplayAbilities/Public/Abilities/GameplayAbilityTargetActor_SingleLineTrace.h
+++ b/Engine/Source/Runtime/GameplayAbilities/Public/Abilities/GameplayAbilityTargetActor_SingleLineTrace.h
@@ -28,11 +28,14 @@ public:
 	UPROPERTY(BlueprintReadOnly, Replicated, Category = Targeting)
 	AActor* SourceActor;		//Maybe add a socket to this?
 
+	UPROPERTY(BlueprintReadOnly, Category = Targeting)
+	APlayerController* SourcePC;
+
 	UPROPERTY(BlueprintReadOnly, meta=(ExposeOnSpawn=true), Category=Projectile)
 	bool bBindToConfirmCancelInputs;
 
 protected:
-	FHitResult PerformTrace(AActor *SourceActor) const;
+	virtual FHitResult PerformTrace(AActor *SourceActor) const;
 
 	FGameplayAbilityTargetDataHandle MakeTargetData(FHitResult HitResult) const;
 };
\ No newline at end of file
