Fixing some submission mistakes.

diff --git a/Engine/Source/Runtime/GameplayAbilities/Private/Abilities/GameplayAbilityWorldReticle.cpp b/Engine/Source/Runtime/GameplayAbilities/Private/Abilities/GameplayAbilityWorldReticle.cpp
new file mode 100644
--- /dev/null
+++ b/Engine/Source/Runtime/GameplayAbilities/Private/Abilities/GameplayAbilityWorldReticle.cpp
@@ -0,0 +1,41 @@
+// Copyright 1998-2013 Epic Games, Inc. All Rights Reserved.
+
+#include "AbilitySystemPrivatePCH.h"
+#include "GameplayAbilityWorldReticle.h"
+
+// --------------------------------------------------------------------------------------------------------------------------------------------------------
+//
+//	AGameplayAbilityWorldReticle
+//
+// --------------------------------------------------------------------------------------------------------------------------------------------------------
+
+AGameplayAbilityWorldReticle::AGameplayAbilityWorldReticle(const class FPostConstructInitializeProperties& PCIP)
+	: Super(PCIP)
+{
+	PrimaryActorTick.bCanEverTick = true;
+	PrimaryActorTick.TickGroup = TG_PrePhysics;
+}
+
+void AGameplayAbilityWorldReticle::InitializeReticle(AGameplayAbilityTargetActor* InTargetingActor)
+{
+	check(InTargetingActor);
+	TargetingActor = InTargetingActor;
+	MasterPC = InTargetingActor->MasterPC;
+	SetTickPrerequisite(TargetingActor);		//We want the reticle to tick after the targeting actor so that designers have the final say on the position
+}
+
+bool AGameplayAbilityWorldReticle::IsNetRelevantFor(class APlayerController* RealViewer, AActor* Viewer, const FVector& SrcLocation)
+{
+	//The player who created the ability doesn't need to be updated about it - there should be local prediction in place.
+	if (RealViewer == MasterPC)
+	{
+		return false;
+	}
+
+	return Super::IsNetRelevantFor(RealViewer, Viewer, SrcLocation);
+}
+
+bool AGameplayAbilityWorldReticle::GetReplicates() const
+{
+	return bReplicates;
+}
diff --git a/Engine/Source/Runtime/GameplayAbilities/Public/Abilities/GameplayAbilityWorldReticle.h b/Engine/Source/Runtime/GameplayAbilities/Public/Abilities/GameplayAbilityWorldReticle.h
new file mode 100644
--- /dev/null
+++ b/Engine/Source/Runtime/GameplayAbilities/Public/Abilities/GameplayAbilityWorldReticle.h
@@ -0,0 +1,31 @@
+// Copyright 1998-2014 Epic Games, Inc. All Rights Reserved.
+#pragma once
+
+#include "GameplayAbilityWorldReticle.generated.h"
+
+
+UCLASS(Blueprintable)
+class GAMEPLAYABILITIES_API AGameplayAbilityWorldReticle : public AActor
+{
+	GENERATED_UCLASS_BODY()
+
+public:
+
+	/** Accessor for checking, before instantiating, if this WorldReticle will replicate. */
+	bool GetReplicates() const;
+
+	// ------------------------------
+
+	virtual bool IsNetRelevantFor(class APlayerController* RealViewer, AActor* Viewer, const FVector& SrcLocation) override;
+
+	void InitializeReticle(AGameplayAbilityTargetActor* InTargetingActor);
+
+protected:
+	/** This is used in the process of determining whether we should replicate to a specific client. */
+	UPROPERTY(BlueprintReadOnly, Category = "Network")
+	APlayerController* MasterPC;
+
+	/** In the future, we may want to grab things like sockets off of this. */
+	UPROPERTY(BlueprintReadOnly, Category = "Network")
+	AGameplayAbilityTargetActor* TargetingActor;
+};
\ No newline at end of file
