Removing deprecated function to stop GUBP complaint from last night.

diff --git a/Engine/Source/Runtime/GameplayAbilities/Private/Abilities/GameplayAbilityWorldReticle.cpp b/Engine/Source/Runtime/GameplayAbilities/Private/Abilities/GameplayAbilityWorldReticle.cpp
--- a/Engine/Source/Runtime/GameplayAbilities/Private/Abilities/GameplayAbilityWorldReticle.cpp
+++ b/Engine/Source/Runtime/GameplayAbilities/Private/Abilities/GameplayAbilityWorldReticle.cpp
@@ -21,7 +21,7 @@ void AGameplayAbilityWorldReticle::InitializeReticle(AGameplayAbilityTargetActor
 	check(InTargetingActor);
 	TargetingActor = InTargetingActor;
 	MasterPC = InTargetingActor->MasterPC;
-	SetTickPrerequisite(TargetingActor);		//We want the reticle to tick after the targeting actor so that designers have the final say on the position
+	AddTickPrerequisiteActor(TargetingActor);		//We want the reticle to tick after the targeting actor so that designers have the final say on the position
 }
 
 bool AGameplayAbilityWorldReticle::IsNetRelevantFor(class APlayerController* RealViewer, AActor* Viewer, const FVector& SrcLocation)
