MoveToLocation task marks itself for pending kill when it completes.

diff --git a/Engine/Source/Runtime/GameplayAbilities/Private/Abilities/Tasks/AbilityTask_MoveToLocation.cpp b/Engine/Source/Runtime/GameplayAbilities/Private/Abilities/Tasks/AbilityTask_MoveToLocation.cpp
--- a/Engine/Source/Runtime/GameplayAbilities/Private/Abilities/Tasks/AbilityTask_MoveToLocation.cpp
+++ b/Engine/Source/Runtime/GameplayAbilities/Private/Abilities/Tasks/AbilityTask_MoveToLocation.cpp
@@ -21,6 +21,7 @@ void UAbilityTask_MoveToLocation::InterpolatePosition()
 				ActorOwner->GetWorld()->GetTimerManager().ClearTimer(this, &UAbilityTask_MoveToLocation::InterpolatePosition);
 				ActorOwner->SetActorLocation(TargetLocation);
 				OnTargetLocationReached.Broadcast();
+				MarkPendingKill();
 			}
 			else
 			{
