Attempt to fix a GUBP compile error.

diff --git a/Engine/Source/Runtime/GameplayAbilities/Public/Abilities/GameplayAbilityWorldReticle.h b/Engine/Source/Runtime/GameplayAbilities/Public/Abilities/GameplayAbilityWorldReticle.h
--- a/Engine/Source/Runtime/GameplayAbilities/Public/Abilities/GameplayAbilityWorldReticle.h
+++ b/Engine/Source/Runtime/GameplayAbilities/Public/Abilities/GameplayAbilityWorldReticle.h
@@ -1,9 +1,9 @@
 // Copyright 1998-2014 Epic Games, Inc. All Rights Reserved.
 #pragma once
 
+#include "GameplayAbilityTargetActor.h"
 #include "GameplayAbilityWorldReticle.generated.h"
 
-
 UCLASS(Blueprintable)
 class GAMEPLAYABILITIES_API AGameplayAbilityWorldReticle : public AActor
 {
