Fixed non-unity build compile error about using an unknown type

diff --git a/Engine/Source/Runtime/GameplayAbilities/Private/Abilities/GameplayAbilityWorldReticle.cpp b/Engine/Source/Runtime/GameplayAbilities/Private/Abilities/GameplayAbilityWorldReticle.cpp
--- a/Engine/Source/Runtime/GameplayAbilities/Private/Abilities/GameplayAbilityWorldReticle.cpp
+++ b/Engine/Source/Runtime/GameplayAbilities/Private/Abilities/GameplayAbilityWorldReticle.cpp
@@ -2,6 +2,7 @@
 
 #include "AbilitySystemPrivatePCH.h"
 #include "GameplayAbilityWorldReticle.h"
+#include "GameplayAbilityTargetActor.h"
 
 // --------------------------------------------------------------------------------------------------------------------------------------------------------
 //
