Added range to Single Line Trace target actor. Moved reticle up to here. Removed unused PlayerController field.

diff --git a/Engine/Source/Runtime/GameplayAbilities/Private/Abilities/GameplayAbilityTargetActor_SingleLineTrace.cpp b/Engine/Source/Runtime/GameplayAbilities/Private/Abilities/GameplayAbilityTargetActor_SingleLineTrace.cpp
--- a/Engine/Source/Runtime/GameplayAbilities/Private/Abilities/GameplayAbilityTargetActor_SingleLineTrace.cpp
+++ b/Engine/Source/Runtime/GameplayAbilities/Private/Abilities/GameplayAbilityTargetActor_SingleLineTrace.cpp
@@ -44,7 +44,7 @@ FHitResult AGameplayAbilityTargetActor_SingleLineTrace::PerformTrace(AActor *InS
 	Params.AddIgnoredActors(ActorsToIgnore);
 
 	FVector TraceStart = InSourceActor->GetActorLocation();
-	FVector TraceEnd = TraceStart + (InSourceActor->GetActorForwardVector() * 3000.f);
+	FVector TraceEnd = TraceStart + (InSourceActor->GetActorForwardVector() * MaxRange);
 
 	// ------------------------------------------------------
 
diff --git a/Engine/Source/Runtime/GameplayAbilities/Public/Abilities/GameplayAbilityTargetActor_SingleLineTrace.h b/Engine/Source/Runtime/GameplayAbilities/Public/Abilities/GameplayAbilityTargetActor_SingleLineTrace.h
--- a/Engine/Source/Runtime/GameplayAbilities/Public/Abilities/GameplayAbilityTargetActor_SingleLineTrace.h
+++ b/Engine/Source/Runtime/GameplayAbilities/Public/Abilities/GameplayAbilityTargetActor_SingleLineTrace.h
@@ -3,6 +3,7 @@
 #pragma once
 
 #include "GameplayAbilityTargetActor.h"
+#include "GameplayAbilityWorldReticle.h"
 #include "GameplayAbilityTargetActor_SingleLineTrace.generated.h"
 
 UCLASS(Blueprintable)
@@ -28,14 +29,16 @@ public:
 	UPROPERTY(BlueprintReadOnly, Replicated, Category = Targeting)
 	AActor* SourceActor;		//Maybe add a socket to this?
 
-	UPROPERTY(BlueprintReadOnly, Category = Targeting)
-	APlayerController* SourcePC;
-
 	UPROPERTY(BlueprintReadOnly, meta=(ExposeOnSpawn=true), Category=Projectile)
 	bool bBindToConfirmCancelInputs;
 
+	UPROPERTY(BlueprintReadOnly, meta = (ExposeOnSpawn = true), Category = Projectile)
+	float MaxRange;
+
 protected:
 	virtual FHitResult PerformTrace(AActor *SourceActor) const;
 
 	FGameplayAbilityTargetDataHandle MakeTargetData(FHitResult HitResult) const;
+
+	TWeakObjectPtr<AGameplayAbilityWorldReticle> ReticleActor;
 };
\ No newline at end of file
