Reduced .h file usage, should not be needed.

diff --git a/Engine/Source/Runtime/GameplayAbilities/Public/Abilities/GameplayAbilityWorldReticle.h b/Engine/Source/Runtime/GameplayAbilities/Public/Abilities/GameplayAbilityWorldReticle.h
--- a/Engine/Source/Runtime/GameplayAbilities/Public/Abilities/GameplayAbilityWorldReticle.h
+++ b/Engine/Source/Runtime/GameplayAbilities/Public/Abilities/GameplayAbilityWorldReticle.h
@@ -1,9 +1,10 @@
 // Copyright 1998-2014 Epic Games, Inc. All Rights Reserved.
 #pragma once
 
-#include "GameplayAbilityTargetActor.h"
 #include "GameplayAbilityWorldReticle.generated.h"
 
+class AGameplayAbilityTargetActor;
+
 UCLASS(Blueprintable)
 class GAMEPLAYABILITIES_API AGameplayAbilityWorldReticle : public AActor
 {
