Moving some assets around, part 1/2.

diff --git a/Engine/Source/Runtime/GameplayAbilities/Public/Abilities/GameplayAbilityTargetActor_SingleLineTrace.h b/Engine/Source/Runtime/GameplayAbilities/Public/Abilities/GameplayAbilityTargetActor_SingleLineTrace.h
--- a/Engine/Source/Runtime/GameplayAbilities/Public/Abilities/GameplayAbilityTargetActor_SingleLineTrace.h
+++ b/Engine/Source/Runtime/GameplayAbilities/Public/Abilities/GameplayAbilityTargetActor_SingleLineTrace.h
@@ -32,7 +32,7 @@ public:
 	UPROPERTY(BlueprintReadOnly, meta=(ExposeOnSpawn=true), Category=Projectile)
 	bool bBindToConfirmCancelInputs;
 
-	UPROPERTY(BlueprintReadOnly, meta = (ExposeOnSpawn = true), Category = Projectile)
+	UPROPERTY(BlueprintReadOnly, EditDefaultsOnly, meta=(ExposeOnSpawn=true), Category=Projectile)
 	float MaxRange;
 
 protected:
