Add temp workaround for passing reticle class in ExposeOnSpawn property

diff --git a/Engine/Source/Runtime/GameplayAbilities/Public/Abilities/GameplayAbilityWorldReticle.h b/Engine/Source/Runtime/GameplayAbilities/Public/Abilities/GameplayAbilityWorldReticle.h
--- a/Engine/Source/Runtime/GameplayAbilities/Public/Abilities/GameplayAbilityWorldReticle.h
+++ b/Engine/Source/Runtime/GameplayAbilities/Public/Abilities/GameplayAbilityWorldReticle.h
@@ -29,4 +29,15 @@ protected:
 	/** In the future, we may want to grab things like sockets off of this. */
 	UPROPERTY(BlueprintReadOnly, Category = "Network")
 	AGameplayAbilityTargetActor* TargetingActor;
+};
+
+
+/** Temporary work around for allowing TSubClassOf property to be used in ExposeOnSpawn situation */
+USTRUCT(BlueprintType)
+struct GAMEPLAYABILITIES_API FGameplayAbilityWorldReticleClass
+{
+	GENERATED_USTRUCT_BODY()
+
+	UPROPERTY(BlueprintReadWrite,EditAnywhere,Category="Class")
+	TSubclassOf<AGameplayAbilityWorldReticle>	Class;
 };
\ No newline at end of file
