Fix case where server would cleanup targeting tasks twice. Added checks to catch future problems earlier.

diff --git a/Engine/Source/Runtime/GameplayAbilities/Private/Abilities/GameplayAbilityTargetActor.cpp b/Engine/Source/Runtime/GameplayAbilities/Private/Abilities/GameplayAbilityTargetActor.cpp
--- a/Engine/Source/Runtime/GameplayAbilities/Private/Abilities/GameplayAbilityTargetActor.cpp
+++ b/Engine/Source/Runtime/GameplayAbilities/Private/Abilities/GameplayAbilityTargetActor.cpp
@@ -21,6 +21,8 @@ void AGameplayAbilityTargetActor::StartTargeting(UGameplayAbility* Ability)
 
 void AGameplayAbilityTargetActor::ConfirmTargeting()
 {
+	check(ShouldProduceTargetData());
+
 	TargetDataReadyDelegate.Broadcast(FGameplayAbilityTargetDataHandle());
 	Destroy();
 }
@@ -57,3 +59,8 @@ bool AGameplayAbilityTargetActor::OnReplicatedTargetDataReceived(FGameplayAbilit
 {
 	return true;
 }
+
+bool AGameplayAbilityTargetActor::ShouldProduceTargetData() const
+{
+	return (MasterPC && MasterPC->IsLocalController());
+}
diff --git a/Engine/Source/Runtime/GameplayAbilities/Private/Abilities/GameplayAbilityTargetActor_SingleLineTrace.cpp b/Engine/Source/Runtime/GameplayAbilities/Private/Abilities/GameplayAbilityTargetActor_SingleLineTrace.cpp
--- a/Engine/Source/Runtime/GameplayAbilities/Private/Abilities/GameplayAbilityTargetActor_SingleLineTrace.cpp
+++ b/Engine/Source/Runtime/GameplayAbilities/Private/Abilities/GameplayAbilityTargetActor_SingleLineTrace.cpp
@@ -102,6 +102,8 @@ void AGameplayAbilityTargetActor_SingleLineTrace::Tick(float DeltaSeconds)
 
 void AGameplayAbilityTargetActor_SingleLineTrace::ConfirmTargeting()
 {
+	check(ShouldProduceTargetData());
+
 	if (SourceActor)
 	{
 		bDebug = false;
diff --git a/Engine/Source/Runtime/GameplayAbilities/Private/Abilities/Tasks/AbilityTask_WaitTargetData.cpp b/Engine/Source/Runtime/GameplayAbilities/Private/Abilities/Tasks/AbilityTask_WaitTargetData.cpp
--- a/Engine/Source/Runtime/GameplayAbilities/Private/Abilities/Tasks/AbilityTask_WaitTargetData.cpp
+++ b/Engine/Source/Runtime/GameplayAbilities/Private/Abilities/Tasks/AbilityTask_WaitTargetData.cpp
@@ -156,7 +156,14 @@ void UAbilityTask_WaitTargetData::OnTargetDataReadyCallback(FGameplayAbilityTarg
 {
 	check(AbilitySystemComponent.IsValid());
 
-	AbilitySystemComponent->ServerSetReplicatedTargetData(Data);
+	// Send TargetData to the server IFF we are the client or this is a locally controlled player on the server 
+	// If server confirmed a client's target data locally, it doesn't need to go through the replicated target data code path
+	const FGameplayAbilityActorInfo* Info = Ability->GetCurrentActorInfo();
+	if (!Info->IsNetAuthority() || Info->IsLocallyControlled())
+	{
+		AbilitySystemComponent->ServerSetReplicatedTargetData(Data);
+	}
+
 	ValidData.Broadcast(Data);
 	
 	Cleanup();
@@ -175,6 +182,12 @@ void UAbilityTask_WaitTargetData::OnTargetDataCancelledCallback(FGameplayAbility
 
 void UAbilityTask_WaitTargetData::Cleanup()
 {
+	if (HasAnyFlags(RF_PendingKill))
+	{
+		ensure(false);
+		return;
+	}
+
 	AbilitySystemComponent->ConsumeAbilityTargetData();
 	AbilitySystemComponent->SetUserAbilityActivationInhibited(false);
 
diff --git a/Engine/Source/Runtime/GameplayAbilities/Public/Abilities/GameplayAbilityTargetActor.h b/Engine/Source/Runtime/GameplayAbilities/Public/Abilities/GameplayAbilityTargetActor.h
--- a/Engine/Source/Runtime/GameplayAbilities/Public/Abilities/GameplayAbilityTargetActor.h
+++ b/Engine/Source/Runtime/GameplayAbilities/Public/Abilities/GameplayAbilityTargetActor.h
@@ -29,6 +29,8 @@ public:
 	UFUNCTION()
 	virtual void CancelTargeting();
 
+	virtual bool ShouldProduceTargetData() const;
+
 	/** Replicated target data was received from a client. Possibly sanitize/verify. return true if data is good and we should broadcast it as valid data. */
 	virtual bool OnReplicatedTargetDataReceived(FGameplayAbilityTargetDataHandle& Data) const;
 
