Fix for SpawnActor tasks passing in class ExposeOnSpawn properties properly

diff --git a/Engine/Source/Editor/GameplayAbilitiesEditor/Private/K2Node_LatentAbilityCall.cpp b/Engine/Source/Editor/GameplayAbilitiesEditor/Private/K2Node_LatentAbilityCall.cpp
--- a/Engine/Source/Editor/GameplayAbilitiesEditor/Private/K2Node_LatentAbilityCall.cpp
+++ b/Engine/Source/Editor/GameplayAbilitiesEditor/Private/K2Node_LatentAbilityCall.cpp
@@ -525,7 +525,8 @@ void UK2Node_LatentAbilityCall::ExpandNode(class FKismetCompilerContext& Compile
 	
 	for (auto SpawnVarPin : SpawnParmPins)
 	{
-		if (SpawnVarPin->LinkedTo.Num() > 0 || SpawnVarPin->DefaultValue != FString())
+		const bool bHasDefaultValue = !SpawnVarPin->DefaultValue.IsEmpty() || !SpawnVarPin->DefaultTextValue.IsEmpty() || SpawnVarPin->DefaultObject;
+		if (SpawnVarPin->LinkedTo.Num() > 0 || bHasDefaultValue)
 		{
 			if (SpawnVarPin->LinkedTo.Num() == 0)
 			{
diff --git a/Engine/Source/Runtime/GameplayAbilities/Public/Abilities/GameplayAbilityWorldReticle.h b/Engine/Source/Runtime/GameplayAbilities/Public/Abilities/GameplayAbilityWorldReticle.h
--- a/Engine/Source/Runtime/GameplayAbilities/Public/Abilities/GameplayAbilityWorldReticle.h
+++ b/Engine/Source/Runtime/GameplayAbilities/Public/Abilities/GameplayAbilityWorldReticle.h
@@ -30,14 +30,3 @@ protected:
 	UPROPERTY(BlueprintReadOnly, Category = "Network")
 	AGameplayAbilityTargetActor* TargetingActor;
 };
-
-
-/** Temporary work around for allowing TSubClassOf property to be used in ExposeOnSpawn situation */
-USTRUCT(BlueprintType)
-struct GAMEPLAYABILITIES_API FGameplayAbilityWorldReticleClass
-{
-	GENERATED_USTRUCT_BODY()
-
-	UPROPERTY(BlueprintReadWrite,EditAnywhere,Category="Class")
-	TSubclassOf<AGameplayAbilityWorldReticle>	Class;
-};
\ No newline at end of file
