Folded code back into SingleLineTrace, fixed up Missile Swarm assets.

diff --git a/Engine/Source/Runtime/GameplayAbilities/Private/Abilities/GameplayAbilityTargetActor_SingleLineTrace.cpp b/Engine/Source/Runtime/GameplayAbilities/Private/Abilities/GameplayAbilityTargetActor_SingleLineTrace.cpp
--- a/Engine/Source/Runtime/GameplayAbilities/Private/Abilities/GameplayAbilityTargetActor_SingleLineTrace.cpp
+++ b/Engine/Source/Runtime/GameplayAbilities/Private/Abilities/GameplayAbilityTargetActor_SingleLineTrace.cpp
@@ -29,6 +29,16 @@ void AGameplayAbilityTargetActor_SingleLineTrace::GetLifetimeReplicatedProps(TAr
 	DOREPLIFETIME(AGameplayAbilityTargetActor_SingleLineTrace, SourceActor);
 }
 
+void AGameplayAbilityTargetActor_SingleLineTrace::EndPlay(const EEndPlayReason::Type EndPlayReason)
+{
+	if (ReticleActor.IsValid())
+	{
+		ReticleActor.Get()->Destroy();
+	}
+
+	Super::EndPlay(EndPlayReason);
+}
+
 
 FHitResult AGameplayAbilityTargetActor_SingleLineTrace::PerformTrace(AActor *InSourceActor) const
 {
@@ -43,13 +53,29 @@ FHitResult AGameplayAbilityTargetActor_SingleLineTrace::PerformTrace(AActor *InS
 	Params.bTraceAsyncScene = true;
 	Params.AddIgnoredActors(ActorsToIgnore);
 
+	FVector AimDirection = InSourceActor->GetActorForwardVector();		//Default
+	UGameplayAbility* MyAbility = Ability.Get();
+	if (MyAbility)		//Server and launching client only
+	{
+		APlayerController* AimingPC = MyAbility->GetCurrentActorInfo()->PlayerController.Get();
+		check(AimingPC);
+		FVector CamLoc;
+		FRotator CamRot;
+		AimingPC->GetPlayerViewPoint(CamLoc, CamRot);
+		AimDirection = CamRot.Vector();
+	}
+
 	FVector TraceStart = InSourceActor->GetActorLocation();
-	FVector TraceEnd = TraceStart + (InSourceActor->GetActorForwardVector() * MaxRange);
+	FVector TraceEnd = TraceStart + (AimDirection * 3000.f);
 
 	// ------------------------------------------------------
 
 	FHitResult ReturnHitResult;
 	InSourceActor->GetWorld()->LineTraceSingle(ReturnHitResult, TraceStart, TraceEnd, ECC_WorldStatic, Params);
+	if (AActor* LocalReticleActor = ReticleActor.Get())
+	{
+		LocalReticleActor->SetActorLocation(ReturnHitResult.Location);
+	}
 	return ReturnHitResult;
 }
 
@@ -78,25 +104,39 @@ void AGameplayAbilityTargetActor_SingleLineTrace::StartTargeting(UGameplayAbilit
 	}
 	
 	bDebug = true;
+
+	ReticleActor = GetWorld()->SpawnActor<AGameplayAbilityWorldReticle>(ReticleClass, GetActorLocation(), GetActorRotation());
+	if (AGameplayAbilityWorldReticle* CachedReticleActor = ReticleActor.Get())
+	{
+		CachedReticleActor->InitializeReticle(this);
+	}
 }
 
 void AGameplayAbilityTargetActor_SingleLineTrace::Tick(float DeltaSeconds)
 {
-	Super::Tick(DeltaSeconds);
-
 	// very temp - do a mostly hardcoded trace from the source actor
 	if (SourceActor)
 	{
 		FHitResult HitResult = PerformTrace(SourceActor);
-
 		FVector EndPoint = HitResult.Component.IsValid() ? HitResult.ImpactPoint : HitResult.TraceEnd;
-		
+
 		if (bDebug)
 		{
 			DrawDebugLine(GetWorld(), SourceActor->GetActorLocation(), EndPoint, FLinearColor::Green, false);
 			DrawDebugSphere(GetWorld(), EndPoint, 16, 10, FLinearColor::Green, false);
 		}
+
 		SetActorLocationAndRotation(EndPoint, SourceActor->GetActorRotation());
+
+		if (ShouldProduceTargetData() && bAutoFire)
+		{
+			if (TimeUntilAutoFire <= 0.0f)
+			{
+				ConfirmTargeting();
+				bAutoFire = false;
+			}
+			TimeUntilAutoFire -= DeltaSeconds;
+		}
 	}
 }
 
@@ -107,6 +147,7 @@ void AGameplayAbilityTargetActor_SingleLineTrace::ConfirmTargeting()
 	if (SourceActor)
 	{
 		bDebug = false;
+		bAutoFire = false;
 		FGameplayAbilityTargetDataHandle Handle = MakeTargetData(PerformTrace(SourceActor));
 		TargetDataReadyDelegate.Broadcast(Handle);
 	}
diff --git a/Engine/Source/Runtime/GameplayAbilities/Public/Abilities/GameplayAbilityTargetActor_SingleLineTrace.h b/Engine/Source/Runtime/GameplayAbilities/Public/Abilities/GameplayAbilityTargetActor_SingleLineTrace.h
--- a/Engine/Source/Runtime/GameplayAbilities/Public/Abilities/GameplayAbilityTargetActor_SingleLineTrace.h
+++ b/Engine/Source/Runtime/GameplayAbilities/Public/Abilities/GameplayAbilityTargetActor_SingleLineTrace.h
@@ -12,6 +12,8 @@ class GAMEPLAYABILITIES_API AGameplayAbilityTargetActor_SingleLineTrace : public
 	GENERATED_UCLASS_BODY()
 
 public:
+	virtual void EndPlay(const EEndPlayReason::Type EndPlayReason) override;
+
 	virtual FGameplayAbilityTargetDataHandle StaticGetTargetData(UWorld * World, const FGameplayAbilityActorInfo* ActorInfo, FGameplayAbilityActivationInfo ActivationInfo) const override;
 
 	virtual void StartTargeting(UGameplayAbility* Ability);
@@ -35,6 +37,15 @@ public:
 	UPROPERTY(BlueprintReadWrite, EditAnywhere, meta = (ExposeOnSpawn = true), Category = Projectile)
 	float MaxRange;
 
+	UPROPERTY(BlueprintReadWrite, EditAnywhere, meta = (ExposeOnSpawn = true), Category = Projectile)
+	bool bAutoFire;				//Automatically confirm after the timer runs out. Should also have support for an auto-cancel feature.
+
+	UPROPERTY(BlueprintReadWrite, EditAnywhere, meta = (ExposeOnSpawn = true), Category = Projectile)
+	float TimeUntilAutoFire;
+
+	UPROPERTY(BlueprintReadWrite, EditAnywhere, meta = (ExposeOnSpawn = true), Category = Projectile)
+	TSubclassOf<AGameplayAbilityWorldReticle> ReticleClass;		//Using a special class for replication purposes
+
 protected:
 	virtual FHitResult PerformTrace(AActor *SourceActor) const;
 
