Pass on ability task creation and cleanup

diff --git a/Engine/Source/Runtime/GameplayAbilities/Private/Abilities/GameplayAbility.cpp b/Engine/Source/Runtime/GameplayAbilities/Private/Abilities/GameplayAbility.cpp
--- a/Engine/Source/Runtime/GameplayAbilities/Private/Abilities/GameplayAbility.cpp
+++ b/Engine/Source/Runtime/GameplayAbilities/Private/Abilities/GameplayAbility.cpp
@@ -248,9 +248,18 @@ bool UGameplayAbility::TryActivateAbility(const FGameplayAbilityActorInfo* Actor
 
 void UGameplayAbility::EndAbility(const FGameplayAbilityActorInfo* ActorInfo)
 {
+	// Tell all our tasks that we are finished and they should cleanup
+	for (TWeakObjectPtr<UAbilityTask> Task : ActiveTasks)
+	{
+		if (Task.IsValid())
+		{
+			Task->AbilityEnded();
+		}
+	}
+	ActiveTasks.Reset();	// Empty the array but dont resize memory, since this object is probably going to be destroyed very soon anyways.
+
+	// Tell owning AbilitySystemComponent that we ended so it can do stuff (including MarkPendingKill us)
 	ActorInfo->AbilitySystemComponent->NotifyAbilityEnded(this);
-	
-	// TODO: generic way of releasing all callbacks?
 }
 
 void UGameplayAbility::ActivateAbility(const FGameplayAbilityActorInfo* ActorInfo, const FGameplayAbilityActivationInfo ActivationInfo)
@@ -465,6 +474,18 @@ void UGameplayAbility::ClientActivateAbilitySucceed_Internal(int32 PredictionKey
 
 //----------------------------------------------------------------------
 
+void UGameplayAbility::TaskStarted(UAbilityTask* NewTask)
+{
+	ActiveTasks.Add(NewTask);
+}
+
+void UGameplayAbility::TaskEnded(UAbilityTask* Task)
+{
+	ActiveTasks.Remove(Task);
+}
+
+//----------------------------------------------------------------------
+
 void FGameplayAbilityActorInfo::InitFromActor(AActor *InActor)
 {
 	check(InActor);
diff --git a/Engine/Source/Runtime/GameplayAbilities/Private/Abilities/Tasks/AbilityTask.cpp b/Engine/Source/Runtime/GameplayAbilities/Private/Abilities/Tasks/AbilityTask.cpp
--- a/Engine/Source/Runtime/GameplayAbilities/Private/Abilities/Tasks/AbilityTask.cpp
+++ b/Engine/Source/Runtime/GameplayAbilities/Private/Abilities/Tasks/AbilityTask.cpp
@@ -17,4 +17,58 @@ void UAbilityTask::InitTask(UGameplayAbility* InAbility)
 {
 	Ability = InAbility;
 	AbilitySystemComponent = InAbility->GetCurrentActorInfo()->AbilitySystemComponent;
+	InAbility->TaskStarted(this);
+}
+
+UWorld* UAbilityTask::GetWorld() const
+{
+	if (AbilitySystemComponent.IsValid())
+	{
+		return AbilitySystemComponent->GetWorld();
+	}
+
+	return nullptr;
+}
+
+AActor* UAbilityTask::GetActor() const
+{
+	if (Ability.IsValid())
+	{
+		const FGameplayAbilityActorInfo* Info = Ability->GetCurrentActorInfo();
+		if (Info)
+		{
+			return Info->Actor.Get();
+		}
+	}
+
+	return nullptr;	
+}
+
+void UAbilityTask::AbilityEnded()
+{
+	if (!IsPendingKill())
+	{
+		OnDestroy(true);
+	}
+}
+
+void UAbilityTask::EndTask()
+{
+	if (!IsPendingKill())
+	{
+		OnDestroy(false);
+	}
+}
+
+void UAbilityTask::OnDestroy(bool AbilityIsEnding)
+{
+	ensure(!IsPendingKill());
+
+	// Remove ourselves from the owning Ability's task list, if the ability isn't ending
+	if (!AbilityIsEnding && Ability.IsValid())
+	{
+		Ability->TaskEnded(this);
+	}
+
+	MarkPendingKill();
 }
\ No newline at end of file
diff --git a/Engine/Source/Runtime/GameplayAbilities/Private/Abilities/Tasks/AbilityTask_MoveToLocation.cpp b/Engine/Source/Runtime/GameplayAbilities/Private/Abilities/Tasks/AbilityTask_MoveToLocation.cpp
--- a/Engine/Source/Runtime/GameplayAbilities/Private/Abilities/Tasks/AbilityTask_MoveToLocation.cpp
+++ b/Engine/Source/Runtime/GameplayAbilities/Private/Abilities/Tasks/AbilityTask_MoveToLocation.cpp
@@ -2,6 +2,8 @@
 #include "AbilitySystemPrivatePCH.h"
 #include "Abilities/Tasks/AbilityTask_MoveToLocation.h"
 
+#pragma optimize( "", off )
+
 UAbilityTask_MoveToLocation::UAbilityTask_MoveToLocation(const class FPostConstructInitializeProperties& PCIP)
 : Super(PCIP)
 {
@@ -18,10 +20,9 @@ void UAbilityTask_MoveToLocation::InterpolatePosition()
 
 			if (CurrentTime >= TimeMoveWillEnd)
 			{
-				ActorOwner->GetWorld()->GetTimerManager().ClearTimer(this, &UAbilityTask_MoveToLocation::InterpolatePosition);
 				ActorOwner->SetActorLocation(TargetLocation);
 				OnTargetLocationReached.Broadcast();
-				MarkPendingKill();
+				EndTask();
 			}
 			else
 			{
@@ -38,22 +39,27 @@ void UAbilityTask_MoveToLocation::InterpolatePosition()
 
 UAbilityTask_MoveToLocation* UAbilityTask_MoveToLocation::MoveToLocation(class UObject* WorldContextObject, FVector Location, float Duration, UCurveFloat* OptionalInterpolationCurve)
 {
-	check(WorldContextObject);
-	UGameplayAbility* ThisAbility = CastChecked<UGameplayAbility>(WorldContextObject);
-	check(ThisAbility);
-	AActor* ActorOwner = Cast<AActor>(ThisAbility->GetOuter());
-	check(ActorOwner);
-
-	UAbilityTask_MoveToLocation* MyObj = NewObject<UAbilityTask_MoveToLocation>();
-	MyObj->Ability = ThisAbility;
+	auto MyObj = NewTask<UAbilityTask_MoveToLocation>(WorldContextObject);
 
-	MyObj->StartLocation = ActorOwner->GetActorLocation();
+	MyObj->StartLocation = MyObj->GetActor()->GetActorLocation();
 	MyObj->TargetLocation = Location;
 	MyObj->DurationOfMovement = FMath::Max(Duration, 0.001f);		//Avoid negative or divide-by-zero cases
-	MyObj->TimeMoveStarted = WorldContextObject->GetWorld()->GetTimeSeconds();
+	MyObj->TimeMoveStarted = MyObj->GetWorld()->GetTimeSeconds();
 	MyObj->TimeMoveWillEnd = MyObj->TimeMoveStarted + MyObj->DurationOfMovement;
 	MyObj->LerpCurve = OptionalInterpolationCurve;
-	WorldContextObject->GetWorld()->GetTimerManager().SetTimer(MyObj, &UAbilityTask_MoveToLocation::InterpolatePosition, 0.001f, true);
 
 	return MyObj;
 }
+
+void UAbilityTask_MoveToLocation::OnDestroy(bool AbilityEnded)
+{
+	GetWorld()->GetTimerManager().ClearTimer(this, &UAbilityTask_MoveToLocation::InterpolatePosition);
+
+	Super::OnDestroy(AbilityEnded);
+}
+
+void UAbilityTask_MoveToLocation::Activate()
+{
+	// FIXME: This is bad - its calling InterpolatePOsition @ 1000hz!!!!
+	GetWorld()->GetTimerManager().SetTimer(this, &UAbilityTask_MoveToLocation::InterpolatePosition, 0.001f, true);
+}
\ No newline at end of file
diff --git a/Engine/Source/Runtime/GameplayAbilities/Private/Abilities/Tasks/AbilityTask_PlayMontageAndWait.cpp b/Engine/Source/Runtime/GameplayAbilities/Private/Abilities/Tasks/AbilityTask_PlayMontageAndWait.cpp
--- a/Engine/Source/Runtime/GameplayAbilities/Private/Abilities/Tasks/AbilityTask_PlayMontageAndWait.cpp
+++ b/Engine/Source/Runtime/GameplayAbilities/Private/Abilities/Tasks/AbilityTask_PlayMontageAndWait.cpp
@@ -18,15 +18,13 @@ void UAbilityTask_PlayMontageAndWait::OnMontageEnded(UAnimMontage* Montage, bool
 	{
 		OnComplete.Broadcast();
 	}
+
+	EndTask();
 }
 
 UAbilityTask_PlayMontageAndWait* UAbilityTask_PlayMontageAndWait::CreatePlayMontageAndWaitProxy(UObject* WorldContextObject, UAnimMontage *MontageToPlay)
 {
-	check(WorldContextObject);
-
-	UAbilityTask_PlayMontageAndWait* MyObj = NewObject<UAbilityTask_PlayMontageAndWait>();
-	UGameplayAbility* ThisAbility = CastChecked<UGameplayAbility>(WorldContextObject);
-	MyObj->Ability = ThisAbility;
+	UAbilityTask_PlayMontageAndWait* MyObj = NewTask<UAbilityTask_PlayMontageAndWait>(WorldContextObject);
 	MyObj->MontageToPlay = MontageToPlay;
 
 	return MyObj;
@@ -36,14 +34,22 @@ void UAbilityTask_PlayMontageAndWait::Activate()
 {
 	if (Ability.IsValid())
 	{
-		AActor * ActorOwner = Cast<AActor>(Ability->GetOuter());
-				
-		TSharedPtr<FGameplayAbilityActorInfo> ActorInfo = TSharedPtr<FGameplayAbilityActorInfo>(UAbilitySystemGlobals::Get().AllocAbilityActorInfo(ActorOwner));
-		
-		FOnMontageEnded EndDelegate;
-		EndDelegate.BindUObject(this, &UAbilityTask_PlayMontageAndWait::OnMontageEnded);
-
-		ActorInfo->AnimInstance->Montage_Play(MontageToPlay, 1.f);
-		ActorInfo->AnimInstance->Montage_SetEndDelegate(EndDelegate);
+		const FGameplayAbilityActorInfo* ActorInfo = Ability->GetCurrentActorInfo();
+		if (ActorInfo->AnimInstance.IsValid())
+		{
+			FOnMontageEnded EndDelegate;
+			EndDelegate.BindUObject(this, &UAbilityTask_PlayMontageAndWait::OnMontageEnded);
+
+			ActorInfo->AnimInstance->Montage_Play(MontageToPlay, 1.f);
+			ActorInfo->AnimInstance->Montage_SetEndDelegate(EndDelegate);
+		}
 	}
+}
+
+void UAbilityTask_PlayMontageAndWait::OnDestroy(bool AbilityEnded)
+{
+	// Note: Clearing montage end delegate isn't necessary since its not a multicast and will be cleared when the next montage plays.
+	// (If we are destroyed, it will detect this and not do anything)
+
+	Super::OnDestroy(AbilityEnded);
 }
\ No newline at end of file
diff --git a/Engine/Source/Runtime/GameplayAbilities/Private/Abilities/Tasks/AbilityTask_SpawnActor.cpp b/Engine/Source/Runtime/GameplayAbilities/Private/Abilities/Tasks/AbilityTask_SpawnActor.cpp
--- a/Engine/Source/Runtime/GameplayAbilities/Private/Abilities/Tasks/AbilityTask_SpawnActor.cpp
+++ b/Engine/Source/Runtime/GameplayAbilities/Private/Abilities/Tasks/AbilityTask_SpawnActor.cpp
@@ -13,10 +13,7 @@ UAbilityTask_SpawnActor::UAbilityTask_SpawnActor(const class FPostConstructIniti
 
 UAbilityTask_SpawnActor* UAbilityTask_SpawnActor::SpawnActor(UObject* WorldContextObject, TSubclassOf<AActor> InClass)
 {
-	UAbilityTask_SpawnActor * MyObj = NewObject<UAbilityTask_SpawnActor>();
-	UGameplayAbility * ThisAbility = CastChecked<UGameplayAbility>(WorldContextObject);
-	MyObj->InitTask(ThisAbility);
-	return MyObj;	
+	return NewTask<UAbilityTask_SpawnActor>(WorldContextObject);
 }
 
 // ---------------------------------------------------------------------------------------
@@ -44,7 +41,7 @@ void UAbilityTask_SpawnActor::FinishSpawningActor(UObject* WorldContextObject, A
 		Success.Broadcast(SpawnedActor);
 	}
 
-	MarkPendingKill();
+	EndTask();
 }
 
 // ---------------------------------------------------------------------------------------
diff --git a/Engine/Source/Runtime/GameplayAbilities/Private/Abilities/Tasks/AbilityTask_WaitAbilityActivate.cpp b/Engine/Source/Runtime/GameplayAbilities/Private/Abilities/Tasks/AbilityTask_WaitAbilityActivate.cpp
--- a/Engine/Source/Runtime/GameplayAbilities/Private/Abilities/Tasks/AbilityTask_WaitAbilityActivate.cpp
+++ b/Engine/Source/Runtime/GameplayAbilities/Private/Abilities/Tasks/AbilityTask_WaitAbilityActivate.cpp
@@ -13,19 +13,11 @@ UAbilityTask_WaitAbilityActivate::UAbilityTask_WaitAbilityActivate(const class F
 
 UAbilityTask_WaitAbilityActivate* UAbilityTask_WaitAbilityActivate::WaitForAbilityActivate(UObject* WorldContextObject, FGameplayTag InWithTag, FGameplayTag InWithoutTag)
 {
-	check(WorldContextObject);
-	UGameplayAbility* ThisAbility = CastChecked<UGameplayAbility>(WorldContextObject);
-	if (ThisAbility)
-	{
-		UAbilityTask_WaitAbilityActivate * MyObj = NULL;
-		MyObj = NewObject<UAbilityTask_WaitAbilityActivate>();
-		MyObj->InitTask(ThisAbility);
-		MyObj->WithTag = InWithTag;
-		MyObj->WithoutTag = InWithoutTag;
+	auto MyObj = NewTask<UAbilityTask_WaitAbilityActivate>(WorldContextObject);
+	MyObj->WithTag = InWithTag;
+	MyObj->WithoutTag = InWithoutTag;
 
-		return MyObj;
-	}
-	return NULL;
+	return MyObj;
 }
 
 void UAbilityTask_WaitAbilityActivate::Activate()
@@ -44,11 +36,18 @@ void UAbilityTask_WaitAbilityActivate::OnAbilityActivate(UGameplayAbility *Activ
 		// Failed tag check
 		return;
 	}
-	
+
+	OnActivate.Broadcast(ActivatedAbility);
+
+	EndTask();
+}
+
+void UAbilityTask_WaitAbilityActivate::OnDestroy(bool AbilityEnded)
+{
 	if (AbilitySystemComponent.IsValid())
 	{
 		AbilitySystemComponent->AbilityActivatedCallbacks.RemoveUObject(this, &UAbilityTask_WaitAbilityActivate::OnAbilityActivate);
 	}
 
-	OnActivate.Broadcast(ActivatedAbility);
+	Super::OnDestroy(AbilityEnded);
 }
\ No newline at end of file
diff --git a/Engine/Source/Runtime/GameplayAbilities/Private/Abilities/Tasks/AbilityTask_WaitAbilityCommit.cpp b/Engine/Source/Runtime/GameplayAbilities/Private/Abilities/Tasks/AbilityTask_WaitAbilityCommit.cpp
--- a/Engine/Source/Runtime/GameplayAbilities/Private/Abilities/Tasks/AbilityTask_WaitAbilityCommit.cpp
+++ b/Engine/Source/Runtime/GameplayAbilities/Private/Abilities/Tasks/AbilityTask_WaitAbilityCommit.cpp
@@ -14,20 +14,11 @@ UAbilityTask_WaitAbilityCommit::UAbilityTask_WaitAbilityCommit(const class FPost
 
 UAbilityTask_WaitAbilityCommit* UAbilityTask_WaitAbilityCommit::WaitForAbilityCommit(UObject* WorldContextObject, FGameplayTag InWithTag, FGameplayTag InWithoutTag)
 {
-	check(WorldContextObject);
-	UGameplayAbility* ThisAbility = CastChecked<UGameplayAbility>(WorldContextObject);
-	if (ThisAbility)
-	{
-		UAbilityTask_WaitAbilityCommit * MyObj = NULL;
-		MyObj = NewObject<UAbilityTask_WaitAbilityCommit>();
-		MyObj->InitTask(ThisAbility);
-		MyObj->WithTag = InWithTag;
-		MyObj->WithoutTag = InWithoutTag;
-
+	auto MyObj = NewTask<UAbilityTask_WaitAbilityCommit>(WorldContextObject);
+	MyObj->WithTag = InWithTag;
+	MyObj->WithoutTag = InWithoutTag;
 
-		return MyObj;
-	}
-	return NULL;
+	return MyObj;
 }
 
 void UAbilityTask_WaitAbilityCommit::Activate()
@@ -38,6 +29,16 @@ void UAbilityTask_WaitAbilityCommit::Activate()
 	}
 }
 
+void UAbilityTask_WaitAbilityCommit::OnDestroy(bool AbilityEnded)
+{
+	if (AbilitySystemComponent.IsValid())
+	{
+		AbilitySystemComponent->AbilityCommitedCallbacks.RemoveUObject(this, &UAbilityTask_WaitAbilityCommit::OnAbilityCommit);
+	}
+
+	Super::OnDestroy(AbilityEnded);
+}
+
 void UAbilityTask_WaitAbilityCommit::OnAbilityCommit(UGameplayAbility *ActivatedAbility)
 {
 	if ( (WithTag.IsValid() && !ActivatedAbility->AbilityTags.HasTag(WithTag, EGameplayTagMatchType::IncludeParentTags, EGameplayTagMatchType::Explicit)) ||
@@ -47,10 +48,7 @@ void UAbilityTask_WaitAbilityCommit::OnAbilityCommit(UGameplayAbility *Activated
 		return;
 	}
 
-	if (AbilitySystemComponent.IsValid())
-	{
-		AbilitySystemComponent->AbilityCommitedCallbacks.RemoveUObject(this, &UAbilityTask_WaitAbilityCommit::OnAbilityCommit);
-	}
-
 	OnCommit.Broadcast(ActivatedAbility);
+
+	EndTask();
 }
\ No newline at end of file
diff --git a/Engine/Source/Runtime/GameplayAbilities/Private/Abilities/Tasks/AbilityTask_WaitAttributeChange.cpp b/Engine/Source/Runtime/GameplayAbilities/Private/Abilities/Tasks/AbilityTask_WaitAttributeChange.cpp
--- a/Engine/Source/Runtime/GameplayAbilities/Private/Abilities/Tasks/AbilityTask_WaitAttributeChange.cpp
+++ b/Engine/Source/Runtime/GameplayAbilities/Private/Abilities/Tasks/AbilityTask_WaitAttributeChange.cpp
@@ -14,20 +14,12 @@ UAbilityTask_WaitAttributeChange::UAbilityTask_WaitAttributeChange(const class F
 
 UAbilityTask_WaitAttributeChange* UAbilityTask_WaitAttributeChange::WaitForAttributeChange(UObject* WorldContextObject, FGameplayAttribute InAttribute, FGameplayTag InWithTag, FGameplayTag InWithoutTag)
 {
-	check(WorldContextObject);
-	UGameplayAbility* ThisAbility = CastChecked<UGameplayAbility>(WorldContextObject);
-	if (ThisAbility)
-	{
-		UAbilityTask_WaitAttributeChange * MyObj = NULL;
-		MyObj = NewObject<UAbilityTask_WaitAttributeChange>();
-		MyObj->InitTask(ThisAbility);
-		MyObj->WithTag = InWithTag;
-		MyObj->WithoutTag = InWithoutTag;
-		MyObj->Attribute = InAttribute;
+	auto MyObj = NewTask<UAbilityTask_WaitAttributeChange>(WorldContextObject);
+	MyObj->WithTag = InWithTag;
+	MyObj->WithoutTag = InWithoutTag;
+	MyObj->Attribute = InAttribute;
 
-		return MyObj;
-	}
-	return NULL;
+	return MyObj;
 }
 
 void UAbilityTask_WaitAttributeChange::Activate()
@@ -57,12 +49,19 @@ void UAbilityTask_WaitAttributeChange::OnAttributeChange(float NewValue, const F
 			// Failed tag check
 			return;
 		}
-	}
-	
+	}	
+
+	OnChange.Broadcast();
+
+	EndTask();
+}
+
+void UAbilityTask_WaitAttributeChange::OnDestroy(bool AbilityEnded)
+{
 	if (AbilitySystemComponent.IsValid())
 	{
 		AbilitySystemComponent->RegisterGameplayAttributeEvent(Attribute).RemoveUObject(this, &UAbilityTask_WaitAttributeChange::OnAttributeChange);
 	}
 
-	OnChange.Broadcast();
+	Super::OnDestroy(AbilityEnded);
 }
\ No newline at end of file
diff --git a/Engine/Source/Runtime/GameplayAbilities/Private/Abilities/Tasks/AbilityTask_WaitConfirm.cpp b/Engine/Source/Runtime/GameplayAbilities/Private/Abilities/Tasks/AbilityTask_WaitConfirm.cpp
--- a/Engine/Source/Runtime/GameplayAbilities/Private/Abilities/Tasks/AbilityTask_WaitConfirm.cpp
+++ b/Engine/Source/Runtime/GameplayAbilities/Private/Abilities/Tasks/AbilityTask_WaitConfirm.cpp
@@ -5,7 +5,7 @@
 UAbilityTask_WaitConfirm::UAbilityTask_WaitConfirm(const class FPostConstructInitializeProperties& PCIP)
 	: Super(PCIP)
 {
-
+	RegisteredCallback = false;
 }
 
 void UAbilityTask_WaitConfirm::OnConfirmCallback(UGameplayAbility* InAbility)
@@ -13,18 +13,12 @@ void UAbilityTask_WaitConfirm::OnConfirmCallback(UGameplayAbility* InAbility)
 	OnConfirm.Broadcast();
 
 	// We are done. Kill us so we don't keep getting broadcast messages
-	MarkPendingKill();
+	EndTask();
 }
 
 UAbilityTask_WaitConfirm* UAbilityTask_WaitConfirm::WaitConfirm(class UObject* WorldContextObject)
 {
-	check(WorldContextObject);
-
-	UAbilityTask_WaitConfirm * MyObj = NewObject<UAbilityTask_WaitConfirm>();
-	UGameplayAbility* ThisAbility = CastChecked<UGameplayAbility>(WorldContextObject);
-	MyObj->InitTask(ThisAbility);	
-
-	return MyObj;
+	return NewTask<UAbilityTask_WaitConfirm>(WorldContextObject);	
 }
 
 void UAbilityTask_WaitConfirm::Activate()
@@ -33,8 +27,10 @@ void UAbilityTask_WaitConfirm::Activate()
 	{
 		if (Ability->GetCurrentActivationInfo().ActivationMode == EGameplayAbilityActivationMode::Predicting)
 		{
-			// Register delegate so that when we are confirmed by the server, we call OnConfirmCallback
+			// Register delegate so that when we are confirmed by the server, we call OnConfirmCallback	
+			
 			Ability->OnConfirmDelegate.AddUObject(this, &UAbilityTask_WaitConfirm::OnConfirmCallback);
+			RegisteredCallback = true;
 		}
 		else
 		{
@@ -42,4 +38,14 @@ void UAbilityTask_WaitConfirm::Activate()
 			OnConfirmCallback(Ability.Get());
 		}
 	}
-}
\ No newline at end of file
+}
+
+void UAbilityTask_WaitConfirm::OnDestroy(bool AbilityEnded)
+{
+	if (RegisteredCallback && Ability.IsValid())
+	{
+		Ability->OnConfirmDelegate.RemoveUObject(this, &UAbilityTask_WaitConfirm::OnConfirmCallback);
+	}
+
+	Super::OnDestroy(AbilityEnded);
+}
diff --git a/Engine/Source/Runtime/GameplayAbilities/Private/Abilities/Tasks/AbilityTask_WaitConfirmCancel.cpp b/Engine/Source/Runtime/GameplayAbilities/Private/Abilities/Tasks/AbilityTask_WaitConfirmCancel.cpp
--- a/Engine/Source/Runtime/GameplayAbilities/Private/Abilities/Tasks/AbilityTask_WaitConfirmCancel.cpp
+++ b/Engine/Source/Runtime/GameplayAbilities/Private/Abilities/Tasks/AbilityTask_WaitConfirmCancel.cpp
@@ -9,6 +9,7 @@
 UAbilityTask_WaitConfirmCancel::UAbilityTask_WaitConfirmCancel(const class FPostConstructInitializeProperties& PCIP)
 	: Super(PCIP)
 {
+	RegisteredCallbacks = false;
 
 }
 
@@ -18,7 +19,7 @@ void UAbilityTask_WaitConfirmCancel::OnConfirmCallback()
 	if (AbilitySystemComponent.IsValid())
 	{
 		AbilitySystemComponent->ConsumeAbilityConfirmCancel();
-		MarkPendingKill();
+		EndTask();
 	}
 }
 
@@ -28,23 +29,13 @@ void UAbilityTask_WaitConfirmCancel::OnCancelCallback()
 	if (AbilitySystemComponent.IsValid())
 	{
 		AbilitySystemComponent->ConsumeAbilityConfirmCancel();
-		MarkPendingKill();
+		EndTask();
 	}
 }
 
 UAbilityTask_WaitConfirmCancel* UAbilityTask_WaitConfirmCancel::WaitConfirmCancel(UObject* WorldContextObject)
 {
-	check(WorldContextObject);
-	UGameplayAbility* ThisAbility = CastChecked<UGameplayAbility>(WorldContextObject);
-	if (ThisAbility)
-	{
-		UAbilityTask_WaitConfirmCancel * MyObj = NULL;
-		MyObj = NewObject<UAbilityTask_WaitConfirmCancel>();
-		MyObj->InitTask(ThisAbility);
-
-		return MyObj;
-	}
-	return NULL;
+	return NewTask<UAbilityTask_WaitConfirmCancel>(WorldContextObject);
 }
 
 void UAbilityTask_WaitConfirmCancel::Activate()
@@ -73,8 +64,21 @@ void UAbilityTask_WaitConfirmCancel::Activate()
 
 		AbilitySystemComponent->ConfirmCallbacks.AddDynamic(this, &UAbilityTask_WaitConfirmCancel::OnConfirmCallback);	// Tell me if the confirm input is pressed
 		AbilitySystemComponent->CancelCallbacks.AddDynamic(this, &UAbilityTask_WaitConfirmCancel::OnCancelCallback);	// Tell me if the cancel input is pressed
+
+		RegisteredCallbacks = true;
 		
 		// Tell me if something else tells me to 'force' my target data (an ability may allow targeting during animation/limited time)
 		// Tell me if something else tells me to cancel my targeting
 	}
+}
+
+void UAbilityTask_WaitConfirmCancel::OnDestroy(bool AbilityEnding)
+{
+	if (RegisteredCallbacks && AbilitySystemComponent.IsValid())
+	{
+		AbilitySystemComponent->ConfirmCallbacks.RemoveDynamic(this, &UAbilityTask_WaitConfirmCancel::OnConfirmCallback);
+		AbilitySystemComponent->CancelCallbacks.RemoveDynamic(this, &UAbilityTask_WaitConfirmCancel::OnCancelCallback);
+	}
+
+	Super::OnDestroy(AbilityEnding);
 }
\ No newline at end of file
diff --git a/Engine/Source/Runtime/GameplayAbilities/Private/Abilities/Tasks/AbilityTask_WaitGameplayEffectRemoved.cpp b/Engine/Source/Runtime/GameplayAbilities/Private/Abilities/Tasks/AbilityTask_WaitGameplayEffectRemoved.cpp
--- a/Engine/Source/Runtime/GameplayAbilities/Private/Abilities/Tasks/AbilityTask_WaitGameplayEffectRemoved.cpp
+++ b/Engine/Source/Runtime/GameplayAbilities/Private/Abilities/Tasks/AbilityTask_WaitGameplayEffectRemoved.cpp
@@ -13,18 +13,10 @@ UAbilityTask_WaitGameplayEffectRemoved::UAbilityTask_WaitGameplayEffectRemoved(c
 
 UAbilityTask_WaitGameplayEffectRemoved* UAbilityTask_WaitGameplayEffectRemoved::WaitForGameplayEffectRemoved(UObject* WorldContextObject, FActiveGameplayEffectHandle InHandle)
 {
-	check(WorldContextObject);
-	UGameplayAbility* ThisAbility = CastChecked<UGameplayAbility>(WorldContextObject);
-	if (ThisAbility)
-	{
-		UAbilityTask_WaitGameplayEffectRemoved * MyObj = NULL;
-		MyObj = NewObject<UAbilityTask_WaitGameplayEffectRemoved>();
-		MyObj->InitTask(ThisAbility);
-		MyObj->Handle = InHandle;
+	auto MyObj = NewTask<UAbilityTask_WaitGameplayEffectRemoved>(WorldContextObject);
+	MyObj->Handle = InHandle;
 
-		return MyObj;
-	}
-	return NULL;
+	return MyObj;
 }
 
 void UAbilityTask_WaitGameplayEffectRemoved::Activate()
@@ -44,8 +36,19 @@ void UAbilityTask_WaitGameplayEffectRemoved::Activate()
 	}
 }
 
+void UAbilityTask_WaitGameplayEffectRemoved::OnDestroy(bool AbilityIsEnding)
+{
+	FOnActiveGameplayEffectRemoved* DelPtr = AbilitySystemComponent->OnGameplayEffectRemovedDelegate(Handle);
+	if (DelPtr)
+	{
+		DelPtr->RemoveUObject(this, &UAbilityTask_WaitGameplayEffectRemoved::OnGameplayEffectRemoved);
+	}
+
+	Super::OnDestroy(AbilityIsEnding);
+}
+
 void UAbilityTask_WaitGameplayEffectRemoved::OnGameplayEffectRemoved()
 {
 	OnRemoved.Broadcast();
-	MarkPendingKill();
+	EndTask();
 }
\ No newline at end of file
diff --git a/Engine/Source/Runtime/GameplayAbilities/Private/Abilities/Tasks/AbilityTask_WaitMovementModeChange.cpp b/Engine/Source/Runtime/GameplayAbilities/Private/Abilities/Tasks/AbilityTask_WaitMovementModeChange.cpp
--- a/Engine/Source/Runtime/GameplayAbilities/Private/Abilities/Tasks/AbilityTask_WaitMovementModeChange.cpp
+++ b/Engine/Source/Runtime/GameplayAbilities/Private/Abilities/Tasks/AbilityTask_WaitMovementModeChange.cpp
@@ -10,33 +10,18 @@ UAbilityTask_WaitMovementModeChange::UAbilityTask_WaitMovementModeChange(const c
 
 UAbilityTask_WaitMovementModeChange* UAbilityTask_WaitMovementModeChange::CreateWaitMovementModeChange(class UObject* WorldContextObject, EMovementMode NewMode)
 {
-	check(WorldContextObject);
-
-	UAbilityTask_WaitMovementModeChange* MyObj = NewObject<UAbilityTask_WaitMovementModeChange>();
-	UGameplayAbility* ThisAbility = CastChecked<UGameplayAbility>(WorldContextObject);
-	MyObj->Ability = ThisAbility;
+	auto MyObj = NewTask<UAbilityTask_WaitMovementModeChange>(WorldContextObject);
 	MyObj->RequiredMode = NewMode;
-
 	return MyObj;
 }
 
 void UAbilityTask_WaitMovementModeChange::Activate()
 {
-	if (Ability.IsValid())
+	ACharacter* Character = Cast<ACharacter>(GetActor());
+	if (Character)
 	{
-		// Fixme: getting from the activating ability to the owning character is pretty awkward. We need
-		// a more standard way of doing this along with with a way for ability tasks to specify their requirements 
-		// (e.g., requires character / pawn, etc)
-
-		AActor * ActorOwner = Cast<AActor>(Ability->GetOuter());
-
-		TSharedPtr<FGameplayAbilityActorInfo> ActorInfo = TSharedPtr<FGameplayAbilityActorInfo>(UAbilitySystemGlobals::Get().AllocAbilityActorInfo(ActorOwner));
-
-		ACharacter * Character = Cast<ACharacter>(ActorInfo->Actor.Get());
-		if (Character)
-		{
-			Character->MovementModeChangedDelegate.AddDynamic(this, &UAbilityTask_WaitMovementModeChange::OnMovementModeChange);
-		}
+		Character->MovementModeChangedDelegate.AddDynamic(this, &UAbilityTask_WaitMovementModeChange::OnMovementModeChange);
+		MyCharacter = Character;
 	}
 }
 
@@ -48,10 +33,20 @@ void UAbilityTask_WaitMovementModeChange::OnMovementModeChange(ACharacter * Char
 		{
 			if (RequiredMode == MOVE_None || MoveComp->MovementMode == RequiredMode)
 			{
-				Character->MovementModeChangedDelegate.RemoveDynamic(this, &UAbilityTask_WaitMovementModeChange::OnMovementModeChange);
 				OnChange.Broadcast(MoveComp->MovementMode);
+				EndTask();
 				return;
 			}
 		}
 	}
+}
+
+void UAbilityTask_WaitMovementModeChange::OnDestroy(bool AbilityEnded)
+{
+	if (MyCharacter.IsValid())
+	{
+		MyCharacter->MovementModeChangedDelegate.RemoveDynamic(this, &UAbilityTask_WaitMovementModeChange::OnMovementModeChange);
+	}
+
+	Super::OnDestroy(AbilityEnded);
 }
\ No newline at end of file
diff --git a/Engine/Source/Runtime/GameplayAbilities/Private/Abilities/Tasks/AbilityTask_WaitOverlap.cpp b/Engine/Source/Runtime/GameplayAbilities/Private/Abilities/Tasks/AbilityTask_WaitOverlap.cpp
--- a/Engine/Source/Runtime/GameplayAbilities/Private/Abilities/Tasks/AbilityTask_WaitOverlap.cpp
+++ b/Engine/Source/Runtime/GameplayAbilities/Private/Abilities/Tasks/AbilityTask_WaitOverlap.cpp
@@ -9,7 +9,7 @@ UAbilityTask_WaitOverlap::UAbilityTask_WaitOverlap(const class FPostConstructIni
 }
 
 
-void UAbilityTask_WaitOverlap::OnOverlapCallback(class AActor* OtherActor, class UPrimitiveComponent* OtherComp, int32 OtherBodyIndex, bool bFromSweep, const FHitResult & SweepResult)
+void UAbilityTask_WaitOverlap::OnOverlapCallback(AActor* OtherActor, class UPrimitiveComponent* OtherComp, int32 OtherBodyIndex, bool bFromSweep, const FHitResult & SweepResult)
 {
 	if (OtherActor && bFromSweep)
 	{	
@@ -29,9 +29,8 @@ void UAbilityTask_WaitOverlap::OnHitCallback(AActor* OtherActor, UPrimitiveCompo
 		Handle.Data = TSharedPtr<FGameplayAbilityTargetData>(TargetData);
 		OnOverlap.Broadcast(Handle);
 
-
 		// We are done. Kill us so we don't keep getting broadcast messages
-		MarkPendingKill();
+		EndTask();
 	}
 }
 
@@ -41,32 +40,47 @@ void UAbilityTask_WaitOverlap::OnHitCallback(AActor* OtherActor, UPrimitiveCompo
 *	-Easy way to specify which types of actors/collision overlaps that we care about/want to block on
 */
 
-UAbilityTask_WaitOverlap* UAbilityTask_WaitOverlap::WaitForOverlap(class UObject* WorldContextObject)
+UAbilityTask_WaitOverlap* UAbilityTask_WaitOverlap::WaitForOverlap(UObject* WorldContextObject)
 {
-	check(WorldContextObject);
-
-	UAbilityTask_WaitOverlap * MyObj = NewObject<UAbilityTask_WaitOverlap>();
-	UGameplayAbility* ThisAbility = CastChecked<UGameplayAbility>(WorldContextObject);
-	MyObj->Ability = ThisAbility;
-
+	auto MyObj = NewTask<UAbilityTask_WaitOverlap>(WorldContextObject);
 	return MyObj;
 }
 
 void UAbilityTask_WaitOverlap::Activate()
 {
-	if (Ability.IsValid())
+	UPrimitiveComponent* PrimComponent = GetComponent();
+	if (PrimComponent)
 	{
-		AActor * ActorOwner = Cast<AActor>(Ability->GetOuter());
+		PrimComponent->OnComponentBeginOverlap.AddDynamic(this, &UAbilityTask_WaitOverlap::OnOverlapCallback);
+		PrimComponent->OnComponentHit.AddDynamic(this, &UAbilityTask_WaitOverlap::OnHitCallback);
+	}
+}
 
-		// TEMP - we are just using root component's collision. A real system will need more data to specify which component to use
-		UPrimitiveComponent * PrimComponent = Cast<UPrimitiveComponent>(ActorOwner->GetRootComponent());
+void UAbilityTask_WaitOverlap::OnDestroy(bool AbilityEnded)
+{
+	UPrimitiveComponent* PrimComponent = GetComponent();
+	if (PrimComponent)
+	{
+		PrimComponent->OnComponentBeginOverlap.RemoveDynamic(this, &UAbilityTask_WaitOverlap::OnOverlapCallback);
+		PrimComponent->OnComponentHit.RemoveDynamic(this, &UAbilityTask_WaitOverlap::OnHitCallback);
+	}
+
+	Super::OnDestroy(AbilityEnded);
+}
+
+UPrimitiveComponent* UAbilityTask_WaitOverlap::GetComponent()
+{
+	// TEMP - we are just using root component's collision. A real system will need more data to specify which component to use
+	UPrimitiveComponent * PrimComponent = nullptr;
+	AActor* ActorOwner = GetActor();
+	if (ActorOwner)
+	{
+		PrimComponent = Cast<UPrimitiveComponent>(ActorOwner->GetRootComponent());
 		if (!PrimComponent)
 		{
 			PrimComponent = ActorOwner->FindComponentByClass<UPrimitiveComponent>();
 		}
-		check(PrimComponent);
-
-		PrimComponent->OnComponentBeginOverlap.AddDynamic(this, &UAbilityTask_WaitOverlap::OnOverlapCallback);
-		PrimComponent->OnComponentHit.AddDynamic(this, &UAbilityTask_WaitOverlap::OnHitCallback);
 	}
+
+	return PrimComponent;
 }
\ No newline at end of file
diff --git a/Engine/Source/Runtime/GameplayAbilities/Private/Abilities/Tasks/AbilityTask_WaitTargetData.cpp b/Engine/Source/Runtime/GameplayAbilities/Private/Abilities/Tasks/AbilityTask_WaitTargetData.cpp
--- a/Engine/Source/Runtime/GameplayAbilities/Private/Abilities/Tasks/AbilityTask_WaitTargetData.cpp
+++ b/Engine/Source/Runtime/GameplayAbilities/Private/Abilities/Tasks/AbilityTask_WaitTargetData.cpp
@@ -12,9 +12,7 @@ UAbilityTask_WaitTargetData::UAbilityTask_WaitTargetData(const class FPostConstr
 
 UAbilityTask_WaitTargetData* UAbilityTask_WaitTargetData::WaitTargetData(UObject* WorldContextObject, TSubclassOf<AGameplayAbilityTargetActor> InTargetClass)
 {
-	UAbilityTask_WaitTargetData * MyObj = NewObject<UAbilityTask_WaitTargetData>();
-	UGameplayAbility * ThisAbility = CastChecked<UGameplayAbility>(WorldContextObject);
-	MyObj->InitTask(ThisAbility);
+	auto MyObj = NewTask<UAbilityTask_WaitTargetData>(WorldContextObject);
 	MyObj->TargetClass = InTargetClass;
 	return MyObj;	
 }
@@ -82,7 +80,7 @@ bool UAbilityTask_WaitTargetData::BeginSpawningActor(UObject* WorldContextObject
 			if (AbilitySystemComponent->ReplicatedTargetData.IsValid())
 			{
 				ValidData.Broadcast(AbilitySystemComponent->ReplicatedTargetData);
-				Cleanup();
+				EndTask();
 			}
 			else
 			{
@@ -138,7 +136,7 @@ void UAbilityTask_WaitTargetData::OnTargetDataReplicatedCallback(FGameplayAbilit
 		ValidData.Broadcast(Data);
 	}
 
-	Cleanup();
+	EndTask();
 }
 
 /** Client cancelled this Targeting Task (we are the server */
@@ -148,7 +146,7 @@ void UAbilityTask_WaitTargetData::OnTargetDataReplicatedCancelledCallback()
 
 	Cancelled.Broadcast(FGameplayAbilityTargetDataHandle());
 	
-	Cleanup();
+	EndTask();
 }
 
 /** The TargetActor we spawned locally has called back with valid target data */
@@ -166,7 +164,7 @@ void UAbilityTask_WaitTargetData::OnTargetDataReadyCallback(FGameplayAbilityTarg
 
 	ValidData.Broadcast(Data);
 	
-	Cleanup();
+	EndTask();
 }
 
 /** The TargetActor we spawned locally has called back with a cancel event (they still include the 'last/best' targetdata but the consumer of this may want to discard it) */
@@ -177,17 +175,11 @@ void UAbilityTask_WaitTargetData::OnTargetDataCancelledCallback(FGameplayAbility
 	AbilitySystemComponent->ServerSetReplicatedTargetDataCancelled();
 	Cancelled.Broadcast(Data);
 
-	Cleanup();
+	EndTask();
 }
 
-void UAbilityTask_WaitTargetData::Cleanup()
+void UAbilityTask_WaitTargetData::OnDestroy(bool AbilityEnded)
 {
-	if (HasAnyFlags(RF_PendingKill))
-	{
-		ensure(false);
-		return;
-	}
-
 	AbilitySystemComponent->ConsumeAbilityTargetData();
 	AbilitySystemComponent->SetUserAbilityActivationInhibited(false);
 
@@ -199,7 +191,7 @@ void UAbilityTask_WaitTargetData::Cleanup()
 		MyTargetActor->Destroy();
 	}
 
-	MarkPendingKill();
+	Super::OnDestroy(AbilityEnded);
 }
 
 
diff --git a/Engine/Source/Runtime/GameplayAbilities/Public/Abilities/GameplayAbility.h b/Engine/Source/Runtime/GameplayAbilities/Public/Abilities/GameplayAbility.h
--- a/Engine/Source/Runtime/GameplayAbilities/Public/Abilities/GameplayAbility.h
+++ b/Engine/Source/Runtime/GameplayAbilities/Public/Abilities/GameplayAbility.h
@@ -62,32 +62,6 @@
  */
 
 
-
-/***
- 
-
- TODO: Update: this is not all accurate anymore now that PredictiveACtivateAbility is gone. Please update comments!
-
-
-This is the basic flow of function calls in a networked environment:
--Client sees if he thinks he can activate the ability.
--Tells server to try and activate ability, predicts the ability immediately
--Server runs through Try/Can/Call, tells client if it passed or failed.
-
-
-																			 [If predictive] **Instance Ability** -> CallPredictiveActivateAbility() 
-			  (Optional)                                                     &&
-Client:		InputPressed()  ->  TryActivateAbility() -> CanActivateAbility() -> ServerTryActivateAbility()            --                         --                      --                          --                         --                ClientActivateAbilitySucceed() -> CallActivateAbility()
-																										  \\                                                                                                                                     //
-Server:		     --                      --                    --                          --                ServerTryActivateAbility()	-> CanActivateAbility() -> **Instance Ability** -> CallActivateAbility() -> ClientActivateAbilitySucceed()
-																																								||
-																																								-> ClientActivateAbilityFailed()
-																																																\\
-																																											Client:               ClientActivateAbilityFailed() [Remove and prediction work]
-
-***/
-
-
 /**
  *	Abilities define custom gameplay logic that can be activated by players or external game logic.
  */
@@ -186,8 +160,14 @@ public:
 	UPROPERTY(EditDefaultsOnly, Category = Tags)
 	FGameplayTagContainer AbilityTags;
 
+	/** Callback for when this ability has been confirmed by the server */
 	FGenericAbilityDelegate	OnConfirmDelegate;
 
+	void TaskStarted(UAbilityTask* NewTask);
+
+	void TaskEnded(UAbilityTask* Task);
+	
+
 protected:
 		
 	// --------------------------------------
@@ -359,7 +339,15 @@ protected:
 
 	// ----------------------------------------------------------------------------------------------------------------
 	//
-	//	Animation callbacks
+	//	Ability Tasks
+	//
+	// ----------------------------------------------------------------------------------------------------------------
+	
+	TArray<TWeakObjectPtr<UAbilityTask> >	ActiveTasks;
+
+	// ----------------------------------------------------------------------------------------------------------------
+	//
+	//	Animation callbacks (still WIP)
 	//
 	// ----------------------------------------------------------------------------------------------------------------
 
diff --git a/Engine/Source/Runtime/GameplayAbilities/Public/Abilities/GameplayAbilityTypes.h b/Engine/Source/Runtime/GameplayAbilities/Public/Abilities/GameplayAbilityTypes.h
--- a/Engine/Source/Runtime/GameplayAbilities/Public/Abilities/GameplayAbilityTypes.h
+++ b/Engine/Source/Runtime/GameplayAbilities/Public/Abilities/GameplayAbilityTypes.h
@@ -9,6 +9,7 @@ class UAnimInstance;
 class UAbilitySystemComponent;
 class UGameplayAbility;
 class AGameplayAbilityTargetActor;
+class UAbilityTask;
 
 
 UENUM(BlueprintType)
diff --git a/Engine/Source/Runtime/GameplayAbilities/Public/Abilities/Tasks/AbilityTask.h b/Engine/Source/Runtime/GameplayAbilities/Public/Abilities/Tasks/AbilityTask.h
--- a/Engine/Source/Runtime/GameplayAbilities/Public/Abilities/Tasks/AbilityTask.h
+++ b/Engine/Source/Runtime/GameplayAbilities/Public/Abilities/Tasks/AbilityTask.h
@@ -6,7 +6,7 @@
 
 /**
  *	AbilityTasks are small, self contained operations that can be performed while executing an ability.
- *	They are latent/asynchronous is nature. They will generate follow the pattern of 'start something and wait until it is finished or interrupted'
+ *	They are latent/asynchronous is nature. They will generally follow the pattern of 'start something and wait until it is finished or interrupted'
  *	
  *	We have code in K2Node_LatentAbilityCall to make using these in blueprints streamlined. The best way to become familiar with AbilityTasks is to 
  *	look at existing tasks like UAbilityTask_WaitOverlap (very simple) and UAbilityTask_WaitTargetData (much more complex).
@@ -27,6 +27,12 @@
  *	This is all you need for basic AbilityTasks. 
  *	
  *	
+ *	CheckList:
+ *		-Override ::OnDestroy() and unregister any callbacks that the task registered. Call Super::EndTask too!
+ *		-Implemented an Activate function which truly 'starts' the task. Do not 'start' the task in your static factory function!
+ *	
+ *	
+ *	--------------------------------------
  *	
  *	We have additional support for AbilityTasks that want to spawn actors. Though this could be accomplished in an Activate() function, it would not be
  *	possible to pass in dynamic "ExposeOnSpawn" actor properties. This is a powerful feature of blueprints, in order to support this, you need to implement 
@@ -61,10 +67,42 @@ class GAMEPLAYABILITIES_API UAbilityTask : public UObject
 	UFUNCTION(BlueprintCallable, meta = (BlueprintInternalUseOnly = "true"), Category = "Abilities")
 	virtual void Activate();
 
+	/** Initailizes the task with the owning GameplayAbility but does not actviate until Activate() is called */
 	virtual void InitTask(UGameplayAbility* InAbility);
+
+	
 	
 	/** GameplayAbility that created us */
 	TWeakObjectPtr<UGameplayAbility> Ability;
 
 	TWeakObjectPtr<UAbilitySystemComponent>	AbilitySystemComponent;
+
+	/** Helper function for getting UWorld off a task */
+	UWorld* GetWorld() const;
+
+	/** Proper way to get the owning actor of the ability that owns this task (usually a pawn, tower, etc) */
+	AActor* GetActor() const;
+
+	/** Helper function for instantiating and initializing a new task */
+	template <class T>
+	static T*	NewTask(UObject* WorldContextObject)
+	{
+		check(WorldContextObject);
+
+		T* MyObj = NewObject<T>();
+		UGameplayAbility* ThisAbility = CastChecked<UGameplayAbility>(WorldContextObject);
+		MyObj->InitTask(ThisAbility);
+		return MyObj;
+	}
+
+	/** Called when owning ability is ended (before the task ends) kills the task. Calls OnDestroy. */
+	void AbilityEnded();
+
+	/** Called explicitly to end the task (usually by the task itself). Calls OnDestroy. */
+	void EndTask();
+
+protected:	
+
+	/** End and CleanUp the task - may be called by the task itself or by the owning ability if the ability is ending. Do NOT call directly! Call EndTask() or AbilityEnded() */
+	virtual void OnDestroy(bool AbilityIsEnding);
 };
\ No newline at end of file
diff --git a/Engine/Source/Runtime/GameplayAbilities/Public/Abilities/Tasks/AbilityTask_MoveToLocation.h b/Engine/Source/Runtime/GameplayAbilities/Public/Abilities/Tasks/AbilityTask_MoveToLocation.h
--- a/Engine/Source/Runtime/GameplayAbilities/Public/Abilities/Tasks/AbilityTask_MoveToLocation.h
+++ b/Engine/Source/Runtime/GameplayAbilities/Public/Abilities/Tasks/AbilityTask_MoveToLocation.h
@@ -26,7 +26,12 @@ class UAbilityTask_MoveToLocation : public UAbilityTask
 	UFUNCTION(BlueprintCallable, Category = "Ability|Tasks", meta = (HidePin = "WorldContextObject", DefaultToSelf = "WorldContextObject", BlueprintInternalUseOnly = "TRUE"))
 	static UAbilityTask_MoveToLocation* MoveToLocation(UObject* WorldContextObject, FVector Location, float Duration, UCurveFloat* OptionalInterpolationCurve);
 
+	virtual void Activate() override;
+
 protected:
+
+	virtual void OnDestroy(bool AbilityEnded) override;
+
 	FVector StartLocation;
 	FVector TargetLocation;
 	float DurationOfMovement;
diff --git a/Engine/Source/Runtime/GameplayAbilities/Public/Abilities/Tasks/AbilityTask_PlayMontageAndWait.h b/Engine/Source/Runtime/GameplayAbilities/Public/Abilities/Tasks/AbilityTask_PlayMontageAndWait.h
--- a/Engine/Source/Runtime/GameplayAbilities/Public/Abilities/Tasks/AbilityTask_PlayMontageAndWait.h
+++ b/Engine/Source/Runtime/GameplayAbilities/Public/Abilities/Tasks/AbilityTask_PlayMontageAndWait.h
@@ -25,7 +25,9 @@ class UAbilityTask_PlayMontageAndWait : public UAbilityTask
 
 	virtual void Activate() override;
 
-public:
+private:
+
+	virtual void OnDestroy(bool AbilityEnded) override;
 
 	UAnimMontage* MontageToPlay;
 	
diff --git a/Engine/Source/Runtime/GameplayAbilities/Public/Abilities/Tasks/AbilityTask_WaitAbilityActivate.h b/Engine/Source/Runtime/GameplayAbilities/Public/Abilities/Tasks/AbilityTask_WaitAbilityActivate.h
--- a/Engine/Source/Runtime/GameplayAbilities/Public/Abilities/Tasks/AbilityTask_WaitAbilityActivate.h
+++ b/Engine/Source/Runtime/GameplayAbilities/Public/Abilities/Tasks/AbilityTask_WaitAbilityActivate.h
@@ -7,7 +7,6 @@
 DECLARE_DYNAMIC_MULTICAST_DELEGATE_OneParam(FWaitAbilityActivateDelegate, UGameplayAbility*, ActivatedAbility);
 
 class AActor;
-class UPrimitiveComponent;
 
 /**
  *	Waits for the actor to activate another ability
@@ -31,4 +30,9 @@ class UAbilityTask_WaitAbilityActivate : public UAbilityTask
 
 	FGameplayTag WithTag;
 	FGameplayTag WithoutTag;
+
+protected:
+
+	virtual void OnDestroy(bool AbilityEnded) override;
+
 };
\ No newline at end of file
diff --git a/Engine/Source/Runtime/GameplayAbilities/Public/Abilities/Tasks/AbilityTask_WaitAbilityCommit.h b/Engine/Source/Runtime/GameplayAbilities/Public/Abilities/Tasks/AbilityTask_WaitAbilityCommit.h
--- a/Engine/Source/Runtime/GameplayAbilities/Public/Abilities/Tasks/AbilityTask_WaitAbilityCommit.h
+++ b/Engine/Source/Runtime/GameplayAbilities/Public/Abilities/Tasks/AbilityTask_WaitAbilityCommit.h
@@ -5,11 +5,7 @@
 #include "AttributeSet.h"
 #include "AbilityTask_WaitAbilityCommit.generated.h"
 
-
-DECLARE_DYNAMIC_MULTICAST_DELEGATE_OneParam(FWaitAbilityCommitDelegate, UGameplayAbility*, ActivatedAbility);
-
 class AActor;
-class UPrimitiveComponent;
 
 /**
  *	Waits for the actor to activate another ability
@@ -19,6 +15,8 @@ class UAbilityTask_WaitAbilityCommit : public UAbilityTask
 {
 	GENERATED_UCLASS_BODY()
 
+	DECLARE_DYNAMIC_MULTICAST_DELEGATE_OneParam(FWaitAbilityCommitDelegate, UGameplayAbility*, ActivatedAbility);
+
 	UPROPERTY(BlueprintAssignable)
 	FWaitAbilityCommitDelegate	OnCommit;
 
@@ -29,8 +27,12 @@ class UAbilityTask_WaitAbilityCommit : public UAbilityTask
 
 	/** Wait until a new ability (of the same or different type) is commited. Used to gracefully interrupt abilities in specific ways. */
 	UFUNCTION(BlueprintCallable, Category="Ability|Tasks", meta = (HidePin = "WorldContextObject", DefaultToSelf = "WorldContextObject", BlueprintInternalUseOnly = "TRUE", FriendlyName = "Wait For New Ability Commit"))
-	static UAbilityTask_WaitAbilityCommit* WaitForAbilityCommit(UObject* WorldContextObject, FGameplayTag WithTag, FGameplayTag WithoutTage);	
+	static UAbilityTask_WaitAbilityCommit* WaitForAbilityCommit(UObject* WorldContextObject, FGameplayTag WithTag, FGameplayTag WithoutTage);
 
 	FGameplayTag WithTag;
 	FGameplayTag WithoutTag;
+
+protected:
+
+	virtual void OnDestroy(bool AbilityEnded) override;
 };
\ No newline at end of file
diff --git a/Engine/Source/Runtime/GameplayAbilities/Public/Abilities/Tasks/AbilityTask_WaitAttributeChange.h b/Engine/Source/Runtime/GameplayAbilities/Public/Abilities/Tasks/AbilityTask_WaitAttributeChange.h
--- a/Engine/Source/Runtime/GameplayAbilities/Public/Abilities/Tasks/AbilityTask_WaitAttributeChange.h
+++ b/Engine/Source/Runtime/GameplayAbilities/Public/Abilities/Tasks/AbilityTask_WaitAttributeChange.h
@@ -6,7 +6,7 @@
 #include "AbilityTask_WaitAttributeChange.generated.h"
 
 
-DECLARE_DYNAMIC_MULTICAST_DELEGATE(FWaitAttributeChangeDelegate);
+
 
 struct FGameplayEffectModCallbackData;
 
@@ -18,10 +18,14 @@ class UAbilityTask_WaitAttributeChange : public UAbilityTask
 {
 	GENERATED_UCLASS_BODY()
 
+	DECLARE_DYNAMIC_MULTICAST_DELEGATE(FWaitAttributeChangeDelegate);
+
 	UPROPERTY(BlueprintAssignable)
 	FWaitAttributeChangeDelegate	OnChange;
 
 	virtual void Activate() override;
+
+	
 		
 	void OnAttributeChange(float NewValue, const FGameplayEffectModCallbackData*);
 
@@ -32,4 +36,8 @@ class UAbilityTask_WaitAttributeChange : public UAbilityTask
 	FGameplayTag WithTag;
 	FGameplayTag WithoutTag;
 	FGameplayAttribute	Attribute;
+
+protected:
+
+	virtual void OnDestroy(bool AbilityEnded) override;
 };
\ No newline at end of file
diff --git a/Engine/Source/Runtime/GameplayAbilities/Public/Abilities/Tasks/AbilityTask_WaitConfirm.h b/Engine/Source/Runtime/GameplayAbilities/Public/Abilities/Tasks/AbilityTask_WaitConfirm.h
--- a/Engine/Source/Runtime/GameplayAbilities/Public/Abilities/Tasks/AbilityTask_WaitConfirm.h
+++ b/Engine/Source/Runtime/GameplayAbilities/Public/Abilities/Tasks/AbilityTask_WaitConfirm.h
@@ -21,6 +21,9 @@ class UAbilityTask_WaitConfirm : public UAbilityTask
 	UFUNCTION(BlueprintCallable, Category="Ability|Tasks", meta = (HidePin = "WorldContextObject", DefaultToSelf = "WorldContextObject", BlueprintInternalUseOnly = "TRUE"))
 	static UAbilityTask_WaitConfirm* WaitConfirm(UObject* WorldContextObject);
 
-public:
-	
+protected:
+
+	virtual void OnDestroy(bool AbilityEnded) override;
+
+	bool RegisteredCallback;
 };
\ No newline at end of file
diff --git a/Engine/Source/Runtime/GameplayAbilities/Public/Abilities/Tasks/AbilityTask_WaitConfirmCancel.h b/Engine/Source/Runtime/GameplayAbilities/Public/Abilities/Tasks/AbilityTask_WaitConfirmCancel.h
--- a/Engine/Source/Runtime/GameplayAbilities/Public/Abilities/Tasks/AbilityTask_WaitConfirmCancel.h
+++ b/Engine/Source/Runtime/GameplayAbilities/Public/Abilities/Tasks/AbilityTask_WaitConfirmCancel.h
@@ -7,6 +7,10 @@
 
 DECLARE_DYNAMIC_MULTICAST_DELEGATE(FWaitConfirmCancelDelegate);
 
+// Fixme: this name is conflicting with AbilityTask_WaitConfirm
+// UAbilityTask_WaitConfirmCancel = Wait for Targeting confirm/cancel
+// UAbilityTask_WaitConfirm = Wait for server to confirm ability activation
+
 UCLASS(MinimalAPI)
 class UAbilityTask_WaitConfirmCancel : public UAbilityTask
 {
@@ -27,8 +31,12 @@ class UAbilityTask_WaitConfirmCancel : public UAbilityTask
 	UFUNCTION(BlueprintCallable, meta=(HidePin = "WorldContextObject", DefaultToSelf = "WorldContextObject", BlueprintInternalUseOnly = "true", FriendlyName="Wait for Confirm Input"), Category="Ability|Tasks")
 	static UAbilityTask_WaitConfirmCancel* WaitConfirmCancel(UObject* WorldContextObject);	
 
-	virtual void Activate() override;	
+	virtual void Activate() override;
+
+protected:
+
+	virtual void OnDestroy(bool AbilityEnding) override;
 
-public:
+	bool RegisteredCallbacks;
 	
 };
\ No newline at end of file
diff --git a/Engine/Source/Runtime/GameplayAbilities/Public/Abilities/Tasks/AbilityTask_WaitGameplayEffectRemoved.h b/Engine/Source/Runtime/GameplayAbilities/Public/Abilities/Tasks/AbilityTask_WaitGameplayEffectRemoved.h
--- a/Engine/Source/Runtime/GameplayAbilities/Public/Abilities/Tasks/AbilityTask_WaitGameplayEffectRemoved.h
+++ b/Engine/Source/Runtime/GameplayAbilities/Public/Abilities/Tasks/AbilityTask_WaitGameplayEffectRemoved.h
@@ -32,4 +32,8 @@ class UAbilityTask_WaitGameplayEffectRemoved : public UAbilityTask
 	static UAbilityTask_WaitGameplayEffectRemoved* WaitForGameplayEffectRemoved(UObject* WorldContextObject, FActiveGameplayEffectHandle Handle);
 
 	FActiveGameplayEffectHandle Handle;
+
+protected:
+
+	virtual void OnDestroy(bool AbilityIsEnding) override;
 };
\ No newline at end of file
diff --git a/Engine/Source/Runtime/GameplayAbilities/Public/Abilities/Tasks/AbilityTask_WaitMovementModeChange.h b/Engine/Source/Runtime/GameplayAbilities/Public/Abilities/Tasks/AbilityTask_WaitMovementModeChange.h
--- a/Engine/Source/Runtime/GameplayAbilities/Public/Abilities/Tasks/AbilityTask_WaitMovementModeChange.h
+++ b/Engine/Source/Runtime/GameplayAbilities/Public/Abilities/Tasks/AbilityTask_WaitMovementModeChange.h
@@ -5,6 +5,8 @@
 
 DECLARE_DYNAMIC_MULTICAST_DELEGATE_OneParam(FMovementModeChangedDelegate, EMovementMode, NewMovementMode);
 
+class ACharacter;
+
 UCLASS(MinimalAPI)
 class UAbilityTask_WaitMovementModeChange : public UAbilityTask
 {
@@ -25,4 +27,9 @@ public:
 	static UAbilityTask_WaitMovementModeChange* CreateWaitMovementModeChange(UObject* WorldContextObject, EMovementMode NewMode);
 
 	virtual void Activate() override;
+private:
+
+	virtual void OnDestroy(bool AbilityEnded) override;
+
+	TWeakObjectPtr<ACharacter>	MyCharacter;
 };
\ No newline at end of file
diff --git a/Engine/Source/Runtime/GameplayAbilities/Public/Abilities/Tasks/AbilityTask_WaitOverlap.h b/Engine/Source/Runtime/GameplayAbilities/Public/Abilities/Tasks/AbilityTask_WaitOverlap.h
--- a/Engine/Source/Runtime/GameplayAbilities/Public/Abilities/Tasks/AbilityTask_WaitOverlap.h
+++ b/Engine/Source/Runtime/GameplayAbilities/Public/Abilities/Tasks/AbilityTask_WaitOverlap.h
@@ -34,6 +34,10 @@ class UAbilityTask_WaitOverlap : public UAbilityTask
 	UFUNCTION(BlueprintCallable, Category="Ability|Tasks", meta = (HidePin = "WorldContextObject", DefaultToSelf = "WorldContextObject", BlueprintInternalUseOnly = "TRUE"))
 	static UAbilityTask_WaitOverlap* WaitForOverlap(UObject* WorldContextObject);
 
-public:
+private:
+
+	virtual void OnDestroy(bool AbilityEnded) override;
+
+	UPrimitiveComponent* GetComponent();
 	
 };
\ No newline at end of file
diff --git a/Engine/Source/Runtime/GameplayAbilities/Public/Abilities/Tasks/AbilityTask_WaitTargetData.h b/Engine/Source/Runtime/GameplayAbilities/Public/Abilities/Tasks/AbilityTask_WaitTargetData.h
--- a/Engine/Source/Runtime/GameplayAbilities/Public/Abilities/Tasks/AbilityTask_WaitTargetData.h
+++ b/Engine/Source/Runtime/GameplayAbilities/Public/Abilities/Tasks/AbilityTask_WaitTargetData.h
@@ -41,12 +41,11 @@ class UAbilityTask_WaitTargetData: public UAbilityTask
 	UFUNCTION(BlueprintCallable, meta = (HidePin = "WorldContextObject", DefaultToSelf = "WorldContextObject", BlueprintInternalUseOnly = "true"), Category = "Abilities")
 	void FinishSpawningActor(UObject* WorldContextObject, AGameplayAbilityTargetActor* SpawnedActor);
 
-
 protected:
 
-	TSubclassOf<AGameplayAbilityTargetActor> TargetClass;
+	virtual void OnDestroy(bool AbilityEnded) override;
 
-	void Cleanup();
+	TSubclassOf<AGameplayAbilityTargetActor> TargetClass;
 
 	/** The TargetActor that we spawned, or the class CDO if this is a static targeting task */
 	TWeakObjectPtr<AGameplayAbilityTargetActor>	MyTargetActor;
