GameplayAbility Nonunity fix

diff --git a/Engine/Source/Runtime/GameplayAbilities/Private/Abilities/GameplayAbility.cpp b/Engine/Source/Runtime/GameplayAbilities/Private/Abilities/GameplayAbility.cpp
--- a/Engine/Source/Runtime/GameplayAbilities/Private/Abilities/GameplayAbility.cpp
+++ b/Engine/Source/Runtime/GameplayAbilities/Private/Abilities/GameplayAbility.cpp
@@ -5,6 +5,7 @@
 #include "GameplayEffect.h"
 #include "AbilitySystemComponent.h"
 #include "AbilitySystemBlueprintLibrary.h"
+#include "AbilityTask.h"
 
 // --------------------------------------------------------------------------------------------------------------------------------------------------------
 //
diff --git a/Engine/Source/Runtime/GameplayAbilities/Private/Abilities/Tasks/AbilityTask.cpp b/Engine/Source/Runtime/GameplayAbilities/Private/Abilities/Tasks/AbilityTask.cpp
--- a/Engine/Source/Runtime/GameplayAbilities/Private/Abilities/Tasks/AbilityTask.cpp
+++ b/Engine/Source/Runtime/GameplayAbilities/Private/Abilities/Tasks/AbilityTask.cpp
@@ -1,6 +1,7 @@
 
 #include "AbilitySystemPrivatePCH.h"
 #include "Abilities/Tasks/AbilityTask.h"
+#include "AbilitySystemComponent.h"
 
 UAbilityTask::UAbilityTask(const class FPostConstructInitializeProperties& PCIP)
 	: Super(PCIP)
