remove pragma optimize

diff --git a/Engine/Source/Runtime/GameplayAbilities/Private/Abilities/Tasks/AbilityTask_MoveToLocation.cpp b/Engine/Source/Runtime/GameplayAbilities/Private/Abilities/Tasks/AbilityTask_MoveToLocation.cpp
--- a/Engine/Source/Runtime/GameplayAbilities/Private/Abilities/Tasks/AbilityTask_MoveToLocation.cpp
+++ b/Engine/Source/Runtime/GameplayAbilities/Private/Abilities/Tasks/AbilityTask_MoveToLocation.cpp
@@ -2,8 +2,6 @@
 #include "AbilitySystemPrivatePCH.h"
 #include "Abilities/Tasks/AbilityTask_MoveToLocation.h"
 
-#pragma optimize( "", off )
-
 UAbilityTask_MoveToLocation::UAbilityTask_MoveToLocation(const class FPostConstructInitializeProperties& PCIP)
 : Super(PCIP)
 {
