Fixing some GUBP issues.

diff --git a/Engine/Source/Runtime/GameplayAbilities/Private/Abilities/GameplayAbility.cpp b/Engine/Source/Runtime/GameplayAbilities/Private/Abilities/GameplayAbility.cpp
--- a/Engine/Source/Runtime/GameplayAbilities/Private/Abilities/GameplayAbility.cpp
+++ b/Engine/Source/Runtime/GameplayAbilities/Private/Abilities/GameplayAbility.cpp
@@ -254,7 +254,7 @@ void UGameplayAbility::EndAbility(const FGameplayAbilityActorInfo* ActorInfo)
 	{
 		if (Task.IsValid())
 		{
-			Task->AbilityEnded();
+			Task.Get()->AbilityEnded();
 		}
 	}
 	ActiveTasks.Reset();	// Empty the array but dont resize memory, since this object is probably going to be destroyed very soon anyways.
diff --git a/Engine/Source/Runtime/GameplayAbilities/Private/Abilities/Tasks/AbilityTask.cpp b/Engine/Source/Runtime/GameplayAbilities/Private/Abilities/Tasks/AbilityTask.cpp
--- a/Engine/Source/Runtime/GameplayAbilities/Private/Abilities/Tasks/AbilityTask.cpp
+++ b/Engine/Source/Runtime/GameplayAbilities/Private/Abilities/Tasks/AbilityTask.cpp
@@ -25,7 +25,7 @@ UWorld* UAbilityTask::GetWorld() const
 {
 	if (AbilitySystemComponent.IsValid())
 	{
-		return AbilitySystemComponent->GetWorld();
+		return AbilitySystemComponent.Get()->GetWorld();
 	}
 
 	return nullptr;
