comment update

diff --git a/Engine/Source/Runtime/GameplayAbilities/Public/Abilities/GameplayAbility_CharacterJump.h b/Engine/Source/Runtime/GameplayAbilities/Public/Abilities/GameplayAbility_CharacterJump.h
--- a/Engine/Source/Runtime/GameplayAbilities/Public/Abilities/GameplayAbility_CharacterJump.h
+++ b/Engine/Source/Runtime/GameplayAbilities/Public/Abilities/GameplayAbility_CharacterJump.h
@@ -6,7 +6,7 @@
 #include "GameplayAbility_CharacterJump.generated.h"
 
 /**
- *	A gameplay ability that plays a single montage and applies a GameplayEffect
+ *	Ability that jumps with a character
  */
 UCLASS()
 class GAMEPLAYABILITIES_API UGameplayAbility_CharacterJump : public UGameplayAbility
diff --git a/Engine/Source/Runtime/GameplayAbilities/Public/GameplayAbilitySet.h b/Engine/Source/Runtime/GameplayAbilities/Public/GameplayAbilitySet.h
--- a/Engine/Source/Runtime/GameplayAbilities/Public/GameplayAbilitySet.h
+++ b/Engine/Source/Runtime/GameplayAbilities/Public/GameplayAbilitySet.h
@@ -30,11 +30,14 @@ struct FGameplayAbilityBindInfo
 };
 
 /**
-* UGameplayAbilitySet
-*	This contains a list of abilities along with key bindings. This will be very game specific, so it is expected for games to override this class.
-*	Functions like InitializeAbilities and BindInputComponentToAbilities are setup to act on the AbilitySystemComponent (rather than AbilitySystemComponent processing the ability set).
-*	Ideally this avoids games having to implement their own UAbilitySystemComponent and UGameplayAbilitySet, and they can just implement thei own UGameplayAbilitySet.
-*/
+ * UGameplayAbilitySet
+ *	This contains a list of abilities along with key bindings. This will be very game specific, so it is expected for games to override this class or create their own.
+ *	
+ *	Functions like InitializeAbilities and BindInputComponentToAbilities are setup to act on the AbilitySystemComponent (rather than AbilitySystemComponent processing the ability set).
+ *	Ideally this avoids games having to implement their own UAbilitySystemComponent and UGameplayAbilitySet, and they can just implement thei own UGameplayAbilitySet.
+ *	
+ *	In the end, this class just facilitates adding stuff to AbilitySystemComponent::ActivatableAbilities. Games are free to do this however they want.
+ */
 UCLASS(Blueprintable)
 class GAMEPLAYABILITIES_API UGameplayAbilitySet : public UDataAsset
 {
