Initial pass for triggering abilities based on events.

diff --git a/Engine/Source/Runtime/GameplayAbilities/Private/Abilities/GameplayAbility.cpp b/Engine/Source/Runtime/GameplayAbilities/Private/Abilities/GameplayAbility.cpp
--- a/Engine/Source/Runtime/GameplayAbilities/Private/Abilities/GameplayAbility.cpp
+++ b/Engine/Source/Runtime/GameplayAbilities/Private/Abilities/GameplayAbility.cpp
@@ -3,6 +3,7 @@
 #include "AbilitySystemPrivatePCH.h"
 #include "Abilities/GameplayAbility.h"
 #include "GameplayEffect.h"
+#include "GameplayTagsModule.h"
 #include "AbilitySystemComponent.h"
 #include "AbilitySystemBlueprintLibrary.h"
 #include "AbilityTask.h"
@@ -19,7 +20,11 @@ UGameplayAbility::UGameplayAbility(const class FPostConstructInitializePropertie
 	CostGameplayEffect = NULL;
 	CooldownGameplayEffect = NULL;
 
-	
+	{
+		static FName FuncName = FName(TEXT("K2_ShouldAbilityRespondToEvent"));
+		UFunction* ShouldRespondFunction = GetClass()->FindFunctionByName(FuncName);
+		HasBlueprintShouldAbilityRespondToEvent = ShouldRespondFunction && ShouldRespondFunction->GetOuter()->IsA(UBlueprintGeneratedClass::StaticClass());
+	}
 	{
 		static FName FuncName = FName(TEXT("K2_CanActivateAbility"));
 		UFunction* CanActivateFunction = GetClass()->FindFunctionByName(FuncName);
@@ -44,6 +49,28 @@ UGameplayAbility::UGameplayAbility(const class FPostConstructInitializePropertie
 #endif
 }
 
+void UGameplayAbility::TriggerAbilityFromGameplayEvent(FGameplayAbilityActorInfo* ActorInfo, FGameplayTag EventTag, FGameplayEventData* Payload)
+{
+	if (ShouldAbilityRespondToEvent(EventTag, Payload))
+	{
+		TryActivateAbility(ActorInfo);
+	}
+}
+
+// TODO: polymorphic payload
+void UGameplayAbility::SendGameplayEvent(FGameplayTag EventTag, FGameplayEventData Payload)
+{
+	AActor *OwnerActor = Cast<AActor>(GetOuter());
+	if (OwnerActor)
+	{
+		UAbilitySystemComponent* AbilitySystemComponent = OwnerActor->FindComponentByClass<UAbilitySystemComponent>();
+		if (ensure(AbilitySystemComponent))
+		{
+			AbilitySystemComponent->HandleGameplayEvent(EventTag, &Payload);
+		}
+	}
+}
+
 void UGameplayAbility::PostNetInit()
 {
 	/** We were dynamically spawned from replication - we need to init a currentactorinfo by looking at outer.
@@ -122,6 +149,20 @@ bool UGameplayAbility::CanActivateAbility(const FGameplayAbilityActorInfo* Actor
 	return true;
 }
 
+bool UGameplayAbility::ShouldAbilityRespondToEvent(FGameplayTag EventTag, const FGameplayEventData* Payload) const
+{
+	if (HasBlueprintShouldAbilityRespondToEvent)
+	{
+		if (K2_ShouldAbilityRespondToEvent(*Payload) == false)
+		{
+			ABILITY_LOG(Log, TEXT("ShouldAbilityRespondToEvent %s failed, blueprint refused"), *GetName());
+			return false;
+		}
+	}
+
+	return true;
+}
+
 bool UGameplayAbility::CommitAbility(const FGameplayAbilityActorInfo* ActorInfo, const FGameplayAbilityActivationInfo ActivationInfo)
 {
 	// Last chance to fail (maybe we no longer have resources to commit since we after we started this ability activation)
@@ -215,6 +256,9 @@ bool UGameplayAbility::TryActivateAbility(const FGameplayAbilityActorInfo* Actor
 	{
 		// This execution is now officially EGameplayAbilityActivationMode:Predicting and has a PredictionKey
 		ActivationInfo.GeneratePredictionKey(ActorInfo->AbilitySystemComponent.Get());
+		
+		// This must be called immediately after GeneratePredictionKey to prevent problems with recursively activating abilities
+		ServerTryActivateAbility(ActorInfo, ActivationInfo);
 
 		if (GetInstancingPolicy() == EGameplayAbilityInstancingPolicy::InstancedPerExecution)
 		{
@@ -240,8 +284,6 @@ bool UGameplayAbility::TryActivateAbility(const FGameplayAbilityActorInfo* Actor
 		{
 			CallActivateAbility(ActorInfo, ActivationInfo);
 		}
-		
-		ServerTryActivateAbility(ActorInfo, ActivationInfo);
 	}
 
 	return true;	
diff --git a/Engine/Source/Runtime/GameplayAbilities/Private/Abilities/GameplayAbilitySet.cpp b/Engine/Source/Runtime/GameplayAbilities/Private/Abilities/GameplayAbilitySet.cpp
--- a/Engine/Source/Runtime/GameplayAbilities/Private/Abilities/GameplayAbilitySet.cpp
+++ b/Engine/Source/Runtime/GameplayAbilities/Private/Abilities/GameplayAbilitySet.cpp
@@ -7,7 +7,7 @@
 
 // --------------------------------------------------------------------------------------------------------------------------------------------------------
 //
-//	UGameplayAbilitySet
+//	Helper Functions
 //
 // --------------------------------------------------------------------------------------------------------------------------------------------------------
 
@@ -68,6 +68,23 @@ void UGameplayAbilitySet::InitializeAbilities(UAbilitySystemComponent *AbilitySy
 
 			AbilitySystemComponent->ActivatableAbilities.Add(Ability);
 		}
+
+		for (FAbilityTriggerData TriggerData : Ability->AbilityTriggers)
+		{
+			FGameplayTag EventTag = TriggerData.TriggerTag;
+
+			if (AbilitySystemComponent->GameplayEventTriggeredAbilities.Contains(EventTag))
+			{
+				TArray<TWeakObjectPtr<UGameplayAbility> > Triggers = AbilitySystemComponent->GameplayEventTriggeredAbilities[EventTag];
+				Triggers.Add(Ability);
+			}
+			else
+			{
+				TArray<TWeakObjectPtr<UGameplayAbility> > Triggers;
+				Triggers.Add(Ability);
+				AbilitySystemComponent->GameplayEventTriggeredAbilities.Add(EventTag, Triggers);
+			}
+		}
 	}
 }
 
diff --git a/Engine/Source/Runtime/GameplayAbilities/Private/AbilitySystemComponent_Abilities.cpp b/Engine/Source/Runtime/GameplayAbilities/Private/AbilitySystemComponent_Abilities.cpp
--- a/Engine/Source/Runtime/GameplayAbilities/Private/AbilitySystemComponent_Abilities.cpp
+++ b/Engine/Source/Runtime/GameplayAbilities/Private/AbilitySystemComponent_Abilities.cpp
@@ -165,6 +165,29 @@ TEXT("AbilitySystem.DenyClientActivations"),
 	);
 #endif
 
+void UAbilitySystemComponent::OnRep_ActivateAbilities()
+{
+	for (UGameplayAbility * Ability : ActivatableAbilities)
+	{
+		for (FAbilityTriggerData TriggerData : Ability->AbilityTriggers)
+		{
+			FGameplayTag EventTag = TriggerData.TriggerTag;
+
+			if (GameplayEventTriggeredAbilities.Contains(EventTag))
+			{
+				TArray<TWeakObjectPtr<UGameplayAbility> > Triggers = GameplayEventTriggeredAbilities[EventTag];
+				Triggers.Add(Ability);
+			}
+			else
+			{
+				TArray<TWeakObjectPtr<UGameplayAbility> > Triggers;
+				Triggers.Add(Ability);
+				GameplayEventTriggeredAbilities.Add(EventTag, Triggers);
+			}
+		}
+	}
+}
+
 void UAbilitySystemComponent::ServerTryActivateAbility_Implementation(UGameplayAbility * AbilityToActivate, int32 PredictionKey)
 {
 #if !(UE_BUILD_SHIPPING || UE_BUILD_TEST)
@@ -342,6 +365,22 @@ void UAbilitySystemComponent::NotifyAbilityActivated(UGameplayAbility* Ability)
 	AbilityActivatedCallbacks.Broadcast(Ability);
 }
 
+void UAbilitySystemComponent::HandleGameplayEvent(FGameplayTag EventTag, FGameplayEventData* Payload)
+{
+	if (GameplayEventTriggeredAbilities.Contains(EventTag))
+	{
+		TArray<TWeakObjectPtr<UGameplayAbility> > TriggeredAbilities = GameplayEventTriggeredAbilities[EventTag];
+
+		for (auto Ability : TriggeredAbilities)
+		{
+			if (Ability.IsValid())
+			{
+				Ability.Get()->TriggerAbilityFromGameplayEvent(AbilityActorInfo.Get(), EventTag, Payload);
+			}
+		}
+	}
+}
+
 // --------------------------------------------------------------------------
 
 void UAbilitySystemComponent::BindToInputComponent(UInputComponent *InputComponent)
diff --git a/Engine/Source/Runtime/GameplayAbilities/Public/Abilities/GameplayAbility.h b/Engine/Source/Runtime/GameplayAbilities/Public/Abilities/GameplayAbility.h
--- a/Engine/Source/Runtime/GameplayAbilities/Public/Abilities/GameplayAbility.h
+++ b/Engine/Source/Runtime/GameplayAbilities/Public/Abilities/GameplayAbility.h
@@ -62,6 +62,15 @@
  */
 
 
+USTRUCT()
+struct FAbilityTriggerData
+{
+	GENERATED_USTRUCT_BODY()
+
+	UPROPERTY(EditAnywhere, Category=TriggerData)
+	FGameplayTag TriggerTag;
+};
+
 /**
  *	Abilities define custom gameplay logic that can be activated by players or external game logic.
  */
@@ -71,6 +80,7 @@ class GAMEPLAYABILITIES_API UGameplayAbility : public UObject
 	GENERATED_UCLASS_BODY()
 
 	friend class UAbilitySystemComponent;
+	friend class UGameplayAbilitySet;
 
 public:
 
@@ -113,6 +123,9 @@ public:
 
 	/** Returns true if this ability can be activated right now. Has no side effects */
 	virtual bool CanActivateAbility(const FGameplayAbilityActorInfo* ActorInfo) const;
+
+	/** Returns true if this ability can be triggered right now. Has no side effects */
+	virtual bool ShouldAbilityRespondToEvent(FGameplayTag EventTag, const FGameplayEventData* Payload) const;
 	
 	float GetCooldownTimeRemaining(const FGameplayAbilityActorInfo* ActorInfo) const;
 		
@@ -169,6 +182,16 @@ public:
 	
 
 protected:
+
+	// --------------------------------------
+	//	ShouldAbilityRespondToEvent
+	// --------------------------------------
+
+	/** Returns true if this ability can be activated right now. Has no side effects */
+	UFUNCTION(BlueprintImplementableEvent, Category = Ability, FriendlyName = "ShouldAbilityRespondToEvent")
+	virtual bool K2_ShouldAbilityRespondToEvent(FGameplayEventData Payload) const;
+
+	bool HasBlueprintShouldAbilityRespondToEvent;
 		
 	// --------------------------------------
 	//	CanActivate
@@ -201,6 +224,11 @@ protected:
 	/** Called on a predictive ability when the server confirms its execution */
 	void ConfirmActivateSucceed();
 
+	virtual void TriggerAbilityFromGameplayEvent(FGameplayAbilityActorInfo* ActorInfo, FGameplayTag Tag, FGameplayEventData* Payload);
+
+	UFUNCTION(BlueprintCallable, Category = "Abilities")
+		void SendGameplayEvent(FGameplayTag EventTag, FGameplayEventData Payload);
+
 	// --------------------------------------
 	//	CommitAbility
 	// --------------------------------------
@@ -312,28 +340,32 @@ protected:
 	UPROPERTY(EditDefaultsOnly, Category = Advanced)
 	TEnumAsByte<EGameplayAbilityInstancingPolicy::Type>	InstancingPolicy;						
 
-	/** This is information specific to this instance of the ability. E.g, whether it is predicting, authorting, confirmed, etc. */
+	/** This is information specific to this instance of the ability. E.g, whether it is predicting, authoring, confirmed, etc. */
 	UPROPERTY(BlueprintReadOnly, Category = Ability)
 	FGameplayAbilityActivationInfo	CurrentActivationInfo;
 
 	UPROPERTY(EditDefaultsOnly, Category=Advanced)
 	TEnumAsByte<EGameplayAbilityNetExecutionPolicy::Type> NetExecutionPolicy;
 
-	/** This GameplayEffect represents the cooldown. It will be applied when the ability is commited and the ability cannot be used again until it is expired. */
+	/** This GameplayEffect represents the cooldown. It will be applied when the ability is committed and the ability cannot be used again until it is expired. */
 	UPROPERTY(EditDefaultsOnly, Category=Cooldowns)
 	class UGameplayEffect* CooldownGameplayEffect;
 
-	/** This GameplayEffect represents the cost (mana, stamina, etc) of the ability. It will be applied when the ability is commited. */
+	/** This GameplayEffect represents the cost (mana, stamina, etc) of the ability. It will be applied when the ability is committed. */
 	UPROPERTY(EditDefaultsOnly, Category=Costs)
 	class UGameplayEffect* CostGameplayEffect;
 
+	/** Triggers to determine if this ability should execute in response to an event */
+	UPROPERTY(EditDefaultsOnly, Category = Triggers)
+	TArray<FAbilityTriggerData> AbilityTriggers;
+	
 	// ----------------------------------------------------------------------------------------------------------------
 	//
-	//	Ability exclusion / cancelling
+	//	Ability exclusion / canceling
 	//
 	// ----------------------------------------------------------------------------------------------------------------
 	
-	/** Abilities with these tags are canelled when this ability is executed */
+	/** Abilities with these tags are canceled when this ability is executed */
 	UPROPERTY(EditDefaultsOnly, Category = Tags)
 	FGameplayTagContainer CancelAbilitiesWithTag;
 
diff --git a/Engine/Source/Runtime/GameplayAbilities/Public/Abilities/GameplayAbilityTypes.h b/Engine/Source/Runtime/GameplayAbilities/Public/Abilities/GameplayAbilityTypes.h
--- a/Engine/Source/Runtime/GameplayAbilities/Public/Abilities/GameplayAbilityTypes.h
+++ b/Engine/Source/Runtime/GameplayAbilities/Public/Abilities/GameplayAbilityTypes.h
@@ -437,6 +437,23 @@ struct TStructOpsTypeTraits<FGameplayAbilityTargetDataHandle> : public TStructOp
 	};
 };
 
+USTRUCT(BlueprintType)
+struct GAMEPLAYABILITIES_API FGameplayEventData
+{
+	GENERATED_USTRUCT_BODY()
+	
+	UPROPERTY(EditAnywhere, BlueprintReadWrite, Category = GameplayAbilityTriggerPayload)
+	AActor* Instigator;
+	UPROPERTY(EditAnywhere, BlueprintReadWrite, Category = GameplayAbilityTriggerPayload)
+	AActor* Target;
+	UPROPERTY(EditAnywhere, BlueprintReadWrite, Category = GameplayAbilityTriggerPayload)
+	float Var1;
+	UPROPERTY(EditAnywhere, BlueprintReadWrite, Category = GameplayAbilityTriggerPayload)
+	float Var2;
+	UPROPERTY(EditAnywhere, BlueprintReadWrite, Category = GameplayAbilityTriggerPayload)
+	float Var3;
+};
+
 /** Generic callback for returning when target data is available */
 DECLARE_MULTICAST_DELEGATE_OneParam(FAbilityTargetData, FGameplayAbilityTargetDataHandle);
 
diff --git a/Engine/Source/Runtime/GameplayAbilities/Public/AbilitySystemComponent.h b/Engine/Source/Runtime/GameplayAbilities/Public/AbilitySystemComponent.h
--- a/Engine/Source/Runtime/GameplayAbilities/Public/AbilitySystemComponent.h
+++ b/Engine/Source/Runtime/GameplayAbilities/Public/AbilitySystemComponent.h
@@ -320,9 +320,11 @@ class GAMEPLAYABILITIES_API UAbilitySystemComponent : public UActorComponent
 	 *	without an AbilitySystemComponent. For example an ability could be written to execute on a StaticMeshActor. As long as the ability doesn't require 
 	 *	instancing or anything else that the AbilitySystemComponent would provide, then it doesn't need the component to function.
 	 */ 
-	UPROPERTY(Replicated, BlueprintReadOnly, Category = "Abilities")
+	UPROPERTY(ReplicatedUsing = OnRep_ActivateAbilities, BlueprintReadOnly, Category = "Abilities")
 	TArray<UGameplayAbility*>	ActivatableAbilities;
 	
+	UFUNCTION()
+	void	OnRep_ActivateAbilities();
 
 	UFUNCTION(Server, reliable, WithValidation)
 	void	ServerTryActivateAbility(class UGameplayAbility* AbilityToActivate, int32 PredictionKey);
@@ -364,6 +366,10 @@ class GAMEPLAYABILITIES_API UAbilitySystemComponent : public UActorComponent
 	FGenericAbilityDelegate AbilityActivatedCallbacks;
 	FGenericAbilityDelegate AbilityCommitedCallbacks;
 
+	void HandleGameplayEvent(FGameplayTag EventTag, FGameplayEventData* Payload);
+
+	TMap<FGameplayTag, TArray<TWeakObjectPtr<UGameplayAbility> > > GameplayEventTriggeredAbilities;
+
 	void NotifyAbilityCommit(UGameplayAbility* Ability);
 	void NotifyAbilityActivated(UGameplayAbility* Ability);
 
