Add GiveAbility to AbilitySystemComponent, minor cleanup.

diff --git a/Engine/Source/Runtime/GameplayAbilities/Private/Abilities/GameplayAbilitySet.cpp b/Engine/Source/Runtime/GameplayAbilities/Private/Abilities/GameplayAbilitySet.cpp
--- a/Engine/Source/Runtime/GameplayAbilities/Private/Abilities/GameplayAbilitySet.cpp
+++ b/Engine/Source/Runtime/GameplayAbilities/Private/Abilities/GameplayAbilitySet.cpp
@@ -58,32 +58,10 @@ void UGameplayAbilitySet::InitializeAbilities(UAbilitySystemComponent *AbilitySy
 {
 	for (const FGameplayAbilityBindInfo &BindInfo : this->Abilities)
 	{
-		UGameplayAbility * Ability = BindInfo.GameplayAbilityClass ? BindInfo.GameplayAbilityClass->GetDefaultObject<UGameplayAbility>() : NULL;
+		UGameplayAbility* Ability = BindInfo.GameplayAbilityClass ? BindInfo.GameplayAbilityClass->GetDefaultObject<UGameplayAbility>() : NULL;
 		if (Ability)
 		{
-			if (Ability->GetInstancingPolicy() == EGameplayAbilityInstancingPolicy::InstancedPerActor)
-			{
-				Ability = AbilitySystemComponent->CreateNewInstanceOfAbility(Ability);
-			}
-
-			AbilitySystemComponent->ActivatableAbilities.Add(Ability);
-		}
-
-		for (FAbilityTriggerData TriggerData : Ability->AbilityTriggers)
-		{
-			FGameplayTag EventTag = TriggerData.TriggerTag;
-
-			if (AbilitySystemComponent->GameplayEventTriggeredAbilities.Contains(EventTag))
-			{
-				TArray<TWeakObjectPtr<UGameplayAbility> > Triggers = AbilitySystemComponent->GameplayEventTriggeredAbilities[EventTag];
-				Triggers.Add(Ability);
-			}
-			else
-			{
-				TArray<TWeakObjectPtr<UGameplayAbility> > Triggers;
-				Triggers.Add(Ability);
-				AbilitySystemComponent->GameplayEventTriggeredAbilities.Add(EventTag, Triggers);
-			}
+			AbilitySystemComponent->GiveAbility(Ability);
 		}
 	}
 }
diff --git a/Engine/Source/Runtime/GameplayAbilities/Private/AbilitySystemComponent_Abilities.cpp b/Engine/Source/Runtime/GameplayAbilities/Private/AbilitySystemComponent_Abilities.cpp
--- a/Engine/Source/Runtime/GameplayAbilities/Private/AbilitySystemComponent_Abilities.cpp
+++ b/Engine/Source/Runtime/GameplayAbilities/Private/AbilitySystemComponent_Abilities.cpp
@@ -29,13 +29,10 @@ void UAbilitySystemComponent::InitializeComponent()
 	GetObjectsWithOuter(Owner, ChildObjects, false, RF_PendingKill);
 	for (UObject * Obj : ChildObjects)
 	{
-		UAttributeSet * Set = Cast<UAttributeSet>(Obj);
+		UAttributeSet* Set = Cast<UAttributeSet>(Obj);
 		if (Set)  
 		{
-			UObject * AT = Set->GetArchetype();
-
-			//ensure(SpawnedAttributes.Contains(Set) == false);
-			//ensure(Set->IsDefaultSubobject());
+			UObject * AT = Set->GetArchetype();	
 			SpawnedAttributes.Add(Set);
 		}
 	}
@@ -55,7 +52,39 @@ void UAbilitySystemComponent::InitAbilityActorInfo()
 	}
 }
 
-UGameplayAbility * UAbilitySystemComponent::CreateNewInstanceOfAbility(UGameplayAbility *Ability)
+UGameplayAbility* UAbilitySystemComponent::GiveAbility(UGameplayAbility* Ability)
+{
+	check(Ability);
+	UGameplayAbility* OutAbility = Ability; // We may instance this ability
+	
+	if (Ability->GetInstancingPolicy() == EGameplayAbilityInstancingPolicy::InstancedPerActor)
+	{
+		Ability = CreateNewInstanceOfAbility(Ability);
+	}
+
+	ActivatableAbilities.Add(Ability);
+	
+
+	for (FAbilityTriggerData TriggerData : Ability->AbilityTriggers)
+	{
+		FGameplayTag EventTag = TriggerData.TriggerTag;
+
+		if (GameplayEventTriggeredAbilities.Contains(EventTag))
+		{
+			GameplayEventTriggeredAbilities[EventTag].Add(Ability);
+		}
+		else
+		{
+			TArray<TWeakObjectPtr<UGameplayAbility> > Triggers;
+			Triggers.Add(Ability);
+			GameplayEventTriggeredAbilities.Add(EventTag, Triggers);
+		}
+	}
+
+	return OutAbility;
+}
+
+UGameplayAbility* UAbilitySystemComponent::CreateNewInstanceOfAbility(UGameplayAbility *Ability)
 {
 	check(Ability);
 	check(Ability->HasAllFlags(RF_ClassDefaultObject));
@@ -271,7 +300,6 @@ void UAbilitySystemComponent::ClientActivateAbilitySucceed_Implementation(UGamep
 
 	// Fixme: We need a better way to link up/reconcile preditive replicated abilities. It would be ideal if we could predictively spawn an
 	// ability and then replace/link it with the server spawned one once the server has confirmed it.
-	// 
 
 	FGameplayAbilityActivationInfo	ActivationInfo(EGameplayAbilityActivationMode::Confirmed, PredictionKey);
 
diff --git a/Engine/Source/Runtime/GameplayAbilities/Public/Abilities/GameplayAbility_CharacterJump.h b/Engine/Source/Runtime/GameplayAbilities/Public/Abilities/GameplayAbility_CharacterJump.h
--- a/Engine/Source/Runtime/GameplayAbilities/Public/Abilities/GameplayAbility_CharacterJump.h
+++ b/Engine/Source/Runtime/GameplayAbilities/Public/Abilities/GameplayAbility_CharacterJump.h
@@ -1,4 +1,4 @@
-// Copyright 1998-2013 Epic Games, Inc. All Rights Reserved.
+// Copyright 1998-2014 Epic Games, Inc. All Rights Reserved.
 
 #pragma once
 
@@ -6,7 +6,7 @@
 #include "GameplayAbility_CharacterJump.generated.h"
 
 /**
- *	Ability that jumps with a character
+ *	Ability that jumps with a character.
  */
 UCLASS()
 class GAMEPLAYABILITIES_API UGameplayAbility_CharacterJump : public UGameplayAbility
diff --git a/Engine/Source/Runtime/GameplayAbilities/Public/AbilitySystemComponent.h b/Engine/Source/Runtime/GameplayAbilities/Public/AbilitySystemComponent.h
--- a/Engine/Source/Runtime/GameplayAbilities/Public/AbilitySystemComponent.h
+++ b/Engine/Source/Runtime/GameplayAbilities/Public/AbilitySystemComponent.h
@@ -295,6 +295,8 @@ class GAMEPLAYABILITIES_API UAbilitySystemComponent : public UActorComponent
 	 *	
 	 */
 
+	UGameplayAbility* GiveAbility(UGameplayAbility* Ability);
+
 	UGameplayAbility* CreateNewInstanceOfAbility(UGameplayAbility* Ability);
 
 	void CancelAbilitiesWithTags(const FGameplayTagContainer Tags, const FGameplayAbilityActorInfo* ActorInfo, const FGameplayAbilityActivationInfo ActivationInfo, UGameplayAbility* Ignore);
