Updated mesh-targeting thing to remove fields it didn't need. Still need to pull scene component/bone socket out.

diff --git a/Engine/Source/Runtime/GameplayAbilities/Private/Abilities/GameplayAbility.cpp b/Engine/Source/Runtime/GameplayAbilities/Private/Abilities/GameplayAbility.cpp
--- a/Engine/Source/Runtime/GameplayAbilities/Private/Abilities/GameplayAbility.cpp
+++ b/Engine/Source/Runtime/GameplayAbilities/Private/Abilities/GameplayAbility.cpp
@@ -706,6 +706,31 @@ bool FGameplayAbilityTargetDataHandle::NetSerialize(FArchive& Ar, class UPackage
 	return true;
 }
 
+bool FGameplayAbilityTargetData_Mesh::NetSerialize(FArchive& Ar, class UPackageMap* Map, bool& bOutSuccess)
+{
+	Ar << SourceActor;
+	if (SourceActor)
+	{
+		Ar << SourceSocketName;
+	}
+	else
+	{
+		StartPoint.NetSerialize(Ar, Map, bOutSuccess);
+	}
+	Ar << bUseTargetPoint;
+	if (bUseTargetPoint)
+	{
+		TargetPoint.NetSerialize(Ar, Map, bOutSuccess);
+	}
+	else
+	{
+		AimDirection.NetSerialize(Ar, Map, bOutSuccess);
+	}
+
+	bOutSuccess = true;
+	return true;
+}
+
 bool FGameplayAbilityTargetData_SingleTargetHit::NetSerialize(FArchive& Ar, class UPackageMap* Map, bool& bOutSuccess)
 {
 	Ar << HitResult.Actor;
diff --git a/Engine/Source/Runtime/GameplayAbilities/Private/Abilities/GameplayAbilityTargetActor_SingleLineTrace.cpp b/Engine/Source/Runtime/GameplayAbilities/Private/Abilities/GameplayAbilityTargetActor_SingleLineTrace.cpp
--- a/Engine/Source/Runtime/GameplayAbilities/Private/Abilities/GameplayAbilityTargetActor_SingleLineTrace.cpp
+++ b/Engine/Source/Runtime/GameplayAbilities/Private/Abilities/GameplayAbilityTargetActor_SingleLineTrace.cpp
@@ -158,7 +158,15 @@ void AGameplayAbilityTargetActor_SingleLineTrace::ConfirmTargeting()
 FGameplayAbilityTargetDataHandle AGameplayAbilityTargetActor_SingleLineTrace::MakeTargetData(FHitResult HitResult) const
 {
 	/** Note this is cleaned up by the FGameplayAbilityTargetDataHandle (via an internal TSharedPtr) */
-	FGameplayAbilityTargetData_SingleTargetHit* ReturnData = new FGameplayAbilityTargetData_SingleTargetHit();
-	ReturnData->HitResult = HitResult;
+	
+	//FGameplayAbilityTargetData_SingleTargetHit* ReturnData = new FGameplayAbilityTargetData_SingleTargetHit();
+	//ReturnData->HitResult = HitResult;
+	//return FGameplayAbilityTargetDataHandle(ReturnData);
+
+	FGameplayAbilityTargetData_Mesh* ReturnData = new FGameplayAbilityTargetData_Mesh();
+	ReturnData->SourceActor = SourceActor;
+	ReturnData->AimDirection = (HitResult.Location - SourceActor->GetActorLocation()).SafeNormal();
+	ReturnData->TargetPoint = HitResult.Location;
+	ReturnData->bUseTargetPoint = true;
 	return FGameplayAbilityTargetDataHandle(ReturnData);
 }
\ No newline at end of file
diff --git a/Engine/Source/Runtime/GameplayAbilities/Public/Abilities/GameplayAbilityTypes.h b/Engine/Source/Runtime/GameplayAbilities/Public/Abilities/GameplayAbilityTypes.h
--- a/Engine/Source/Runtime/GameplayAbilities/Public/Abilities/GameplayAbilityTypes.h
+++ b/Engine/Source/Runtime/GameplayAbilities/Public/Abilities/GameplayAbilityTypes.h
@@ -283,7 +283,7 @@ struct GAMEPLAYABILITIES_API FGameplayAbilityTargetData
 		return false;
 	}
 
-	virtual const FHitResult * GetHitResult() const
+	virtual const FHitResult* GetHitResult() const
 	{
 		return NULL;
 	}
@@ -295,14 +295,26 @@ struct GAMEPLAYABILITIES_API FGameplayAbilityTargetData
 		return false;
 	}
 
-	virtual FTransform* GetOrigin() const
+	virtual FTransform GetOrigin() const
 	{
-		return nullptr;
+		return FTransform::Identity;
 	}
 
 	// -------------------------------------
 
-	virtual UScriptStruct * GetScriptStruct()
+	virtual bool HasEndPoint() const
+	{
+		return false;
+	}
+
+	virtual FVector GetEndPoint() const
+	{
+		return FVector::ZeroVector;
+	}
+
+	// -------------------------------------
+
+	virtual UScriptStruct* GetScriptStruct()
 	{
 		return FGameplayAbilityTargetData::StaticStruct();
 	}
@@ -310,6 +322,119 @@ struct GAMEPLAYABILITIES_API FGameplayAbilityTargetData
 	virtual FString ToString() const;
 };
 
+USTRUCT()
+struct GAMEPLAYABILITIES_API FGameplayAbilityTargetData_Mesh : public FGameplayAbilityTargetData
+{
+	GENERATED_USTRUCT_BODY()
+
+	FGameplayAbilityTargetData_Mesh()
+	: bUseTargetPoint(false)
+	, SourceActor(NULL)
+	, AimDirection(FVector::ZeroVector)
+	{}
+
+
+	//UPROPERTY(BlueprintReadWrite, EditAnywhere, Category = Projectile)
+
+	/** If this is valid, projectile will spawn StartPoint will be based on this actor. */
+	UPROPERTY()
+	AActor* SourceActor;
+
+	/** If SourceActor is valid, this will allow a socket to be used instead of the actor's location. */
+	UPROPERTY()
+		FName SourceSocketName;
+
+	/** Projectile will be spawned here if SourceActor is NULL or invalid. */
+	UPROPERTY()
+		FVector_NetQuantize StartPoint;
+
+	/** Projectile will be fired in this direction. */
+	UPROPERTY()
+		FVector_NetQuantizeNormal AimDirection;
+
+	/** Projectile will be aimed from its spawn point toward this point if UseTargetPoint is true. */
+	UPROPERTY()
+		FVector_NetQuantize TargetPoint;
+
+	/** If this is true, AimDirection will be ignored, and the projectile will be aimed at TargetPoint. */
+	UPROPERTY()
+		bool bUseTargetPoint;
+
+
+
+	//virtual TArray<AActor*>	GetActors() const
+
+	// -------------------------------------
+
+	//virtual bool HasHitResult() const
+	//virtual const FHitResult* GetHitResult() const
+
+	// -------------------------------------
+
+	virtual bool HasOrigin() const
+	{
+		return (SourceActor != NULL);
+	}
+
+	virtual FTransform GetOrigin() const
+	{
+		if (SourceActor)
+		{
+			FTransform ReturnTransform = SourceActor->GetTransform();
+			if (bUseTargetPoint)
+			{
+				ReturnTransform.SetRotation((TargetPoint - ReturnTransform.GetLocation()).SafeNormal().Rotation().Quaternion());
+			}
+			else
+			{
+				ReturnTransform.SetRotation(AimDirection.Rotation().Quaternion());
+			}
+			return ReturnTransform;
+		}
+		return FTransform::Identity;
+	}
+
+	// -------------------------------------
+
+	virtual bool HasEndPoint() const
+	{
+		return false;
+	}
+
+	virtual FVector GetEndPoint() const
+	{
+		if (bUseTargetPoint)
+		{
+			return TargetPoint;
+		}
+		return FVector::ZeroVector;
+	}
+
+	// -------------------------------------
+
+	virtual UScriptStruct* GetScriptStruct()
+	{
+		return FGameplayAbilityTargetData_Mesh::StaticStruct();
+	}
+
+	virtual FString ToString() const
+	{
+		return TEXT("FGameplayAbilityTargetData_Mesh");
+	}
+
+	bool NetSerialize(FArchive& Ar, class UPackageMap* Map, bool& bOutSuccess);
+};
+
+template<>
+struct TStructOpsTypeTraits<FGameplayAbilityTargetData_Mesh> : public TStructOpsTypeTraitsBase
+{
+	enum
+	{
+		WithNetSerializer = true	// For now this is REQUIRED for FGameplayAbilityTargetDataHandle net serialization to work
+	};
+};
+
+
 USTRUCT()
 struct GAMEPLAYABILITIES_API FGameplayAbilityTargetData_SingleTargetHit : public FGameplayAbilityTargetData
 {
@@ -341,7 +466,7 @@ struct GAMEPLAYABILITIES_API FGameplayAbilityTargetData_SingleTargetHit : public
 		return true;
 	}
 
-	virtual const FHitResult * GetHitResult() const
+	virtual const FHitResult* GetHitResult() const
 	{
 		return &HitResult;
 	}
@@ -353,7 +478,7 @@ struct GAMEPLAYABILITIES_API FGameplayAbilityTargetData_SingleTargetHit : public
 
 	bool NetSerialize(FArchive& Ar, class UPackageMap* Map, bool& bOutSuccess);
 
-	virtual UScriptStruct * GetScriptStruct()
+	virtual UScriptStruct* GetScriptStruct()
 	{
 		return FGameplayAbilityTargetData_SingleTargetHit::StaticStruct();
 	}
