Fix for clang issue caused by property order in initialization

diff --git a/Engine/Source/Runtime/GameplayAbilities/Public/Abilities/GameplayAbilityTypes.h b/Engine/Source/Runtime/GameplayAbilities/Public/Abilities/GameplayAbilityTypes.h
--- a/Engine/Source/Runtime/GameplayAbilities/Public/Abilities/GameplayAbilityTypes.h
+++ b/Engine/Source/Runtime/GameplayAbilities/Public/Abilities/GameplayAbilityTypes.h
@@ -328,9 +328,9 @@ struct GAMEPLAYABILITIES_API FGameplayAbilityTargetData_Mesh : public FGameplayA
 	GENERATED_USTRUCT_BODY()
 
 	FGameplayAbilityTargetData_Mesh()
-	: bUseTargetPoint(false)
-	, SourceActor(NULL)
+	: SourceActor(NULL)
 	, AimDirection(FVector::ZeroVector)
+	, bUseTargetPoint(false)
 	{}
 
 
