Rework GameplayAbilitySet
Allow AbilitySystemComponent to live anywhere (not just underneath its 'avatar' actor)

diff --git a/Engine/Source/Runtime/GameplayAbilities/Private/Abilities/GameplayAbility.cpp b/Engine/Source/Runtime/GameplayAbilities/Private/Abilities/GameplayAbility.cpp
--- a/Engine/Source/Runtime/GameplayAbilities/Private/Abilities/GameplayAbility.cpp
+++ b/Engine/Source/Runtime/GameplayAbilities/Private/Abilities/GameplayAbility.cpp
@@ -364,14 +364,14 @@ void UGameplayAbility::ClientActivateAbility(const FGameplayAbilityActorInfo* Ac
 	CallActivateAbility(ActorInfo, ActivationInfo);
 }
 
-void UGameplayAbility::InputPressed(int32 InputID, const FGameplayAbilityActorInfo* ActorInfo)
+void UGameplayAbility::InputPressed(const FGameplayAbilityActorInfo* ActorInfo)
 {
 	TryActivateAbility(ActorInfo);
 }
 
-void UGameplayAbility::InputReleased(int32 InputID, const FGameplayAbilityActorInfo* ActorInfo)
+void UGameplayAbility::InputReleased(const FGameplayAbilityActorInfo* ActorInfo)
 {
-	ABILITY_LOG(Log, TEXT("InputReleased: %d"), InputID);
+	
 }
 
 bool UGameplayAbility::CheckCooldown(const FGameplayAbilityActorInfo* ActorInfo) const
@@ -529,13 +529,13 @@ void UGameplayAbility::TaskEnded(UAbilityTask* Task)
 
 //----------------------------------------------------------------------
 
-void FGameplayAbilityActorInfo::InitFromActor(AActor *InActor)
+void FGameplayAbilityActorInfo::InitFromActor(AActor *InActor, UAbilitySystemComponent* InAbilitySystemComponent)
 {
 	check(InActor);
-	Actor = InActor;
+	check(InAbilitySystemComponent);
 
-	UStruct * Struct = FGameplayAbilityActorInfo::StaticStruct();
-	UScriptStruct * ScriptStruct = FGameplayAbilityActorInfo::StaticStruct();
+	Actor = InActor;
+	AbilitySystemComponent = InAbilitySystemComponent;
 
 	// Look for a player controller or pawn in the owner chain.
 	AActor *TestActor = InActor;
@@ -563,8 +563,6 @@ void FGameplayAbilityActorInfo::InitFromActor(AActor *InActor)
 		this->AnimInstance = SkelMeshComponent->GetAnimInstance();
 	}
 
-	AbilitySystemComponent = InActor->FindComponentByClass<UAbilitySystemComponent>();
-
 	MovementComponent = InActor->FindComponentByClass<UMovementComponent>();
 }
 
@@ -683,7 +681,7 @@ bool FGameplayAbilityTargetDataHandle::NetSerialize(FArchive& Ar, class UPackage
 			//	1) we have to manually crawl through the topmost struct's fields since we don't have a UStructProperty for it (just the UScriptProperty)
 			//	2) if there are any UStructProperties in the topmost struct's fields, we will assert in UStructProperty::NetSerializeItem.
 			
-			ABILITY_LOG(Fatal, TEXT("FGameplayAbilityTargetDataHandle::NetSerialize called on data struct %s without a native NetSerilaize"), *ScriptStruct->GetName());
+			ABILITY_LOG(Fatal, TEXT("FGameplayAbilityTargetDataHandle::NetSerialize called on data struct %s without a native NetSerialize"), *ScriptStruct->GetName());
 
 			for (TFieldIterator<UProperty> It(ScriptStruct); It; ++It)
 			{
diff --git a/Engine/Source/Runtime/GameplayAbilities/Private/Abilities/GameplayAbilitySet.cpp b/Engine/Source/Runtime/GameplayAbilities/Private/Abilities/GameplayAbilitySet.cpp
--- a/Engine/Source/Runtime/GameplayAbilities/Private/Abilities/GameplayAbilitySet.cpp
+++ b/Engine/Source/Runtime/GameplayAbilities/Private/Abilities/GameplayAbilitySet.cpp
@@ -5,41 +5,6 @@
 #include "Abilities/GameplayAbility.h"
 #include "GameplayAbilitySet.h"
 
-// --------------------------------------------------------------------------------------------------------------------------------------------------------
-//
-//	Helper Functions
-//
-// --------------------------------------------------------------------------------------------------------------------------------------------------------
-
-void InputPressed(TWeakObjectPtr<UGameplayAbility> Ability, int32 InputID, TWeakObjectPtr<AActor> ActorOwner)
-{
-	FGameplayAbilityActorInfo	ActorInfo;
-	ActorInfo.InitFromActor(ActorOwner.Get());
-
-	UAbilitySystemComponent * AbilitySystemComponent = ActorInfo.AbilitySystemComponent.Get();
-	if (AbilitySystemComponent)
-	{
-		if (Ability.IsValid())
-		{
-			Ability.Get()->InputPressed(InputID, AbilitySystemComponent->AbilityActorInfo.Get());
-		}
-	}
-}
-
-void InputReleased(TWeakObjectPtr<UGameplayAbility> Ability, int32 InputID, TWeakObjectPtr<AActor> ActorOwner)
-{
-	FGameplayAbilityActorInfo	ActorInfo;
-	ActorInfo.InitFromActor(ActorOwner.Get());
-
-	UAbilitySystemComponent * AbilitySystemComponent = ActorInfo.AbilitySystemComponent.Get();
-	if (AbilitySystemComponent)
-	{
-		if (Ability.IsValid())
-		{
-			Ability.Get()->InputReleased(InputID, AbilitySystemComponent->AbilityActorInfo.Get());
-		}
-	}
-}
 
 // --------------------------------------------------------------------------------------------------------------------------------------------------------
 //
@@ -54,44 +19,13 @@ UGameplayAbilitySet::UGameplayAbilitySet(const class FPostConstructInitializePro
 
 }
 
-void UGameplayAbilitySet::InitializeAbilities(UAbilitySystemComponent *AbilitySystemComponent) const
+void UGameplayAbilitySet::GiveAbilities(UAbilitySystemComponent* AbilitySystemComponent) const
 {
-	for (const FGameplayAbilityBindInfo &BindInfo : this->Abilities)
+	for (const FGameplayAbilityBindInfo& BindInfo : Abilities)
 	{
-		UGameplayAbility* Ability = BindInfo.GameplayAbilityClass ? BindInfo.GameplayAbilityClass->GetDefaultObject<UGameplayAbility>() : NULL;
-		if (Ability)
+		if (BindInfo.GameplayAbilityClass)
 		{
-			AbilitySystemComponent->GiveAbility(Ability);
+			AbilitySystemComponent->GiveAbility(BindInfo.GameplayAbilityClass->GetDefaultObject<UGameplayAbility>(), (int32)BindInfo.Command);
 		}
 	}
 }
-
-void UGameplayAbilitySet::BindInputComponentToAbilities(UInputComponent *InputComponent, UAbilitySystemComponent *AbilitySystemComponent) const
-{
-	AActor * OwnerActor = AbilitySystemComponent->GetOwner();
-	check(OwnerActor);
-
-	for (int32 idx = 0; idx < this->Abilities.Num(); ++idx)
-	{
-		const FGameplayAbilityBindInfo & BindInfo = this->Abilities[idx];
-		UGameplayAbility * Ability = AbilitySystemComponent->ActivatableAbilities[idx];	
-		
-		for (const FGameplayAbilityBindInfoCommandIDPair &CommandPair : BindInfo.CommandBindings)
-		{
-			// Pressed event
-			{
-				FInputActionBinding AB(FName(*CommandPair.CommandString), IE_Pressed);
-				AB.ActionDelegate.GetDelegateForManualSet().BindStatic(&InputPressed, TWeakObjectPtr<UGameplayAbility>(Ability), CommandPair.InputID, TWeakObjectPtr<AActor>(OwnerActor));
-				InputComponent->AddActionBinding(AB);
-			}
-
-			// Released event
-			{
-				FInputActionBinding AB(FName(*CommandPair.CommandString), IE_Released);
-				AB.ActionDelegate.GetDelegateForManualSet().BindStatic(&InputReleased, TWeakObjectPtr<UGameplayAbility>(Ability), CommandPair.InputID, TWeakObjectPtr<AActor>(OwnerActor));
-				InputComponent->AddActionBinding(AB);
-			}
-		}
-	}
-}
-
diff --git a/Engine/Source/Runtime/GameplayAbilities/Private/Abilities/GameplayAbility_CharacterJump.cpp b/Engine/Source/Runtime/GameplayAbilities/Private/Abilities/GameplayAbility_CharacterJump.cpp
--- a/Engine/Source/Runtime/GameplayAbilities/Private/Abilities/GameplayAbility_CharacterJump.cpp
+++ b/Engine/Source/Runtime/GameplayAbilities/Private/Abilities/GameplayAbility_CharacterJump.cpp
@@ -30,7 +30,7 @@ void UGameplayAbility_CharacterJump::ActivateAbility(const FGameplayAbilityActor
 	}
 }
 
-void UGameplayAbility_CharacterJump::InputReleased(int32 InputID, const FGameplayAbilityActorInfo* ActorInfo)
+void UGameplayAbility_CharacterJump::InputReleased(const FGameplayAbilityActorInfo* ActorInfo)
 {
 	// This is the 'top' of the ability execution code path. We don't have an ActivationInfo yet.
 	// Normally for InputPressed, we do a bunch of stuff (for clients): predict, tell the server, wait for confirm, rollback, etc.
diff --git a/Engine/Source/Runtime/GameplayAbilities/Private/Abilities/Tasks/AbilityTask_MoveToLocation.cpp b/Engine/Source/Runtime/GameplayAbilities/Private/Abilities/Tasks/AbilityTask_MoveToLocation.cpp
--- a/Engine/Source/Runtime/GameplayAbilities/Private/Abilities/Tasks/AbilityTask_MoveToLocation.cpp
+++ b/Engine/Source/Runtime/GameplayAbilities/Private/Abilities/Tasks/AbilityTask_MoveToLocation.cpp
@@ -10,27 +10,25 @@ UAbilityTask_MoveToLocation::UAbilityTask_MoveToLocation(const class FPostConstr
 
 void UAbilityTask_MoveToLocation::InterpolatePosition()
 {
-	if (Ability.IsValid())
+	AActor* MyActor = GetActor();
+	if (MyActor)
 	{
-		if (AActor* ActorOwner = Cast<AActor>(Ability->GetOuter()))
-		{
-			float CurrentTime = ActorOwner->GetWorld()->GetTimeSeconds();
+		float CurrentTime = GetWorld()->GetTimeSeconds();
 
-			if (CurrentTime >= TimeMoveWillEnd)
-			{
-				ActorOwner->SetActorLocation(TargetLocation);
-				OnTargetLocationReached.Broadcast();
-				EndTask();
-			}
-			else
+		if (CurrentTime >= TimeMoveWillEnd)
+		{
+			MyActor->SetActorLocation(TargetLocation);
+			OnTargetLocationReached.Broadcast();
+			EndTask();
+		}
+		else
+		{
+			float MoveFraction = (CurrentTime - TimeMoveStarted) / DurationOfMovement;
+			if (LerpCurve.IsValid())
 			{
-				float MoveFraction = (CurrentTime - TimeMoveStarted) / DurationOfMovement;
-				if (LerpCurve.IsValid())
-				{
-					MoveFraction = LerpCurve.Get()->GetFloatValue(MoveFraction);
-				}
-				ActorOwner->SetActorLocation(FMath::Lerp<FVector, float>(StartLocation, TargetLocation, MoveFraction));				
+				MoveFraction = LerpCurve.Get()->GetFloatValue(MoveFraction);
 			}
+			MyActor->SetActorLocation(FMath::Lerp<FVector, float>(StartLocation, TargetLocation, MoveFraction));
 		}
 	}
 }
diff --git a/Engine/Source/Runtime/GameplayAbilities/Private/AbilitySystemBlueprintLibrary.cpp b/Engine/Source/Runtime/GameplayAbilities/Private/AbilitySystemBlueprintLibrary.cpp
--- a/Engine/Source/Runtime/GameplayAbilities/Private/AbilitySystemBlueprintLibrary.cpp
+++ b/Engine/Source/Runtime/GameplayAbilities/Private/AbilitySystemBlueprintLibrary.cpp
@@ -76,6 +76,10 @@ FVector UAbilitySystemBlueprintLibrary::GetTargetDataEndPoint(FGameplayAbilityTa
 				return HitResultPtr->TraceEnd;
 			}
 		}
+		else if(Data->HasEndPoint())
+		{
+			return Data->GetEndPoint();
+		}
 	}
 
 	return FVector();
diff --git a/Engine/Source/Runtime/GameplayAbilities/Private/AbilitySystemComponent.cpp b/Engine/Source/Runtime/GameplayAbilities/Private/AbilitySystemComponent.cpp
--- a/Engine/Source/Runtime/GameplayAbilities/Private/AbilitySystemComponent.cpp
+++ b/Engine/Source/Runtime/GameplayAbilities/Private/AbilitySystemComponent.cpp
@@ -136,7 +136,7 @@ bool UAbilitySystemComponent::AreGameplayEffectApplicationRequirementsSatisfied(
 		Instigator.GetOwnedGameplayTags(InstigatorTags);
 
 		FGameplayTagContainer TargetTags;
-		IGameplayTagAssetInterface* OwnerGTA = InterfaceCast<IGameplayTagAssetInterface>(GetOwner());
+		IGameplayTagAssetInterface* OwnerGTA = InterfaceCast<IGameplayTagAssetInterface>(AbilityActorInfo->Actor.Get());
 		if (OwnerGTA)
 		{
 			OwnerGTA->GetOwnedGameplayTags(TargetTags);
@@ -172,7 +172,7 @@ FGameplayEffectSpecHandle UAbilitySystemComponent::GetOutgoingSpec(UGameplayEffe
 	SCOPE_CYCLE_COUNTER(STAT_GetOutgoingSpec);
 	// Fixme: we should build a map and cache these off. We can invalidate the map when an OutgoingGE modifier is applied or removed from us.
 
-	FGameplayEffectSpec* NewSpec = new FGameplayEffectSpec(GameplayEffect, GetOwner(), FGameplayEffectLevelSpec::INVALID_LEVEL, GetCurveDataOverride());
+	FGameplayEffectSpec* NewSpec = new FGameplayEffectSpec(GameplayEffect, AbilityActorInfo->Actor.Get(), FGameplayEffectLevelSpec::INVALID_LEVEL, GetCurveDataOverride());
 	if (ActiveGameplayEffects.ApplyActiveEffectsTo(*NewSpec, FModifierQualifier().Type(EGameplayMod::OutgoingGE)))
 	{
 		return FGameplayEffectSpecHandle(NewSpec);
@@ -185,7 +185,7 @@ FGameplayEffectSpecHandle UAbilitySystemComponent::GetOutgoingSpec(UGameplayEffe
 FGameplayEffectInstigatorContext UAbilitySystemComponent::GetInstigatorContext() const
 {
 	FGameplayEffectInstigatorContext Context;
-	Context.AddInstigator(GetOwner());
+	Context.AddInstigator(AbilityActorInfo->Actor.Get());
 	return Context;
 }
 
@@ -194,7 +194,7 @@ FActiveGameplayEffectHandle UAbilitySystemComponent::ApplyGameplayEffectToTarget
 {
 	check(GameplayEffect);
 
-	FGameplayEffectSpec	Spec(GameplayEffect, GetOwner(), Level, GetCurveDataOverride());
+	FGameplayEffectSpec	Spec(GameplayEffect, AbilityActorInfo->Actor.Get(), Level, GetCurveDataOverride());
 	
 	return ApplyGameplayEffectSpecToTarget(Spec, Target, BaseQualifier);
 }
@@ -435,7 +435,7 @@ void UAbilitySystemComponent::InvokeGameplayCueEvent(const FGameplayEffectSpec &
 	}
 
 
-	AActor *ActorOwner = GetOwner();
+	AActor *ActorOwner = AbilityActorInfo->Actor.Get();
 	IGameplayCueInterface * GameplayCueInterface = InterfaceCast<IGameplayCueInterface>(ActorOwner);
 	if (!GameplayCueInterface)
 	{
@@ -601,7 +601,7 @@ void UAbilitySystemComponent::DispatchBlueprintCustomHandler(AActor* Actor, UFun
 void UAbilitySystemComponent::PrintAllGameplayEffects() const
 {
 	ABILITY_LOG_SCOPE(TEXT("PrintAllGameplayEffects %s"), *GetName());
-	ABILITY_LOG(Log, TEXT("Owner: %s"), *GetOwner()->GetName());
+	ABILITY_LOG(Log, TEXT("Owner: %s. Avatar: %s"), *GetOwner()->GetName(), *AbilityActorInfo->Actor->GetName());
 	ActiveGameplayEffects.PrintAllGameplayEffects();
 }
 
diff --git a/Engine/Source/Runtime/GameplayAbilities/Private/AbilitySystemComponent_Abilities.cpp b/Engine/Source/Runtime/GameplayAbilities/Private/AbilitySystemComponent_Abilities.cpp
--- a/Engine/Source/Runtime/GameplayAbilities/Private/AbilitySystemComponent_Abilities.cpp
+++ b/Engine/Source/Runtime/GameplayAbilities/Private/AbilitySystemComponent_Abilities.cpp
@@ -20,9 +20,10 @@ void UAbilitySystemComponent::InitializeComponent()
 {
 	Super::InitializeComponent();
 
-	InitAbilityActorInfo();
+	/** Allocate an AbilityActorInfo. Note: this goes through a global function and is a SharedPtr so projects can make their own AbilityActorInfo */
+	AbilityActorInfo = TSharedPtr<FGameplayAbilityActorInfo>(UAbilitySystemGlobals::Get().AllocAbilityActorInfo());
 	
-	// Look for DSO AttributeSets
+	// Look for DSO AttributeSets (note we are currently requiring all attribute sets to be subobjects of the same owner. This doesn't *have* to be the case forever.
 	AActor *Owner = GetOwner();
 
 	TArray<UObject*> ChildObjects;
@@ -38,40 +39,44 @@ void UAbilitySystemComponent::InitializeComponent()
 	}
 }
 
-void UAbilitySystemComponent::InitAbilityActorInfo()
+void UAbilitySystemComponent::InitAbilityActorInfo(AActor* AvatarActor)
 {
-	if (!AbilityActorInfo.IsValid())
-	{
-		// Alloc (and init) a new actor info
-		AbilityActorInfo = TSharedPtr<FGameplayAbilityActorInfo>(UAbilitySystemGlobals::Get().AllocAbilityActorInfo(GetOwner()));
-	}
-	else
-	{
-		// We already have a valid actor info, just reinit it.
-		AbilityActorInfo->InitFromActor(GetOwner());
-	}
+	check(AbilityActorInfo.IsValid());
+	AbilityActorInfo->InitFromActor(AvatarActor, this);
+}
+
+void UAbilitySystemComponent::RefreshAbilityActorInfo()
+{
+	check(AbilityActorInfo.IsValid());
+	AbilityActorInfo->InitFromActor(AbilityActorInfo->Actor.Get(), this);
 }
 
-UGameplayAbility* UAbilitySystemComponent::GiveAbility(UGameplayAbility* Ability)
+UGameplayAbility* UAbilitySystemComponent::GiveAbility(UGameplayAbility* Ability, int32 InputID)
 {
 	check(Ability);
-	UGameplayAbility* OutAbility = Ability; // We may instance this ability
+	check(IsOwnerActorAuthoritative());	// Should be called on authority
 	
 	if (Ability->GetInstancingPolicy() == EGameplayAbilityInstancingPolicy::InstancedPerActor)
 	{
 		Ability = CreateNewInstanceOfAbility(Ability);
 	}
 
-	ActivatableAbilities.Add(Ability);
-	
+	ActivatableAbilities.Add(FGameplayAbilityInputIDPair(Ability, InputID));
+
+	OnGiveAbility(Ability, InputID);
 
+	return Ability;
+}
+
+void UAbilitySystemComponent::OnGiveAbility(UGameplayAbility* Ability, int32 InputID)
+{
 	for (FAbilityTriggerData TriggerData : Ability->AbilityTriggers)
 	{
 		FGameplayTag EventTag = TriggerData.TriggerTag;
 
 		if (GameplayEventTriggeredAbilities.Contains(EventTag))
 		{
-			GameplayEventTriggeredAbilities[EventTag].Add(Ability);
+			GameplayEventTriggeredAbilities[EventTag].AddUnique(Ability);
 		}
 		else
 		{
@@ -80,8 +85,6 @@ UGameplayAbility* UAbilitySystemComponent::GiveAbility(UGameplayAbility* Ability
 			GameplayEventTriggeredAbilities.Add(EventTag, Triggers);
 		}
 	}
-
-	return OutAbility;
 }
 
 UGameplayAbility* UAbilitySystemComponent::CreateNewInstanceOfAbility(UGameplayAbility *Ability)
@@ -135,12 +138,9 @@ bool UAbilitySystemComponent::ActivateAbility(TWeakObjectPtr<UGameplayAbility> A
 {
 	check(AbilityActorInfo.IsValid());
 
-	AActor * OwnerActor = GetOwner();
-	check(OwnerActor);
-
 	if (Ability.IsValid())
 	{
-		Ability.Get()->InputPressed(0, AbilityActorInfo.Get());
+		Ability.Get()->InputPressed(AbilityActorInfo.Get());
 	}
 
 	return false;
@@ -159,20 +159,23 @@ void UAbilitySystemComponent::CancelAbilitiesWithTags(const FGameplayTagContaine
 
 	struct local
 	{
-		static void CancelAbilitiesWithTags(const FGameplayTagContainer InTags, TArray<UGameplayAbility*> Abilities, const FGameplayAbilityActorInfo* InActorInfo, const FGameplayAbilityActivationInfo InActivationInfo, UGameplayAbility* InIgnore)
+		static void CancelAbilitiesWithTags(const FGameplayTagContainer InTags, TArray<FGameplayAbilityInputIDPair> Abilities, const FGameplayAbilityActorInfo* InActorInfo, const FGameplayAbilityActivationInfo InActivationInfo, UGameplayAbility* InIgnore)
 		{
 			for (int32 idx=0; idx < Abilities.Num(); ++idx)
 			{
-				UGameplayAbility *Ability = Abilities[idx];
+				UGameplayAbility *Ability = Abilities[idx].Ability;
 				if (Ability && (Ability != InIgnore) && Ability->AbilityTags.MatchesAny(InTags, false))
 				{
 					Ability->CancelAbility(InActorInfo, InActivationInfo);
+
+					/*
 					if (!Ability->HasAnyFlags(RF_ClassDefaultObject) && Ability->GetInstancingPolicy() == EGameplayAbilityInstancingPolicy::InstancedPerExecution)
 					{
-						Ability->MarkPendingKill();
+						Ability->EndAbility();
 						Abilities.RemoveAtSwap(idx);
 						idx--;
 					}
+					*/
 				}
 			}
 		}
@@ -196,23 +199,12 @@ TEXT("AbilitySystem.DenyClientActivations"),
 
 void UAbilitySystemComponent::OnRep_ActivateAbilities()
 {
-	for (UGameplayAbility * Ability : ActivatableAbilities)
+	for (FGameplayAbilityInputIDPair& Pair : ActivatableAbilities)
 	{
-		for (FAbilityTriggerData TriggerData : Ability->AbilityTriggers)
+		UGameplayAbility* Ability = Pair.Ability;
+		if (Ability)
 		{
-			FGameplayTag EventTag = TriggerData.TriggerTag;
-
-			if (GameplayEventTriggeredAbilities.Contains(EventTag))
-			{
-				TArray<TWeakObjectPtr<UGameplayAbility> > Triggers = GameplayEventTriggeredAbilities[EventTag];
-				Triggers.Add(Ability);
-			}
-			else
-			{
-				TArray<TWeakObjectPtr<UGameplayAbility> > Triggers;
-				Triggers.Add(Ability);
-				GameplayEventTriggeredAbilities.Add(EventTag, Triggers);
-			}
+			OnGiveAbility(Ability, Pair.InputID);
 		}
 	}
 }
@@ -231,9 +223,6 @@ void UAbilitySystemComponent::ServerTryActivateAbility_Implementation(UGameplayA
 	ensure(AbilityToActivate);
 	ensure(AbilityActorInfo.IsValid());
 
-	AActor * OwnerActor = GetOwner();
-	check(OwnerActor);
-
 	UGameplayAbility *InstancedAbility = NULL;
 
 	if (AbilityToActivate->TryActivateAbility(AbilityActorInfo.Get(), PredictionKey, &InstancedAbility))
@@ -295,9 +284,6 @@ void UAbilitySystemComponent::ClientActivateAbilitySucceed_Implementation(UGamep
 	ensure(AbilityToActivate);
 	ensure(AbilityActorInfo.IsValid());
 
-	AActor * OwnerActor = GetOwner();
-	check(OwnerActor);
-
 	// Fixme: We need a better way to link up/reconcile preditive replicated abilities. It would be ideal if we could predictively spawn an
 	// ability and then replace/link it with the server spawned one once the server has confirmed it.
 
@@ -343,12 +329,6 @@ void UAbilitySystemComponent::MontageBranchPoint_AbilityDecisionStop()
 {
 	if (AnimatingAbility)
 	{
-		AActor * OwnerActor = GetOwner();
-		check(OwnerActor);
-
-		FGameplayAbilityActorInfo	ActorInfo;
-		ActorInfo.InitFromActor(OwnerActor);
-
 		AnimatingAbility->MontageBranchPoint_AbilityDecisionStop(AbilityActorInfo.Get());
 	}
 }
@@ -357,9 +337,6 @@ void UAbilitySystemComponent::MontageBranchPoint_AbilityDecisionStart()
 {
 	if (AnimatingAbility)
 	{
-		AActor * OwnerActor = GetOwner();
-		check(OwnerActor);
-
 		AnimatingAbility->MontageBranchPoint_AbilityDecisionStart(AbilityActorInfo.Get());
 	}
 }
@@ -409,9 +386,12 @@ void UAbilitySystemComponent::HandleGameplayEvent(FGameplayTag EventTag, FGamepl
 	}
 }
 
-// --------------------------------------------------------------------------
+// ----------------------------------------------------------------------------------------------------------------------------------------------------
+//								Input 
+// ----------------------------------------------------------------------------------------------------------------------------------------------------
 
-void UAbilitySystemComponent::BindToInputComponent(UInputComponent *InputComponent)
+
+void UAbilitySystemComponent::BindToInputComponent(UInputComponent* InputComponent)
 {
 	static const FName ConfirmBindName(TEXT("AbilityConfirm"));
 	static const FName CancelBindName(TEXT("AbilityCancel"));
@@ -431,6 +411,71 @@ void UAbilitySystemComponent::BindToInputComponent(UInputComponent *InputCompone
 	}
 }
 
+void UAbilitySystemComponent::BindAbilityActivationToInputComponent(UInputComponent* InputComponent, FGameplayAbiliyInputBinds BindInfo)
+{
+	UEnum* EnumBinds = BindInfo.GetBindEnum();
+
+	for(int32 idx=0; idx < EnumBinds->NumEnums(); ++idx)
+	{
+		FString FullStr = EnumBinds->GetEnum(idx).ToString();
+		FString BindStr;
+
+		FullStr.Split(TEXT("::"), nullptr, &BindStr);
+
+		// Pressed event
+		{
+			FInputActionBinding AB(FName(*BindStr), IE_Pressed);
+			AB.ActionDelegate.GetDelegateForManualSet().BindUObject(this, &UAbilitySystemComponent::AbilityInputPressed, idx);
+			InputComponent->AddActionBinding(AB);
+		}
+
+		// Released event
+		{
+			FInputActionBinding AB(FName(*BindStr), IE_Released);
+			AB.ActionDelegate.GetDelegateForManualSet().BindUObject(this, &UAbilitySystemComponent::AbilityInputReleased, idx);
+			InputComponent->AddActionBinding(AB);
+		}
+	}
+
+	// Bind Confirm/Cancel. Note: these have to come last!
+	{
+		FInputActionBinding AB(FName(*BindInfo.ConfirmTargetCommand), IE_Pressed);
+		AB.ActionDelegate.GetDelegateForManualSet().BindUObject(this, &UAbilitySystemComponent::InputConfirm);
+		InputComponent->AddActionBinding(AB);
+	}
+	
+	{
+		FInputActionBinding AB(FName(*BindInfo.CancelTargetCommand), IE_Pressed);
+		AB.ActionDelegate.GetDelegateForManualSet().BindUObject(this, &UAbilitySystemComponent::InputCancel);
+		InputComponent->AddActionBinding(AB);
+	}
+
+}
+
+void UAbilitySystemComponent::AbilityInputPressed(int32 InputID)
+{
+	// FIXME: Maps or multicast delegate to actually do this
+	for (FGameplayAbilityInputIDPair& Pair: ActivatableAbilities)
+	{
+		if (Pair.InputID == InputID)
+		{
+			Pair.Ability->InputPressed(AbilityActorInfo.Get());
+		}
+	}
+}
+
+void UAbilitySystemComponent::AbilityInputReleased(int32 InputID)
+{
+	// FIXME: Maps or multicast delegate to actually do this
+	for (FGameplayAbilityInputIDPair& Pair : ActivatableAbilities)
+	{
+		if (Pair.InputID == InputID)
+		{
+			Pair.Ability->InputReleased(AbilityActorInfo.Get());
+		}
+	}
+}
+
 void UAbilitySystemComponent::InputConfirm()
 {
 	if (GetOwnerRole() != ROLE_Authority && ConfirmCallbacks.IsBound())
diff --git a/Engine/Source/Runtime/GameplayAbilities/Private/AbilitySystemGlobals.cpp b/Engine/Source/Runtime/GameplayAbilities/Private/AbilitySystemGlobals.cpp
--- a/Engine/Source/Runtime/GameplayAbilities/Private/AbilitySystemGlobals.cpp
+++ b/Engine/Source/Runtime/GameplayAbilities/Private/AbilitySystemGlobals.cpp
@@ -77,11 +77,9 @@ UAbilitySystemComponent * UAbilitySystemGlobals::GetAbilitySystemComponentFromAc
 	return NULL;
 }
 
-FGameplayAbilityActorInfo * UAbilitySystemGlobals::AllocAbilityActorInfo(AActor *Actor) const
+FGameplayAbilityActorInfo * UAbilitySystemGlobals::AllocAbilityActorInfo() const
 {
-	FGameplayAbilityActorInfo * Info = new FGameplayAbilityActorInfo();
-	Info->InitFromActor(Actor);
-	return Info;
+	return new FGameplayAbilityActorInfo();
 }
 
 /** This is just some syntax sugar to avoid calling gode to have to IGameplayAbilitiesModule::Get() */
diff --git a/Engine/Source/Runtime/GameplayAbilities/Public/Abilities/GameplayAbility.h b/Engine/Source/Runtime/GameplayAbilities/Public/Abilities/GameplayAbility.h
--- a/Engine/Source/Runtime/GameplayAbilities/Public/Abilities/GameplayAbility.h
+++ b/Engine/Source/Runtime/GameplayAbilities/Public/Abilities/GameplayAbility.h
@@ -107,10 +107,10 @@ public:
 	// ----------------------------------------------------------------------------------------------------------------
 	
 	/** Input binding. Base implementation calls TryActivateAbility */
-	virtual void InputPressed(int32 InputID, const FGameplayAbilityActorInfo* ActorInfo);
+	virtual void InputPressed(const FGameplayAbilityActorInfo* ActorInfo);
 
 	/** Input binding. Base implementation does nothing */
-	virtual void InputReleased(int32 InputID, const FGameplayAbilityActorInfo* ActorInfo);
+	virtual void InputReleased(const FGameplayAbilityActorInfo* ActorInfo);
 	
 	/**
 	 * Attempts to activate the ability.
diff --git a/Engine/Source/Runtime/GameplayAbilities/Public/Abilities/GameplayAbilityTypes.h b/Engine/Source/Runtime/GameplayAbilities/Public/Abilities/GameplayAbilityTypes.h
--- a/Engine/Source/Runtime/GameplayAbilities/Public/Abilities/GameplayAbilityTypes.h
+++ b/Engine/Source/Runtime/GameplayAbilities/Public/Abilities/GameplayAbilityTypes.h
@@ -106,23 +106,23 @@ private:
 	int32 Handle;
 };
 
-/**
-*	Not implemented yet, but we will need something like this to track ability + levels if we choose to support ability leveling in the base classes.
-*/
 
 USTRUCT()
-struct FGameplayAbilitySpec
+struct FGameplayAbilityInputIDPair
 {
 	GENERATED_USTRUCT_BODY()
 
-	UPROPERTY()
-	int32	Level;
+	FGameplayAbilityInputIDPair()
+	: Ability(nullptr), InputID(-1) { }
+
+	FGameplayAbilityInputIDPair(UGameplayAbility* InAbility, int32 InInputID)
+	: Ability(InAbility), InputID(InInputID) { }
 
 	UPROPERTY()
-	class UGameplayAbility * Ability;
+	UGameplayAbility* Ability;
 
 	UPROPERTY()
-	FGameplayAbilityHandle	Handle;
+	int32	InputID;
 };
 
 
@@ -176,7 +176,7 @@ struct GAMEPLAYABILITIES_API FGameplayAbilityActorInfo
 
 	bool IsNetAuthority() const;
 
-	virtual void InitFromActor(AActor *Actor);
+	virtual void InitFromActor(AActor *Actor, UAbilitySystemComponent* InAbilitySystemComponent);
 };
 
 /**
@@ -579,6 +579,30 @@ struct GAMEPLAYABILITIES_API FGameplayEventData
 	float Var3;
 };
 
+/** 
+ *	Structure that tells AbilitySystemComponent what to bind to an InputComponent (see BindAbilityActivationToInputComponent) 
+ *	
+ */
+struct FGameplayAbiliyInputBinds
+{
+	FGameplayAbiliyInputBinds(FString InConfirmTargetCommand, FString InCancelTargetCommand, FString InEnumName)
+	: ConfirmTargetCommand(InConfirmTargetCommand)
+	, CancelTargetCommand(InCancelTargetCommand)
+	, EnumName(InEnumName)
+	{ }
+
+	/** Defines command string that will be bound to Confirm Targeting */
+	FString ConfirmTargetCommand;
+
+	/** Defines command string that will be bound to Cancel Targeting */
+	FString CancelTargetCommand;
+
+	/** Returns enum to use for ability bings. E.g., "Ability1"-"Ability8" input commands will be bound to ability activations inside the AbiltiySystemComponent */
+	FString	EnumName;
+
+	UEnum* GetBindEnum() { return FindObject<UEnum>(ANY_PACKAGE, *EnumName); }
+};
+
 /** Generic callback for returning when target data is available */
 DECLARE_MULTICAST_DELEGATE_OneParam(FAbilityTargetData, FGameplayAbilityTargetDataHandle);
 
diff --git a/Engine/Source/Runtime/GameplayAbilities/Public/Abilities/GameplayAbility_CharacterJump.h b/Engine/Source/Runtime/GameplayAbilities/Public/Abilities/GameplayAbility_CharacterJump.h
--- a/Engine/Source/Runtime/GameplayAbilities/Public/Abilities/GameplayAbility_CharacterJump.h
+++ b/Engine/Source/Runtime/GameplayAbilities/Public/Abilities/GameplayAbility_CharacterJump.h
@@ -19,7 +19,7 @@ public:
 	
 	virtual void ActivateAbility(const FGameplayAbilityActorInfo* OwnerInfo, const FGameplayAbilityActivationInfo ActivationInfo) override;
 
-	virtual void InputReleased(int32 InputID, const FGameplayAbilityActorInfo* ActorInfo) override;
+	virtual void InputReleased(const FGameplayAbilityActorInfo* ActorInfo) override;
 
 	virtual void CancelAbility(const FGameplayAbilityActorInfo* ActorInfo, const FGameplayAbilityActivationInfo ActivationInfo) override;
 };
\ No newline at end of file
diff --git a/Engine/Source/Runtime/GameplayAbilities/Public/AbilitySystemComponent.h b/Engine/Source/Runtime/GameplayAbilities/Public/AbilitySystemComponent.h
--- a/Engine/Source/Runtime/GameplayAbilities/Public/AbilitySystemComponent.h
+++ b/Engine/Source/Runtime/GameplayAbilities/Public/AbilitySystemComponent.h
@@ -295,7 +295,10 @@ class GAMEPLAYABILITIES_API UAbilitySystemComponent : public UActorComponent
 	 *	
 	 */
 
-	UGameplayAbility* GiveAbility(UGameplayAbility* Ability);
+	UGameplayAbility* GiveAbility(UGameplayAbility* Ability, int32 InputID);
+
+	/** Will be called from GiveAbility or from OnRep. Initializes events (triggers and inputs) with the given ability */
+	void OnGiveAbility(UGameplayAbility* Ability, int32 InputID);
 
 	UGameplayAbility* CreateNewInstanceOfAbility(UGameplayAbility* Ability);
 
@@ -321,7 +324,8 @@ class GAMEPLAYABILITIES_API UAbilitySystemComponent : public UActorComponent
 	 *	instancing or anything else that the AbilitySystemComponent would provide, then it doesn't need the component to function.
 	 */ 
 	UPROPERTY(ReplicatedUsing = OnRep_ActivateAbilities, BlueprintReadOnly, Category = "Abilities")
-	TArray<UGameplayAbility*>	ActivatableAbilities;
+	TArray<FGameplayAbilityInputIDPair>	ActivatableAbilities;
+
 	
 	UFUNCTION()
 	void	OnRep_ActivateAbilities();
@@ -352,7 +356,13 @@ class GAMEPLAYABILITIES_API UAbilitySystemComponent : public UActorComponent
 
 	// ----------------------------------------------------------------------------------------------------------------
 
-	virtual void BindToInputComponent(UInputComponent *InputComponent);
+	virtual void BindToInputComponent(UInputComponent* InputComponent);
+	
+	virtual void BindAbilityActivationToInputComponent(UInputComponent* InputComponent, FGameplayAbiliyInputBinds BindInfo);
+
+	void AbilityInputPressed(int32 InputID);
+
+	void AbilityInputReleased(int32 InputID);
 
 	UFUNCTION(BlueprintCallable, Category="Abilities")
 	void InputConfirm();
@@ -403,7 +413,18 @@ class GAMEPLAYABILITIES_API UAbilitySystemComponent : public UActorComponent
 	/** Cached off data about the owning actor that abilities will need to frequently access (movement component, mesh component, anim instance, etc) */
 	TSharedPtr<FGameplayAbilityActorInfo>	AbilityActorInfo;
 
-	void InitAbilityActorInfo();
+
+	/**
+	 *	Initialized the Abilities' ActorInfo - the structure that holds information about who we are acting on and who controls us.
+	 *		AvatarActor is what physical actor in the world we are acting on. Usually a Pawn but it could be a Tower, Building, Turret, etc.
+	 */
+	void InitAbilityActorInfo(AActor* AvatarActor);
+
+	/**
+	 *	This will refresh the Ability's ActorInfo structure based on the current ActorInfo. That is, AvatarActor will be the same but we will look for new
+	 *	AnimInstance, MovementComponent, PlayerController, etc.
+	 */	
+	void RefreshAbilityActorInfo();
 
 	// -----------------------------------------------------------------------------
 
diff --git a/Engine/Source/Runtime/GameplayAbilities/Public/AbilitySystemGlobals.h b/Engine/Source/Runtime/GameplayAbilities/Public/AbilitySystemGlobals.h
--- a/Engine/Source/Runtime/GameplayAbilities/Public/AbilitySystemGlobals.h
+++ b/Engine/Source/Runtime/GameplayAbilities/Public/AbilitySystemGlobals.h
@@ -46,7 +46,8 @@ class GAMEPLAYABILITIES_API UAbilitySystemGlobals : public UObject
 	 */
 	virtual UAbilitySystemComponent * GetAbilitySystemComponentFromActor(AActor *Actor) const;
 
-	virtual FGameplayAbilityActorInfo * AllocAbilityActorInfo(AActor *Actor) const;
+	/** Should allocate a project specific AbilityActorInfo struct. Caller is responsible for dellocation */
+	virtual FGameplayAbilityActorInfo * AllocAbilityActorInfo() const;
 
 	UFunction* GetGameplayCueFunction(const FGameplayTag &Tag, UClass* Class, FName &MatchedTag);
 
diff --git a/Engine/Source/Runtime/GameplayAbilities/Public/GameplayAbilitySet.h b/Engine/Source/Runtime/GameplayAbilities/Public/GameplayAbilitySet.h
--- a/Engine/Source/Runtime/GameplayAbilities/Public/GameplayAbilitySet.h
+++ b/Engine/Source/Runtime/GameplayAbilities/Public/GameplayAbilitySet.h
@@ -4,53 +4,62 @@
 
 #include "GameplayAbilitySet.generated.h"
 
-USTRUCT()
-struct FGameplayAbilityBindInfoCommandIDPair
-{
-	GENERATED_USTRUCT_BODY()
-
-	// We want this to be a drop down of valid lists but needs to be defined per game
-	UPROPERTY(EditAnywhere, Category = BindInfo)
-	FString		CommandString;
+class UGameplayAbility;
+class UAbilitySystemComponent;
 
-	UPROPERTY(EditAnywhere, Category = BindInfo)
-	int32		InputID;
-};
+/**
+ *	This is an example input binding enum for GameplayAbilities. Your project may want to create its own.
+ *	The MetaData default bind keys (LMB, RMB, Q, E, etc) are a convenience for designers setting up abilities sets
+ *	or whatever other data you have that gives an ability with a default key binding. Actual key binding is up to DefaultInput.ini
+ *	
+ *	E.g., "Ability1" is the command string that is bound to AbilitySystemComponent::ActivateAbility(1). The Meta data only *suggests*
+ *	that you are binding "Ability1" to LMB by default in your projects DefaultInput.ini.
+ */
+UENUM(BlueprintType)
+namespace EGameplayAbilityInputBinds
+{
+	enum Type
+	{
+		Ability1				UMETA(DisplayName = "Ability1 (LMB)"),
+		Ability2				UMETA(DisplayName = "Ability2 (RMB)"),
+		Ability3				UMETA(DisplayName = "Ability3 (Q)"),
+		Ability4				UMETA(DisplayName = "Ability4 (E)"),
+		Ability5				UMETA(DisplayName = "Ability5 (R)"),
+		Ability6				UMETA(DisplayName = "Ability6 (Shift)"),
+		Ability7				UMETA(DisplayName = "Ability7 (Space)"),
+		Ability8				UMETA(DisplayName = "Ability7 (B)"),
+	};
+}
 
+/**
+ *	Example struct that pairs a enum input command to a GameplayAbilityClass.6
+ */
 USTRUCT()
 struct FGameplayAbilityBindInfo
 {
 	GENERATED_USTRUCT_BODY()
 
-	UPROPERTY(EditAnywhere, Category=BindInfo)
-	TArray<FGameplayAbilityBindInfoCommandIDPair>	CommandBindings;
+	UPROPERTY(EditAnywhere, Category = BindInfo)
+	TEnumAsByte<EGameplayAbilityInputBinds::Type>	Command;
 
-	UPROPERTY(EditAnywhere, Category=BindInfo)
+	UPROPERTY(EditAnywhere, Category = BindInfo)
 	TSubclassOf<UGameplayAbility>	GameplayAbilityClass;
 };
 
 /**
- * UGameplayAbilitySet
- *	This contains a list of abilities along with key bindings. This will be very game specific, so it is expected for games to override this class or create their own.
+ *	This is an example DataAsset that could be used for defining a set of abilities to give to an AbilitySystemComponent and bind to an input command.
+ *	Your project is free to implement this however it wants!
  *	
- *	Functions like InitializeAbilities and BindInputComponentToAbilities are setup to act on the AbilitySystemComponent (rather than AbilitySystemComponent processing the ability set).
- *	Ideally this avoids games having to implement their own UAbilitySystemComponent and UGameplayAbilitySet, and they can just implement thei own UGameplayAbilitySet.
  *	
- *	In the end, this class just facilitates adding stuff to AbilitySystemComponent::ActivatableAbilities. Games are free to do this however they want.
  */
-UCLASS(Blueprintable)
+UCLASS()
 class GAMEPLAYABILITIES_API UGameplayAbilitySet : public UDataAsset
 {
 	GENERATED_UCLASS_BODY()
 
-public:
-
-	UPROPERTY(EditAnywhere, Category=AbilitySet)
+	
+	UPROPERTY(EditAnywhere, Category = AbilitySet)
 	TArray<FGameplayAbilityBindInfo>	Abilities;
 
-	/** Give all of the abilities in this set to given AbilitySystemComponent */
-	virtual void InitializeAbilities(UAbilitySystemComponent *AbilitySystemComponent) const;
-
-	/** Binds the ability set directly to this inputcomponent and AbilitySystemComponent. This is a convenience method, not all games will want to bind input directly to abilities. */
-	virtual void BindInputComponentToAbilities(UInputComponent *InputComponent, UAbilitySystemComponent *AbilitySystemComponent) const;
+	void GiveAbilities(UAbilitySystemComponent* AbilitySystemComponent) const;
 };
\ No newline at end of file
