compile fix

diff --git a/Engine/Source/Runtime/GameplayAbilities/Public/Abilities/GameplayAbility.h b/Engine/Source/Runtime/GameplayAbilities/Public/Abilities/GameplayAbility.h
--- a/Engine/Source/Runtime/GameplayAbilities/Public/Abilities/GameplayAbility.h
+++ b/Engine/Source/Runtime/GameplayAbilities/Public/Abilities/GameplayAbility.h
@@ -5,7 +5,7 @@
 #include "GameplayTagContainer.h"
 #include "GameplayTagAssetInterface.h"
 #include "Runtime/Engine/Classes/Animation/AnimInstance.h"
-#include "Abilities/GameplayAbilityTypes.h"
+#include "GameplayAbilityTypes.h"
 #include "GameplayEffect.h"
 #include "GameplayAbility.generated.h"
 
diff --git a/Engine/Source/Runtime/GameplayAbilities/Public/Abilities/GameplayAbilityTypes.h b/Engine/Source/Runtime/GameplayAbilities/Public/Abilities/GameplayAbilityTypes.h
--- a/Engine/Source/Runtime/GameplayAbilities/Public/Abilities/GameplayAbilityTypes.h
+++ b/Engine/Source/Runtime/GameplayAbilities/Public/Abilities/GameplayAbilityTypes.h
@@ -10,6 +10,7 @@ class UAbilitySystemComponent;
 class UGameplayAbility;
 class AGameplayAbilityTargetActor;
 class UAbilityTask;
+class UAttributeSet;
 
 
 UENUM(BlueprintType)
diff --git a/Engine/Source/Runtime/GameplayAbilities/Public/AbilitySystemComponent.h b/Engine/Source/Runtime/GameplayAbilities/Public/AbilitySystemComponent.h
--- a/Engine/Source/Runtime/GameplayAbilities/Public/AbilitySystemComponent.h
+++ b/Engine/Source/Runtime/GameplayAbilities/Public/AbilitySystemComponent.h
@@ -17,7 +17,7 @@ struct GAMEPLAYABILITIES_API FAttributeDefaults
 	GENERATED_USTRUCT_BODY()
 
 	UPROPERTY(EditAnywhere, Category="AttributeTest")
-	TSubclassOf<class UAttributeSet> Attributes;
+	TSubclassOf<UAttributeSet> Attributes;
 
 	UPROPERTY(EditAnywhere, BlueprintReadWrite, Category="AttributeTest")
 	class UDataTable*	DefaultStartingTable;
