Mesh targeting actually uses sockets, has some failsafe capability.

diff --git a/Engine/Source/Runtime/GameplayAbilities/Private/Abilities/GameplayAbility.cpp b/Engine/Source/Runtime/GameplayAbilities/Private/Abilities/GameplayAbility.cpp
--- a/Engine/Source/Runtime/GameplayAbilities/Private/Abilities/GameplayAbility.cpp
+++ b/Engine/Source/Runtime/GameplayAbilities/Private/Abilities/GameplayAbility.cpp
@@ -706,24 +706,11 @@ bool FGameplayAbilityTargetDataHandle::NetSerialize(FArchive& Ar, class UPackage
 
 bool FGameplayAbilityTargetData_Mesh::NetSerialize(FArchive& Ar, class UPackageMap* Map, bool& bOutSuccess)
 {
+	//SourceActor can be used as a backup if the component isn't found.
 	Ar << SourceActor;
-	if (SourceActor)
-	{
-		Ar << SourceSocketName;
-	}
-	else
-	{
-		StartPoint.NetSerialize(Ar, Map, bOutSuccess);
-	}
-	Ar << bUseTargetPoint;
-	if (bUseTargetPoint)
-	{
-		TargetPoint.NetSerialize(Ar, Map, bOutSuccess);
-	}
-	else
-	{
-		AimDirection.NetSerialize(Ar, Map, bOutSuccess);
-	}
+	Ar << SourceComponent;
+	Ar << SourceSocketName;
+	TargetPoint.NetSerialize(Ar, Map, bOutSuccess);
 
 	bOutSuccess = true;
 	return true;
diff --git a/Engine/Source/Runtime/GameplayAbilities/Public/Abilities/GameplayAbilityTypes.h b/Engine/Source/Runtime/GameplayAbilities/Public/Abilities/GameplayAbilityTypes.h
--- a/Engine/Source/Runtime/GameplayAbilities/Public/Abilities/GameplayAbilityTypes.h
+++ b/Engine/Source/Runtime/GameplayAbilities/Public/Abilities/GameplayAbilityTypes.h
@@ -323,44 +323,34 @@ struct GAMEPLAYABILITIES_API FGameplayAbilityTargetData
 	virtual FString ToString() const;
 };
 
-USTRUCT()
+USTRUCT(BlueprintType)
 struct GAMEPLAYABILITIES_API FGameplayAbilityTargetData_Mesh : public FGameplayAbilityTargetData
 {
 	GENERATED_USTRUCT_BODY()
 
 	FGameplayAbilityTargetData_Mesh()
 	: SourceActor(NULL)
-	, AimDirection(FVector::ZeroVector)
-	, bUseTargetPoint(false)
+	, SourceComponent(NULL)
 	{}
 
 
 	//UPROPERTY(BlueprintReadWrite, EditAnywhere, Category = Projectile)
 
-	/** If this is valid, projectile will spawn StartPoint will be based on this actor. */
-	UPROPERTY()
+	/** Actor who owns the named component. Actor's location is used as start point if component cannot be found. */
+	UPROPERTY(BlueprintReadWrite, EditAnywhere, Category = Targeting)
 	AActor* SourceActor;
 
-	/** If SourceActor is valid, this will allow a socket to be used instead of the actor's location. */
-	UPROPERTY()
-		FName SourceSocketName;
-
-	/** Projectile will be spawned here if SourceActor is NULL or invalid. */
-	UPROPERTY()
-		FVector_NetQuantize StartPoint;
-
-	/** Projectile will be fired in this direction. */
-	UPROPERTY()
-		FVector_NetQuantizeNormal AimDirection;
-
-	/** Projectile will be aimed from its spawn point toward this point if UseTargetPoint is true. */
-	UPROPERTY()
-		FVector_NetQuantize TargetPoint;
+	/** Local skeletal mesh component that holds the socket. */
+	UPROPERTY(BlueprintReadWrite, EditAnywhere, Category = Targeting)
+	USkeletalMeshComponent* SourceComponent;
 
-	/** If this is true, AimDirection will be ignored, and the projectile will be aimed at TargetPoint. */
-	UPROPERTY()
-		bool bUseTargetPoint;
+	/** If SourceActor and SourceComponent are valid, this is the name of the socket that will be used instead of the actor's location. */
+	UPROPERTY(BlueprintReadWrite, EditAnywhere, Category=Targeting)
+	FName SourceSocketName;
 
+	/** Point being targeted. */
+	UPROPERTY(BlueprintReadWrite, EditAnywhere, Category = Targeting)
+	FVector_NetQuantize TargetPoint;
 
 
 	//virtual TArray<AActor*>	GetActors() const
@@ -382,14 +372,11 @@ struct GAMEPLAYABILITIES_API FGameplayAbilityTargetData_Mesh : public FGameplayA
 		if (SourceActor)
 		{
 			FTransform ReturnTransform = SourceActor->GetTransform();
-			if (bUseTargetPoint)
-			{
-				ReturnTransform.SetRotation((TargetPoint - ReturnTransform.GetLocation()).SafeNormal().Rotation().Quaternion());
-			}
-			else
+			if (SourceComponent)
 			{
-				ReturnTransform.SetRotation(AimDirection.Rotation().Quaternion());
+				ReturnTransform.SetLocation(SourceComponent->GetSocketLocation(SourceSocketName));
 			}
+			ReturnTransform.SetRotation((TargetPoint - ReturnTransform.GetLocation()).SafeNormal().Rotation().Quaternion());
 			return ReturnTransform;
 		}
 		return FTransform::Identity;
@@ -404,11 +391,7 @@ struct GAMEPLAYABILITIES_API FGameplayAbilityTargetData_Mesh : public FGameplayA
 
 	virtual FVector GetEndPoint() const
 	{
-		if (bUseTargetPoint)
-		{
-			return TargetPoint;
-		}
-		return FVector::ZeroVector;
+		return TargetPoint;
 	}
 
 	// -------------------------------------
