Engine: Comment out broken code to get build going

diff --git a/Engine/Source/Runtime/GameplayAbilities/Private/Abilities/GameplayAbilityTargetActor_SingleLineTrace.cpp b/Engine/Source/Runtime/GameplayAbilities/Private/Abilities/GameplayAbilityTargetActor_SingleLineTrace.cpp
--- a/Engine/Source/Runtime/GameplayAbilities/Private/Abilities/GameplayAbilityTargetActor_SingleLineTrace.cpp
+++ b/Engine/Source/Runtime/GameplayAbilities/Private/Abilities/GameplayAbilityTargetActor_SingleLineTrace.cpp
@@ -165,8 +165,12 @@ FGameplayAbilityTargetDataHandle AGameplayAbilityTargetActor_SingleLineTrace::Ma
 
 	FGameplayAbilityTargetData_Mesh* ReturnData = new FGameplayAbilityTargetData_Mesh();
 	ReturnData->SourceActor = SourceActor;
+#if 0
 	ReturnData->AimDirection = (HitResult.Location - SourceActor->GetActorLocation()).SafeNormal();
+#endif
 	ReturnData->TargetPoint = HitResult.Location;
+#if 0
 	ReturnData->bUseTargetPoint = true;
+#endif
 	return FGameplayAbilityTargetDataHandle(ReturnData);
 }
\ No newline at end of file
