Restoring to old targeting for this ability. (Cleanup.)

diff --git a/Engine/Source/Runtime/GameplayAbilities/Private/Abilities/GameplayAbilityTargetActor_SingleLineTrace.cpp b/Engine/Source/Runtime/GameplayAbilities/Private/Abilities/GameplayAbilityTargetActor_SingleLineTrace.cpp
--- a/Engine/Source/Runtime/GameplayAbilities/Private/Abilities/GameplayAbilityTargetActor_SingleLineTrace.cpp
+++ b/Engine/Source/Runtime/GameplayAbilities/Private/Abilities/GameplayAbilityTargetActor_SingleLineTrace.cpp
@@ -159,18 +159,7 @@ FGameplayAbilityTargetDataHandle AGameplayAbilityTargetActor_SingleLineTrace::Ma
 {
 	/** Note this is cleaned up by the FGameplayAbilityTargetDataHandle (via an internal TSharedPtr) */
 	
-	//FGameplayAbilityTargetData_SingleTargetHit* ReturnData = new FGameplayAbilityTargetData_SingleTargetHit();
-	//ReturnData->HitResult = HitResult;
-	//return FGameplayAbilityTargetDataHandle(ReturnData);
-
-	FGameplayAbilityTargetData_Mesh* ReturnData = new FGameplayAbilityTargetData_Mesh();
-	ReturnData->SourceActor = SourceActor;
-#if 0
-	ReturnData->AimDirection = (HitResult.Location - SourceActor->GetActorLocation()).SafeNormal();
-#endif
-	ReturnData->TargetPoint = HitResult.Location;
-#if 0
-	ReturnData->bUseTargetPoint = true;
-#endif
+	FGameplayAbilityTargetData_SingleTargetHit* ReturnData = new FGameplayAbilityTargetData_SingleTargetHit();
+	ReturnData->HitResult = HitResult;
 	return FGameplayAbilityTargetDataHandle(ReturnData);
 }
\ No newline at end of file
