Hot-swap targeting stuff, made some functions on AnimInstance public to enable skeletal mesh components to be accessible in Orion abilities.

diff --git a/Engine/Source/Runtime/Engine/Classes/Animation/AnimInstance.h b/Engine/Source/Runtime/Engine/Classes/Animation/AnimInstance.h
--- a/Engine/Source/Runtime/Engine/Classes/Animation/AnimInstance.h
+++ b/Engine/Source/Runtime/Engine/Classes/Animation/AnimInstance.h
@@ -310,7 +310,7 @@ public:
 	UFUNCTION(BlueprintPure, Category = "Animation")
 	virtual APawn* TryGetPawnOwner();
 
-protected:
+public:
 	/** Returns the owning actor of this AnimInstance */
 	UFUNCTION(BlueprintPure, Category="Animation")
 	AActor* GetOwningActor();
diff --git a/Engine/Source/Runtime/GameplayAbilities/Private/AbilitySystemBlueprintLibrary.cpp b/Engine/Source/Runtime/GameplayAbilities/Private/AbilitySystemBlueprintLibrary.cpp
--- a/Engine/Source/Runtime/GameplayAbilities/Private/AbilitySystemBlueprintLibrary.cpp
+++ b/Engine/Source/Runtime/GameplayAbilities/Private/AbilitySystemBlueprintLibrary.cpp
@@ -32,6 +32,14 @@ void UAbilitySystemBlueprintLibrary::ApplyGameplayEffectToTargetData(FGameplayAb
 	}
 }
 
+FGameplayAbilityTargetDataHandle UAbilitySystemBlueprintLibrary::AbilityTargetDataHandleFromAbilityTargetDataMesh(FGameplayAbilityTargetData_Mesh Data)
+{
+	FGameplayAbilityTargetData_Mesh*	NewData = new FGameplayAbilityTargetData_Mesh(Data);
+	FGameplayAbilityTargetDataHandle	Handle;
+	Handle.Data = TSharedPtr<FGameplayAbilityTargetData_Mesh>(NewData);
+	return Handle;
+}
+
 FGameplayAbilityTargetDataHandle UAbilitySystemBlueprintLibrary::AbilityTargetDataFromHitResult(FHitResult HitResult)
 {
 	// Construct TargetData
@@ -44,6 +52,16 @@ FGameplayAbilityTargetDataHandle UAbilitySystemBlueprintLibrary::AbilityTargetDa
 	return Handle;
 }
 
+bool UAbilitySystemBlueprintLibrary::TargetDataHasHitResult(FGameplayAbilityTargetDataHandle TargetData)
+{
+	FGameplayAbilityTargetData* Data = TargetData.Data.Get();
+	if (Data)
+	{
+		return Data->HasHitResult();
+	}
+	return false;
+}
+
 FHitResult UAbilitySystemBlueprintLibrary::GetHitResultFromTargetData(FGameplayAbilityTargetDataHandle TargetData)
 {
 	FGameplayAbilityTargetData* Data = TargetData.Data.Get();
@@ -59,9 +77,20 @@ FHitResult UAbilitySystemBlueprintLibrary::GetHitResultFromTargetData(FGameplayA
 	return FHitResult();
 }
 
+bool UAbilitySystemBlueprintLibrary::TargetDataHasEndPoint(FGameplayAbilityTargetDataHandle TargetData)
+{
+	FGameplayAbilityTargetData* Data = TargetData.Data.Get();
+	if (Data)
+	{
+		const FHitResult* HitResultPtr = Data->GetHitResult();
+		return (Data->HasHitResult() || Data->HasEndPoint());
+	}
+	return false;
+}
+
 FVector UAbilitySystemBlueprintLibrary::GetTargetDataEndPoint(FGameplayAbilityTargetDataHandle TargetData)
 {
-	FGameplayAbilityTargetData * Data = TargetData.Data.Get();
+	FGameplayAbilityTargetData* Data = TargetData.Data.Get();
 	if (Data)
 	{
 		const FHitResult* HitResultPtr = Data->GetHitResult();
@@ -76,7 +105,7 @@ FVector UAbilitySystemBlueprintLibrary::GetTargetDataEndPoint(FGameplayAbilityTa
 				return HitResultPtr->TraceEnd;
 			}
 		}
-		else if(Data->HasEndPoint())
+		else if (Data->HasEndPoint())
 		{
 			return Data->GetEndPoint();
 		}
diff --git a/Engine/Source/Runtime/GameplayAbilities/Public/AbilitySystemBlueprintLibrary.h b/Engine/Source/Runtime/GameplayAbilities/Public/AbilitySystemBlueprintLibrary.h
--- a/Engine/Source/Runtime/GameplayAbilities/Public/AbilitySystemBlueprintLibrary.h
+++ b/Engine/Source/Runtime/GameplayAbilities/Public/AbilitySystemBlueprintLibrary.h
@@ -28,12 +28,21 @@ class UAbilitySystemBlueprintLibrary : public UBlueprintFunctionLibrary
 	UFUNCTION(BlueprintCallable, Category = "Ability|TargetData")
 	static void ApplyGameplayEffectToTargetData(FGameplayAbilityTargetDataHandle Target, UGameplayEffect *GameplayEffect, const FGameplayAbilityActorInfo InstigatorInfo);
 
+	UFUNCTION(BlueprintPure, Category = "Ability|TargetData")
+	static FGameplayAbilityTargetDataHandle	AbilityTargetDataHandleFromAbilityTargetDataMesh(FGameplayAbilityTargetData_Mesh Data);
+
 	UFUNCTION(BlueprintPure, Category = "Ability|TargetData")
 	static FGameplayAbilityTargetDataHandle	AbilityTargetDataFromHitResult(FHitResult HitResult);
 
+	UFUNCTION(BlueprintPure, Category = "Ability|TargetData")
+	static bool TargetDataHasHitResult(FGameplayAbilityTargetDataHandle HitResult);
+
 	UFUNCTION(BlueprintPure, Category = "Ability|TargetData")
 	static FHitResult GetHitResultFromTargetData(FGameplayAbilityTargetDataHandle HitResult);
 
+	UFUNCTION(BlueprintPure, Category = "Ability|TargetData")
+	static bool TargetDataHasEndPoint(FGameplayAbilityTargetDataHandle TargetData);
+
 	UFUNCTION(BlueprintPure, Category = "Ability|TargetData")
 	static FVector GetTargetDataEndPoint(FGameplayAbilityTargetDataHandle TargetData);
 
@@ -50,6 +59,7 @@ class UAbilitySystemBlueprintLibrary : public UBlueprintFunctionLibrary
 	UFUNCTION(BlueprintPure, Category = "Ability|GameplayCue")
 	static bool HasHitResult(FGameplayCueParameters Parameters);
 
+
 	// -------------------------------------------------------------------------------
 
 };
\ No newline at end of file
