Cleaner mesh targeting data transmission/usage.

diff --git a/Engine/Source/Runtime/GameplayAbilities/Private/Abilities/GameplayAbilityTargetActor_SingleLineTrace.cpp b/Engine/Source/Runtime/GameplayAbilities/Private/Abilities/GameplayAbilityTargetActor_SingleLineTrace.cpp
--- a/Engine/Source/Runtime/GameplayAbilities/Private/Abilities/GameplayAbilityTargetActor_SingleLineTrace.cpp
+++ b/Engine/Source/Runtime/GameplayAbilities/Private/Abilities/GameplayAbilityTargetActor_SingleLineTrace.cpp
@@ -27,6 +27,8 @@ void AGameplayAbilityTargetActor_SingleLineTrace::GetLifetimeReplicatedProps(TAr
 	Super::GetLifetimeReplicatedProps(OutLifetimeProps);
 	DOREPLIFETIME(AGameplayAbilityTargetActor_SingleLineTrace, bDebug);
 	DOREPLIFETIME(AGameplayAbilityTargetActor_SingleLineTrace, SourceActor);
+	DOREPLIFETIME(AGameplayAbilityTargetActor_SingleLineTrace, SourceComponent);
+	DOREPLIFETIME(AGameplayAbilityTargetActor_SingleLineTrace, SourceSocketName);
 }
 
 void AGameplayAbilityTargetActor_SingleLineTrace::EndPlay(const EEndPlayReason::Type EndPlayReason)
@@ -40,7 +42,7 @@ void AGameplayAbilityTargetActor_SingleLineTrace::EndPlay(const EEndPlayReason::
 }
 
 
-FHitResult AGameplayAbilityTargetActor_SingleLineTrace::PerformTrace(AActor *InSourceActor) const
+FHitResult AGameplayAbilityTargetActor_SingleLineTrace::PerformTrace(AActor* InSourceActor, USkeletalMeshComponent* InSourceComponent, FName InSourceSocketName) const
 {
 	static const FName LineTraceSingleName(TEXT("AGameplayAbilityTargetActor_SingleLineTrace"));
 	bool bTraceComplex = false;
@@ -68,6 +70,13 @@ FHitResult AGameplayAbilityTargetActor_SingleLineTrace::PerformTrace(AActor *InS
 	FVector TraceStart = InSourceActor->GetActorLocation();
 	FVector TraceEnd = TraceStart + (AimDirection * 3000.f);
 
+	//If we're using a socket, adjust the starting location and aim direction after the end position has been found. This way we can still aim with the camera, then fire accurately from the socket.
+	if (InSourceComponent)
+	{
+		TraceStart = InSourceComponent->GetSocketLocation(InSourceSocketName);		//If socket name is invalid, this just gets the location of the component.
+		AimDirection = (TraceEnd - TraceStart).SafeNormal();
+	}
+
 	// ------------------------------------------------------
 
 	FHitResult ReturnHitResult;
@@ -83,14 +92,18 @@ FGameplayAbilityTargetDataHandle AGameplayAbilityTargetActor_SingleLineTrace::St
 {
 	AActor* StaticSourceActor = ActorInfo->Actor.Get();
 	check(StaticSourceActor);
-	
-	return MakeTargetData(PerformTrace(StaticSourceActor));
+	UAnimInstance* AnimInstance = ActorInfo->AnimInstance.Get();
+	USkeletalMeshComponent* StaticSourceComponent = AnimInstance ? AnimInstance->GetOwningComponent() : NULL;
+
+	return MakeTargetData(PerformTrace(StaticSourceActor, StaticSourceComponent, SourceSocketName));
 }
 
 void AGameplayAbilityTargetActor_SingleLineTrace::StartTargeting(UGameplayAbility* InAbility)
 {
 	Ability = InAbility;
 	SourceActor = InAbility->GetCurrentActorInfo()->Actor.Get();
+	UAnimInstance* AnimInstance = InAbility->GetCurrentActorInfo()->AnimInstance.Get();
+	SourceComponent = AnimInstance ? AnimInstance->GetOwningComponent() : NULL;
 
 	// We can bind directly to our ASC's confirm/cancel events, or wait to be told from an outside source to confirm or cancel
 	if (bBindToConfirmCancelInputs && (MasterPC && MasterPC->IsLocalController()))
@@ -117,7 +130,7 @@ void AGameplayAbilityTargetActor_SingleLineTrace::Tick(float DeltaSeconds)
 	// very temp - do a mostly hardcoded trace from the source actor
 	if (SourceActor)
 	{
-		FHitResult HitResult = PerformTrace(SourceActor);
+		FHitResult HitResult = PerformTrace(SourceActor, SourceComponent, SourceSocketName);
 		FVector EndPoint = HitResult.Component.IsValid() ? HitResult.ImpactPoint : HitResult.TraceEnd;
 
 		if (bDebug)
@@ -148,7 +161,7 @@ void AGameplayAbilityTargetActor_SingleLineTrace::ConfirmTargeting()
 	{
 		bDebug = false;
 		bAutoFire = false;
-		FGameplayAbilityTargetDataHandle Handle = MakeTargetData(PerformTrace(SourceActor));
+		FGameplayAbilityTargetDataHandle Handle = MakeTargetData(PerformTrace(SourceActor, SourceComponent, SourceSocketName));
 		TargetDataReadyDelegate.Broadcast(Handle);
 	}
 
@@ -157,9 +170,20 @@ void AGameplayAbilityTargetActor_SingleLineTrace::ConfirmTargeting()
 
 FGameplayAbilityTargetDataHandle AGameplayAbilityTargetActor_SingleLineTrace::MakeTargetData(FHitResult HitResult) const
 {
-	/** Note this is cleaned up by the FGameplayAbilityTargetDataHandle (via an internal TSharedPtr) */
-	
-	FGameplayAbilityTargetData_SingleTargetHit* ReturnData = new FGameplayAbilityTargetData_SingleTargetHit();
-	ReturnData->HitResult = HitResult;
-	return FGameplayAbilityTargetDataHandle(ReturnData);
+	/** Note these are cleaned up by the FGameplayAbilityTargetDataHandle (via an internal TSharedPtr) */
+	if (SourceComponent)
+	{
+		FGameplayAbilityTargetData_Mesh* ReturnData = new FGameplayAbilityTargetData_Mesh();
+		ReturnData->SourceActor = SourceActor;
+		ReturnData->SourceComponent = SourceComponent;
+		ReturnData->SourceSocketName = SourceSocketName;
+		ReturnData->TargetPoint = HitResult.Location;
+		return FGameplayAbilityTargetDataHandle(ReturnData);
+	}
+	else
+	{
+		FGameplayAbilityTargetData_SingleTargetHit* ReturnData = new FGameplayAbilityTargetData_SingleTargetHit();
+		ReturnData->HitResult = HitResult;
+		return FGameplayAbilityTargetDataHandle(ReturnData);
+	}
 }
\ No newline at end of file
diff --git a/Engine/Source/Runtime/GameplayAbilities/Public/Abilities/GameplayAbilityTargetActor_SingleLineTrace.h b/Engine/Source/Runtime/GameplayAbilities/Public/Abilities/GameplayAbilityTargetActor_SingleLineTrace.h
--- a/Engine/Source/Runtime/GameplayAbilities/Public/Abilities/GameplayAbilityTargetActor_SingleLineTrace.h
+++ b/Engine/Source/Runtime/GameplayAbilities/Public/Abilities/GameplayAbilityTargetActor_SingleLineTrace.h
@@ -27,9 +27,14 @@ public:
 	UPROPERTY(BlueprintReadOnly, Replicated, Category = Targeting)
 	bool bDebug;
 
-	//For remote clients, the actor who starts the trace (also gives us trace direction)
 	UPROPERTY(BlueprintReadOnly, Replicated, Category = Targeting)
-	AActor* SourceActor;		//Maybe add a socket to this?
+	AActor* SourceActor;
+
+	UPROPERTY(BlueprintReadOnly, Replicated, Category = Targeting)
+	USkeletalMeshComponent* SourceComponent;
+
+	UPROPERTY(EditDefaultsOnly, Replicated, Category = Targeting)
+	FName SourceSocketName;
 
 	UPROPERTY(BlueprintReadOnly, meta=(ExposeOnSpawn=true), Category=Projectile)
 	bool bBindToConfirmCancelInputs;
@@ -47,7 +52,7 @@ public:
 	TSubclassOf<AGameplayAbilityWorldReticle> ReticleClass;		//Using a special class for replication purposes
 
 protected:
-	virtual FHitResult PerformTrace(AActor *SourceActor) const;
+	virtual FHitResult PerformTrace(AActor* InSourceActor, USkeletalMeshComponent* InSourceComponent, FName InSourceSocketName) const;
 
 	FGameplayAbilityTargetDataHandle MakeTargetData(FHitResult HitResult) const;
 
diff --git a/Engine/Source/Runtime/GameplayAbilities/Public/Abilities/GameplayAbilityTypes.h b/Engine/Source/Runtime/GameplayAbilities/Public/Abilities/GameplayAbilityTypes.h
--- a/Engine/Source/Runtime/GameplayAbilities/Public/Abilities/GameplayAbilityTypes.h
+++ b/Engine/Source/Runtime/GameplayAbilities/Public/Abilities/GameplayAbilityTypes.h
@@ -352,7 +352,6 @@ struct GAMEPLAYABILITIES_API FGameplayAbilityTargetData_Mesh : public FGameplayA
 	UPROPERTY(BlueprintReadWrite, EditAnywhere, Category = Targeting)
 	FVector_NetQuantize TargetPoint;
 
-
 	//virtual TArray<AActor*>	GetActors() const
 
 	// -------------------------------------
@@ -386,7 +385,7 @@ struct GAMEPLAYABILITIES_API FGameplayAbilityTargetData_Mesh : public FGameplayA
 
 	virtual bool HasEndPoint() const
 	{
-		return false;
+		return true;
 	}
 
 	virtual FVector GetEndPoint() const
@@ -396,17 +395,17 @@ struct GAMEPLAYABILITIES_API FGameplayAbilityTargetData_Mesh : public FGameplayA
 
 	// -------------------------------------
 
-	virtual UScriptStruct* GetScriptStruct()
-	{
-		return FGameplayAbilityTargetData_Mesh::StaticStruct();
-	}
-
 	virtual FString ToString() const
 	{
 		return TEXT("FGameplayAbilityTargetData_Mesh");
 	}
 
 	bool NetSerialize(FArchive& Ar, class UPackageMap* Map, bool& bOutSuccess);
+
+	virtual UScriptStruct* GetScriptStruct()
+	{
+		return FGameplayAbilityTargetData_Mesh::StaticStruct();
+	}
 };
 
 template<>
@@ -419,7 +418,7 @@ struct TStructOpsTypeTraits<FGameplayAbilityTargetData_Mesh> : public TStructOps
 };
 
 
-USTRUCT()
+USTRUCT(BlueprintType)
 struct GAMEPLAYABILITIES_API FGameplayAbilityTargetData_SingleTargetHit : public FGameplayAbilityTargetData
 {
 	GENERATED_USTRUCT_BODY()
