FAttributeSetInitter - technique for initializing attributes from spread sheet.

diff --git a/Engine/Source/Runtime/GameplayAbilities/Private/AbilitySystemGlobals.cpp b/Engine/Source/Runtime/GameplayAbilities/Private/AbilitySystemGlobals.cpp
--- a/Engine/Source/Runtime/GameplayAbilities/Private/AbilitySystemGlobals.cpp
+++ b/Engine/Source/Runtime/GameplayAbilities/Private/AbilitySystemGlobals.cpp
@@ -1,6 +1,7 @@
 // Copyright 1998-2014 Epic Games, Inc. All Rights Reserved.
 
 #include "AbilitySystemPrivatePCH.h"
+#include "AttributeSet.h"
 #include "AbilitySystemComponent.h"
 #include "Abilities/GameplayAbilityTypes.h"
 #include "GameplayTagsModule.h"
@@ -9,55 +10,60 @@
 UAbilitySystemGlobals::UAbilitySystemGlobals(const class FPostConstructInitializeProperties& PCIP)
 : Super(PCIP)
 {
+	RegisteredReimportCallback = false;
+}
 
+void UAbilitySystemGlobals::InitGlobalData()
+{
+	GetGlobalCurveTable();
+	GetGlobalAttributeMetaDataTable();
+	
+	InitAtributeDefaults();
 }
 
 UCurveTable * UAbilitySystemGlobals::GetGlobalCurveTable()
 {
-	if (!GlobalCurveTable && !GlobalCurveTableName.IsEmpty())
-	{
-		GlobalCurveTable = LoadObject<UCurveTable>(NULL, *GlobalCurveTableName, NULL, LOAD_None, NULL);
-#if WITH_EDITOR
-		// Hook into notifications for object re-imports so that the gameplay tag tree can be reconstructed if the table changes
-		if (GIsEditor && GlobalCurveTable)
-		{
-			GEditor->OnObjectReimported().AddUObject(this, &UAbilitySystemGlobals::OnCurveTableReimported);
-		}
-#endif
-	}
-	
-	return GlobalCurveTable;
+	return InternalGetLoadTable<UCurveTable>(GlobalCurveTable, GlobalCurveTableName);
+}
+
+UDataTable * UAbilitySystemGlobals::GetGlobalAttributeMetaDataTable()
+{
+	return InternalGetLoadTable<UDataTable>(GlobalAttributeMetaDataTable, GlobalAttributeMetaDataTableName);
 }
 
-UDataTable * UAbilitySystemGlobals::GetGlobalAttributeDataTable()
+template <class T>
+T* UAbilitySystemGlobals::InternalGetLoadTable(T*& Table, FString TableName)
 {
-	if (!GlobalAttributeDataTable && !GlobalAttributeDataTableName.IsEmpty())
+	if (!Table && !TableName.IsEmpty())
 	{
-		GlobalAttributeDataTable = LoadObject<UDataTable>(NULL, *GlobalAttributeDataTableName, NULL, LOAD_None, NULL);
+		Table = LoadObject<T>(NULL, *TableName, NULL, LOAD_None, NULL);
 #if WITH_EDITOR
 		// Hook into notifications for object re-imports so that the gameplay tag tree can be reconstructed if the table changes
-		if (GIsEditor && GlobalAttributeDataTable)
+		if (GIsEditor && Table && !RegisteredReimportCallback)
 		{
-			GEditor->OnObjectReimported().AddUObject(this, &UAbilitySystemGlobals::OnDataTableReimported);
+			GEditor->OnObjectReimported().AddUObject(this, &UAbilitySystemGlobals::OnTableReimported);
+			RegisteredReimportCallback = true;
 		}
 #endif
 	}
 
-	return GlobalAttributeDataTable;
+	return Table;
 }
 
+
 #if WITH_EDITOR
 
-void UAbilitySystemGlobals::OnCurveTableReimported(UObject* InObject)
+void UAbilitySystemGlobals::OnTableReimported(UObject* InObject)
 {
 	if (GIsEditor && !IsRunningCommandlet() && InObject && InObject == GlobalCurveTable)
 	{
 	}
-}
+	
+	if (GIsEditor && !IsRunningCommandlet() && InObject && InObject == GlobalAttributeDefaultsTable)
+	{
+	}
 
-void UAbilitySystemGlobals::OnDataTableReimported(UObject* InObject)
-{
-	if (GIsEditor && !IsRunningCommandlet() && InObject && InObject == GlobalAttributeDataTable)
+	if (GIsEditor && !IsRunningCommandlet() && InObject && InObject == GlobalAttributeMetaDataTable)
 	{
 	}
 }
@@ -94,6 +100,7 @@ UAbilitySystemGlobals& UAbilitySystemGlobals::Get()
 
 
 // --------------------------------------------------------------------
+
 UFunction* UAbilitySystemGlobals::GetGameplayCueFunction(const FGameplayTag& ChildTag, UClass* Class, FName &MatchedTag)
 {
 	SCOPE_CYCLE_COUNTER(STAT_GetGameplayCueFunction);
@@ -134,3 +141,26 @@ UFunction* UAbilitySystemGlobals::GetGameplayCueFunction(const FGameplayTag& Chi
 	return nullptr;
 }
 
+// --------------------------------------------------------------------
+
+/** Initialize FAttributeSetInitter. This is virtual so projects can override what class they use */
+void UAbilitySystemGlobals::AllocAttributeSetInitter()
+{
+	GlobalAttributeSetInitter = TSharedPtr<FAttributeSetInitter>(new FAttributeSetInitter());
+}
+
+FAttributeSetInitter* UAbilitySystemGlobals::GetAttributeSetInitter() const
+{
+	check(GlobalAttributeSetInitter.IsValid());
+	return GlobalAttributeSetInitter.Get();
+}
+
+void UAbilitySystemGlobals::InitAtributeDefaults()
+{
+	if (InternalGetLoadTable<UCurveTable>(GlobalAttributeDefaultsTable, GlobalAttributeSetDefaultsTableName))
+	{	
+		AllocAttributeSetInitter();
+		GlobalAttributeSetInitter->PreloadAttributeSetData(GlobalAttributeDefaultsTable);
+	}
+
+}
\ No newline at end of file
diff --git a/Engine/Source/Runtime/GameplayAbilities/Private/AbilitySystemStats.cpp b/Engine/Source/Runtime/GameplayAbilities/Private/AbilitySystemStats.cpp
--- a/Engine/Source/Runtime/GameplayAbilities/Private/AbilitySystemStats.cpp
+++ b/Engine/Source/Runtime/GameplayAbilities/Private/AbilitySystemStats.cpp
@@ -12,5 +12,6 @@ DEFINE_STAT(STAT_RemoveActiveGameplayEffect);
 DEFINE_STAT(STAT_CreateNewActiveGameplayEffect);
 DEFINE_STAT(STAT_GetGameplayCueFunction);
 DEFINE_STAT(STAT_GetOutgoingSpec);
+DEFINE_STAT(STAT_InitAttributeSetDefaults);
 
 
diff --git a/Engine/Source/Runtime/GameplayAbilities/Private/AttributeSet.cpp b/Engine/Source/Runtime/GameplayAbilities/Private/AttributeSet.cpp
--- a/Engine/Source/Runtime/GameplayAbilities/Private/AttributeSet.cpp
+++ b/Engine/Source/Runtime/GameplayAbilities/Private/AttributeSet.cpp
@@ -92,80 +92,6 @@ FAttributeMetaData::FAttributeMetaData()
 
 }
 
-UProperty* FAttributeModifierTest::GetUProperty()
-{
-	if (!CachedUProperty)
-	{
-		FString ClassName;
-		FString ThisPropertyName;
-
-		PropertyName.Split( TEXT("."), &ClassName, &ThisPropertyName);
-
-		UClass *FoundClass = FindObject<UClass>(ANY_PACKAGE, *ClassName);
-		if (FoundClass)
-		{
-			CachedUProperty = FindField<UProperty>(FoundClass, *ThisPropertyName);
-		}
-	}
-
-	return CachedUProperty;
-}
-
-void FAttributeModifierTest::ApplyModifier(UAttributeSet *Object)
-{
-	UProperty *TheProperty = GetUProperty();
-	if (TheProperty)
-	{
-		void * ValuePtr = TheProperty->ContainerPtrToValuePtr<void>(Object);
-		switch( ModifierType)
-		{
-			case EAttributeModifierType::EATTRMOD_Add:
-			{
-				UNumericProperty *NumericProperty = Cast<UNumericProperty>(TheProperty);
-				if (NumericProperty)
-				{
-					float CurrentValue = NumericProperty->GetFloatingPointPropertyValue(ValuePtr);
-					float AdditionalValue = 0.f;
-					TTypeFromString<float>::FromString(AdditionalValue, *this->NewValue);
-					
-					float FinalValue = CurrentValue + AdditionalValue;
-					NumericProperty->SetFloatingPointPropertyValue(ValuePtr, FinalValue);
-
-					UE_LOG(LogActorComponent, Warning, TEXT("FAttributeModifierTest::ApplyModifier - %s. EATTRMOD_Add - CurrentValue: %.2f AdditionalValue: %.2f -> FinalValue: %.2f"), *TheProperty->GetName(),
-						CurrentValue, AdditionalValue, FinalValue );
-				}
-			}
-			break;
-
-			case EAttributeModifierType::EATTRMOD_Multiple:
-			{
-				UNumericProperty *NumericProperty = Cast<UNumericProperty>(TheProperty);
-				if (NumericProperty)
-				{
-					float CurrentValue = NumericProperty->GetFloatingPointPropertyValue(ValuePtr);
-					float Factor = 0.f;
-					TTypeFromString<float>::FromString(Factor, *this->NewValue);
-
-					float FinalValue = CurrentValue * Factor;
-					NumericProperty->SetFloatingPointPropertyValue(ValuePtr, FinalValue);
-
-					UE_LOG(LogActorComponent, Warning, TEXT("FAttributeModifierTest::ApplyModifier - %s. EATTRMOD_Multiple - CurrentValue: %.2f Factor: %.2f -> FinalValue: %.2f"), *TheProperty->GetName(),
-						CurrentValue, Factor, FinalValue );
-				}
-			}
-			break;
-
-			case EAttributeModifierType::EATTRMOD_Override:
-				TheProperty->ImportText(*this->NewValue, ValuePtr, 0, Object); // Not sure if the 4th parameter is right
-
-				UE_LOG(LogActorComponent, Warning, TEXT("FAttributeModifierTest::ApplyModifier - %s. EATTRMOD_Override - %s"), *TheProperty->GetName(), *this->NewValue );
-				break;
-		}
-	}
-
-	Object->PrintDebug();
-}
-
 void FScalableFloat::FinalizeCurveData(const FGlobalCurveDataOverride *GlobalOverrides)
 {
 	static const FString ContextString = TEXT("FScalableFloat::FinalizeCurveData");
@@ -238,3 +164,175 @@ bool FScalableFloat::operator!=(const FScalableFloat& Other) const
 {
 	return ((Other.Curve != Curve) || (Other.Value != Value));
 }
+
+// ------------------------------------------------------------------------------------
+//
+// ------------------------------------------------------------------------------------
+TSubclassOf<UAttributeSet> FindBestAttributeClass(TArray<TSubclassOf<UAttributeSet> >& ClassList, FString PartialName)
+{
+	for (auto Class : ClassList)
+	{
+		if (Class->GetName().Contains(PartialName))
+		{
+			return Class;
+		}
+	}
+
+	return nullptr;
+}
+
+/**
+ *	Transforms CurveTable data into format more effecient to read at runtime.
+ *	UCurveTable requires string parsing to map to GroupName/AttributeSet/Attribute
+ *	Each curve in the table represents a *single attribute's values for all levels*.
+ *	At runtime, we want *all attribute values at given level*.
+ */
+void FAttributeSetInitter::PreloadAttributeSetData(UCurveTable* CurveData)
+{
+	if(!ensure(CurveData))
+	{
+		return;
+	}
+
+	/**
+	 *	Get list of AttributeSet classes loaded
+	 */
+
+	TArray<TSubclassOf<UAttributeSet> >	ClassList;
+	for (TObjectIterator<UClass> ClassIt; ClassIt; ++ClassIt)
+	{
+		UClass* TestClass = *ClassIt;
+		if (TestClass->IsChildOf(UAttributeSet::StaticClass()))
+		{
+			ClassList.Add(TestClass);
+#if !(UE_BUILD_SHIPPING || UE_BUILD_TEST)
+			// This can only work right now on POD attribute sets. If we ever support FStrings or TArrays in AttributeSets
+			// we will need to update this code to not use memcpy etc.
+			for (TFieldIterator<UProperty> PropIt(TestClass, EFieldIteratorFlags::IncludeSuper); PropIt; ++PropIt)
+			{
+				if (!PropIt->HasAllPropertyFlags(CPF_IsPlainOldData))
+				{
+					ABILITY_LOG(Error, TEXT("FAttributeSetInitter::PreloadAttributeSetData Unable to Handle AttributeClass %s because it has a non POD property: %s"),
+						*TestClass->GetName(), *PropIt->GetName());
+					return;
+				}
+			}
+#endif
+		}
+	}
+
+	/**
+	 *	Loop through CurveData table and build sets of Defaults that keyed off of Name + Level
+	 */
+
+	for (auto It = CurveData->RowMap.CreateConstIterator(); It; ++It)
+	{
+		FString RowName = It.Key().ToString();
+		FString ClassName;
+		FString SetName;
+		FString AttributeName;
+		FString Temp;
+
+		RowName.Split(TEXT("."), &ClassName, &Temp);
+		Temp.Split(TEXT("."), &SetName, &AttributeName);
+
+		if (!ensure(!ClassName.IsEmpty() && !SetName.IsEmpty() && !AttributeName.IsEmpty()))
+		{
+			ABILITY_LOG(Warning, TEXT("FAttributeSetInitter::PreloadAttributeSetData Unable to parse row %s in %s"), *RowName, *CurveData->GetName());
+			continue;
+		}
+
+		// Find the AttributeSet
+
+		TSubclassOf<UAttributeSet> Set = FindBestAttributeClass(ClassList, SetName);
+		if (!Set)
+		{
+			ABILITY_LOG(Warning, TEXT("FAttributeSetInitter::PreloadAttributeSetData Unable to match AttributeSet from %s (row: %s)"), *SetName, *RowName);
+			continue;
+		}
+
+		// Find the UProperty
+
+		UNumericProperty* Property = FindField<UNumericProperty>(*Set, *AttributeName);
+		if (!Property)
+		{
+			ABILITY_LOG(Warning, TEXT("FAttributeSetInitter::PreloadAttributeSetData Unable to match Attribute from %s (row: %s)"), *AttributeName, *RowName);
+			continue;
+		}
+
+		FRichCurve* Curve = It.Value();
+		FName ClassFName = FName(*ClassName);
+		FAttributeSetDefaulsCollection& DefaultCollection = Defaults.FindOrAdd(ClassFName);
+
+		int32 LastLevel = Curve->GetLastKey().Time;
+		DefaultCollection.LevelData.SetNum(FMath::Max(LastLevel, DefaultCollection.LevelData.Num()));
+
+		
+		//At this point we know the Name of this "class"/"group", the AttributeSet, and the Property Name. Now loop through the values on the curve to get the attribute default value at each level.
+		for (auto KeyIter = Curve->GetKeyIterator(); KeyIter; ++KeyIter)
+		{
+			const FRichCurveKey& CurveKey = *KeyIter;
+
+			int32 Level = CurveKey.Time;
+			float Value = CurveKey.Value;
+
+			FAttributeSetDefaults& SetDefaults = DefaultCollection.LevelData[Level-1];
+
+			FAttributeDefaultValueList* DefaultDataList = SetDefaults.DataMap.Find(Set);
+			if (DefaultDataList == nullptr)
+			{
+				ABILITY_LOG(Verbose, TEXT("Initializing new default set for %s[%d]. PropertySize: %d.. DefaultSize: %d"), *Set->GetName(), Level, Set->GetPropertiesSize(), UAttributeSet::StaticClass()->GetPropertiesSize());
+				
+				DefaultDataList = &SetDefaults.DataMap.Add(Set);
+			}
+
+			// Import curve value into default data
+
+			check(DefaultDataList);
+			DefaultDataList->AddPair(Property, Value);
+		}
+
+
+	}
+}
+
+void FAttributeSetInitter::InitAttributeSetDefaults(UAbilitySystemComponent* AbilitySystemComponent, FName GroupName, int32 Level) const
+{
+	SCOPE_CYCLE_COUNTER(STAT_InitAttributeSetDefaults);
+	
+	const FAttributeSetDefaulsCollection* Collection = Defaults.Find(GroupName);
+	if (!Collection)
+	{
+		ABILITY_LOG(Warning, TEXT("Unable to find DefaultAttributeSet Group %s. Failing back to Defaults"), *GroupName.ToString());
+		Collection = Defaults.Find(FName(TEXT("Default")));
+		if (!Collection)
+		{
+			ABILITY_LOG(Error, TEXT("FAttributeSetInitter::InitAttributeSetDefaults Default DefaultAttributeSet not found! Skipping Initialization"));
+			return;
+		}
+	}
+
+	if (!Collection->LevelData.IsValidIndex(Level))
+	{
+		// We could eventually extrapolate values outside of the max defined levels
+		ABILITY_LOG(Warning, TEXT("Attribute defaults for Level %d are not defined! Skipping"), Level);
+		return;
+	}
+
+	const FAttributeSetDefaults& SetDefaults = Collection->LevelData[Level];
+	for (const UAttributeSet* Set : AbilitySystemComponent->SpawnedAttributes)
+	{
+		const FAttributeDefaultValueList* DefaultDataList = SetDefaults.DataMap.Find(Set->GetClass());
+		if (DefaultDataList)
+		{
+			ABILITY_LOG(Warning, TEXT("Initializing Set %s"), *Set->GetName());
+
+			for (auto& DataPair : DefaultDataList->List)
+			{
+				check(DataPair.Property);
+
+				DataPair.Property->SetFloatingPointPropertyValue(DataPair.Property->ContainerPtrToValuePtr<void>(const_cast<UAttributeSet*>(Set)), DataPair.Value);
+			}
+		}		
+	}
+}
diff --git a/Engine/Source/Runtime/GameplayAbilities/Private/GameplayEffect.cpp b/Engine/Source/Runtime/GameplayAbilities/Private/GameplayEffect.cpp
--- a/Engine/Source/Runtime/GameplayAbilities/Private/GameplayEffect.cpp
+++ b/Engine/Source/Runtime/GameplayAbilities/Private/GameplayEffect.cpp
@@ -79,7 +79,7 @@ void UGameplayEffect::ValidateStacking()
 			StackingPolicy = EGameplayEffectStackingPolicy::Unlimited;
 		}
 
-		UDataTable* DataTable = IGameplayAbilitiesModule::Get().GetAbilitySystemGlobals()->GetGlobalAttributeDataTable();
+		UDataTable* DataTable = IGameplayAbilitiesModule::Get().GetAbilitySystemGlobals()->GetGlobalAttributeMetaDataTable();
 
 		if (DataTable && DataTable->RowMap.Contains(StackedAttribName))
 		{
diff --git a/Engine/Source/Runtime/GameplayAbilities/Private/GameplayEffectTests.cpp b/Engine/Source/Runtime/GameplayAbilities/Private/GameplayEffectTests.cpp
--- a/Engine/Source/Runtime/GameplayAbilities/Private/GameplayEffectTests.cpp
+++ b/Engine/Source/Runtime/GameplayAbilities/Private/GameplayEffectTests.cpp
@@ -6049,7 +6049,7 @@ bool FGameplayEffectsTest::RunTest( const FString& Parameters )
 #if WITH_EDITOR
 
 	UCurveTable *CurveTable = IGameplayAbilitiesModule::Get().GetAbilitySystemGlobals()->GetGlobalCurveTable();
-	UDataTable *DataTable = IGameplayAbilitiesModule::Get().GetAbilitySystemGlobals()->GetGlobalAttributeDataTable();
+	UDataTable *DataTable = IGameplayAbilitiesModule::Get().GetAbilitySystemGlobals()->GetGlobalAttributeMetaDataTable();
 
 	// setup required GameplayTags
 	UDataTable* TagTable = CreateGameplayDataTable();
diff --git a/Engine/Source/Runtime/GameplayAbilities/Public/AbilitySystemGlobals.h b/Engine/Source/Runtime/GameplayAbilities/Public/AbilitySystemGlobals.h
--- a/Engine/Source/Runtime/GameplayAbilities/Public/AbilitySystemGlobals.h
+++ b/Engine/Source/Runtime/GameplayAbilities/Public/AbilitySystemGlobals.h
@@ -6,14 +6,20 @@
 
 class AActor;
 class UAbilitySystemComponent;
+class UCurveTable;
+class UDataTable;
+
 struct FGameplayAbilityActorInfo;
 struct FGameplayTag;
+struct FAttributeSetInitter;
 
 /** Holds global data for the skill system. Can be configured per project via config file */
 UCLASS(config=Game)
 class GAMEPLAYABILITIES_API UAbilitySystemGlobals : public UObject
 {
 	GENERATED_UCLASS_BODY()
+
+	virtual void InitGlobalData();
 	
 
 	/** Holds all of the valid gameplay-related tags that can be applied to assets */
@@ -22,11 +28,17 @@ class GAMEPLAYABILITIES_API UAbilitySystemGlobals : public UObject
 
 	/** Holds information about the valid attributes' min and max values and stacking rules */
 	UPROPERTY(config)
-	FString GlobalAttributeDataTableName;
+	FString GlobalAttributeMetaDataTableName;
 
-	class UCurveTable *	GetGlobalCurveTable();
+	/** Holds default values for attribute sets, keyed off of Name/Levels. */
+	UPROPERTY(config)
+	FString GlobalAttributeSetDefaultsTableName;
+	
+	UCurveTable* GetGlobalCurveTable();
 
-	class UDataTable * GetGlobalAttributeDataTable();
+	UDataTable* GetGlobalAttributeMetaDataTable();
+
+	
 
 	void AutomationTestOnly_SetGlobalCurveTable(class UCurveTable *InTable)
 	{
@@ -35,7 +47,7 @@ class GAMEPLAYABILITIES_API UAbilitySystemGlobals : public UObject
 
 	void AutomationTestOnly_SetGlobalAttributeDataTable(class UDataTable *InTable)
 	{
-		GlobalAttributeDataTable = InTable;
+		GlobalAttributeMetaDataTable = InTable;
 	}
 
 	static UAbilitySystemGlobals& Get();
@@ -51,15 +63,32 @@ class GAMEPLAYABILITIES_API UAbilitySystemGlobals : public UObject
 
 	UFunction* GetGameplayCueFunction(const FGameplayTag &Tag, UClass* Class, FName &MatchedTag);
 
+	FAttributeSetInitter* GetAttributeSetInitter() const;
+
+protected:
+
+	virtual void InitAtributeDefaults();
+	virtual void AllocAttributeSetInitter();
+
 private:
 
-	class UCurveTable* GlobalCurveTable;
+	UCurveTable* GlobalCurveTable;
+
+	UCurveTable* GlobalAttributeDefaultsTable;
+
+	UDataTable* GlobalAttributeMetaDataTable;
 
-	class UDataTable* GlobalAttributeDataTable;
+	TSharedPtr<FAttributeSetInitter> GlobalAttributeSetInitter;
+
+	template <class T>
+	T* InternalGetLoadTable(T*& Table, FString TableName);
 
 #if WITH_EDITOR
-	void OnCurveTableReimported(UObject* InObject);
-	void OnDataTableReimported(UObject* InObject);
+	void OnTableReimported(UObject* InObject);
+#endif
+
+#if WITH_EDITORONLY_DATA
+	bool RegisteredReimportCallback;
 #endif
 
 };
diff --git a/Engine/Source/Runtime/GameplayAbilities/Public/AbilitySystemStats.h b/Engine/Source/Runtime/GameplayAbilities/Public/AbilitySystemStats.h
--- a/Engine/Source/Runtime/GameplayAbilities/Public/AbilitySystemStats.h
+++ b/Engine/Source/Runtime/GameplayAbilities/Public/AbilitySystemStats.h
@@ -15,4 +15,5 @@ DECLARE_CYCLE_STAT_EXTERN(TEXT("GetCooldownTimeRemaining"), STAT_GameplayAbility
 DECLARE_CYCLE_STAT_EXTERN(TEXT("RemoveActiveGameplayEffect"), STAT_RemoveActiveGameplayEffect, STATGROUP_AbilitySystem, );
 DECLARE_CYCLE_STAT_EXTERN(TEXT("CreateNewActiveGameplayEffect"), STAT_CreateNewActiveGameplayEffect, STATGROUP_AbilitySystem, );
 DECLARE_CYCLE_STAT_EXTERN(TEXT("GetGameplayCueFunction"), STAT_GetGameplayCueFunction, STATGROUP_AbilitySystem, );
-DECLARE_CYCLE_STAT_EXTERN(TEXT("GetOutgoingSpec"), STAT_GetOutgoingSpec, STATGROUP_AbilitySystem, );
\ No newline at end of file
+DECLARE_CYCLE_STAT_EXTERN(TEXT("GetOutgoingSpec"), STAT_GetOutgoingSpec, STATGROUP_AbilitySystem, );
+DECLARE_CYCLE_STAT_EXTERN(TEXT("InitAttributeSetDefaults"), STAT_InitAttributeSetDefaults, STATGROUP_AbilitySystem, );
\ No newline at end of file
diff --git a/Engine/Source/Runtime/GameplayAbilities/Public/AttributeSet.h b/Engine/Source/Runtime/GameplayAbilities/Public/AttributeSet.h
--- a/Engine/Source/Runtime/GameplayAbilities/Public/AttributeSet.h
+++ b/Engine/Source/Runtime/GameplayAbilities/Public/AttributeSet.h
@@ -4,6 +4,8 @@
 
 #include "AttributeSet.generated.h"
 
+class UAbilitySystemComponent;
+
 USTRUCT(BlueprintType)
 struct GAMEPLAYABILITIES_API FGameplayAttribute
 {
@@ -94,6 +96,10 @@ struct GAMEPLAYABILITIES_API FGlobalCurveDataOverride
 	TArray<UCurveTable*>	Overrides;
 };
 
+
+/**
+ *	Generic numerical value in the form Coeffecient * Curve[Level] 
+ */
 USTRUCT()
 struct GAMEPLAYABILITIES_API FScalableFloat
 {
@@ -175,14 +181,9 @@ private:
 	FRichCurve * FinalCurve;
 };
 
-
 /**
- *	This is an example of using the property name as the row handle name.
- *	FAttributeMetaData is a table where each row is a single property in a UAttribute class.
- *
- *	I need an exmaple of having a property reference in line
- *
-*/
+ *	DataTable that allows us to define meta data about attributes. Still a work in progress.
+ */
 USTRUCT(BlueprintType)
 struct GAMEPLAYABILITIES_API FAttributeMetaData : public FTableRowBase
 {
@@ -208,61 +209,89 @@ public:
 	bool		bCanStack;
 };
 
+/**
+ *	Helper struct that facilates initializing attribute set default values from spread sheets (UCurveTable).
+ *	Projects are free to initialize their attribute sets however they want. This is just want example that is 
+ *	useful in some cases.
+ *	
+ *	Basic idea is to have a spreadsheet in this form: 
+ *	
+ *									1	2	3	4	5	6	7	8	9	10	11	12	13	14	15	16	17	18	19	20
+ *
+ *	Default.Health.MaxHealth		100	200	300	400	500	600	700	800	900	999	999	999	999	999	999	999	999	999	999	999
+ *	Default.Health.HealthRegenRate	1	1	1	1	1	1	1	1	1	1	1	1	1	1	1	1	1	1	1	1
+ *	Default.Health.AttackRating		10	10	10	10	10	10	10	10	10	10	10	10	10	10	10	10	10	10	10	10
+ *	Default.Move.MaxMoveSpeed		500	500	500	500	500	500	500	500	500	500	500	500	500	500	500	500	500	500	500	500
+ *	Hero1.Health.MaxHealth			100	100	100	100	100	100	100	100	100	100	100	100	100	100	100	100	100	100	100	100
+ *	Hero1.Health.HealthRegenRate	1	1	1	1	1	1	1	1	1	1	1	1	1	1	1	1 	1	1	1	1
+ *	Hero1.Health.AttackRating		10	10	10	10	10	10	10	10	10	10	10	10	10	10	10	10	10	10	10	10
+ *	Hero1.Move.MaxMoveSpeed			500	500	500	500	500	500	500	500	500	500	500	500	500	500	500	500	500	500	500	500
+ *	
+ *	
+ *	Where rows are in the form: [GroupName].[AttributeSetName].[Attribute]
+ *	GroupName			- arbitrary name to identify the "group"
+ *	AttributeSetName	- what UAttributeSet the attributes belong to. (Note that this is a simple partial match on the UClass name. "Health" matches "UMyGameHealthSet").
+ *	Attribute			- the name of the actual attribute property (matches full name).
+ *		
+ *	Colums represent "Level". 
+ *	
+ *	FAttributeSetInitter::PreloadAttributeSetData(UCurveTable*)
+ *	This transforms the CurveTable into a more effecient format to read in at run time. Should be called from UAbilitySystemGlobals for example.
+ *
+ *	FAttributeSetInitter::InitAttributeSetDefaults(UAbilitySystemComponent* AbilitySystemComponent, FName GroupName, int32 Level) const;
+ *	This initializes the given AbilitySystemComponent's attribute sets with the specified GroupName and Level. Game code would be expected to call
+ *	this when spawning a new Actor, or leveling up an actor, etc.
+ *	
+ *	Example Game code usage:
+ *	
+ *	IGameplayAbilitiesModule::Get().GetAbilitySystemGlobals()->GetAttributeSetInitter()->InitAttributeSetDefaults(MyCharacter->AbilitySystemComponent, "Hero1", MyLevel);
+ *	
+ *	Notes:
+ *	-This lets system designers specify arbitrary values for attributes. They can be based on any formula they want.
+ *	-Projects with very large level caps may wish to take a simpler "Attributes gained per level" approach.
+ *	-Anything initialized in this method should not be directly modified by gameplay effects. E.g., if MaxMoveSpeed scales with level, anything else that 
+ *		modifies MaxMoveSpeed should do so with a non-instant GameplayEffect.
+ *	-"Default" is currently the hardcoded, fallback GroupName. If InitAttributeSetDefaults is called without a valid GroupName, we will fallback to default.
+ *
+ *
+ */
 
-UENUM()
-enum EAttributeModifierType
-{
-	EATTRMOD_Add,
-	EATTRMOD_Multiple,
-	EATTRMOD_Override,
-};
-
-USTRUCT(BlueprintType)
-struct FAttributeModifierTest : public FTableRowBase
+struct GAMEPLAYABILITIES_API FAttributeSetInitter
 {
-	GENERATED_USTRUCT_BODY()
-
-public:
-
-	UPROPERTY()
-	FString	PropertyName;
-
-	UPROPERTY()
-	TEnumAsByte<EAttributeModifierType>	ModifierType;
+	void PreloadAttributeSetData(UCurveTable* CurveData);
 
-	UPROPERTY()
-	FString NewValue;
-
-	UProperty*	GetUProperty();
-
-	void	ApplyModifier(UAttributeSet *Object);
+	void InitAttributeSetDefaults(UAbilitySystemComponent* AbilitySystemComponent, FName GroupName, int32 Level) const;
 
 private:
 
-	UProperty *CachedUProperty;
-};
-
-USTRUCT()
-struct GAMEPLAYABILITIES_API FAttributeModifier
-{
-	GENERATED_USTRUCT_BODY()
+	struct FAttributeDefaultValueList
+	{
+		void AddPair(UNumericProperty* InProperty, float InValue)
+		{
+			List.Add(FOffsetValuePair(InProperty, InValue));
+		}
 
-	UPROPERTY()
-	FGameplayAttribute	PropertyToModify;
+		struct FOffsetValuePair
+		{
+			FOffsetValuePair(UNumericProperty* InProperty, float InValue)
+			: Property(InProperty), Value(InValue) { }
 
-	UPROPERTY()
-	float		NumericValue;
+			UNumericProperty*	Property;
+			float		Value;
+		};
 
-	UPROPERTY()
-	TEnumAsByte<EAttributeModifierType>	ModifierType;
-};
+		TArray<FOffsetValuePair>	List;
+	};
 
-USTRUCT()
-struct FSomeThingThatHoldsAListOfOtherthings
-{
-	GENERATED_USTRUCT_BODY()
+	struct FAttributeSetDefaults
+	{
+		TMap<TSubclassOf<UAttributeSet>, FAttributeDefaultValueList> DataMap;
+	};
 
-	UPROPERTY()
-	TArray< FAttributeModifier	>	ListOfModifiers;
+	struct FAttributeSetDefaulsCollection
+	{
+		TArray<FAttributeSetDefaults>		LevelData;
+	};
 
+	TMap<FName, FAttributeSetDefaulsCollection>	Defaults;
 };
\ No newline at end of file
