Targeting location struct working, but needs to be checked for a little touch-up.

diff --git a/Engine/Source/Runtime/GameplayAbilities/Private/Abilities/GameplayAbilityTargetActor.cpp b/Engine/Source/Runtime/GameplayAbilities/Private/Abilities/GameplayAbilityTargetActor.cpp
--- a/Engine/Source/Runtime/GameplayAbilities/Private/Abilities/GameplayAbilityTargetActor.cpp
+++ b/Engine/Source/Runtime/GameplayAbilities/Private/Abilities/GameplayAbilityTargetActor.cpp
@@ -15,6 +15,12 @@ AGameplayAbilityTargetActor::AGameplayAbilityTargetActor(const class FPostConstr
 	StaticTargetFunction = false;
 }
 
+void AGameplayAbilityTargetActor::GetLifetimeReplicatedProps(TArray< FLifetimeProperty > & OutLifetimeProps) const
+{
+	Super::GetLifetimeReplicatedProps(OutLifetimeProps);
+	DOREPLIFETIME(AGameplayAbilityTargetActor, StartLocation);
+}
+
 void AGameplayAbilityTargetActor::StartTargeting(UGameplayAbility* Ability)
 {
 }
diff --git a/Engine/Source/Runtime/GameplayAbilities/Private/Abilities/GameplayAbilityTargetActor_SingleLineTrace.cpp b/Engine/Source/Runtime/GameplayAbilities/Private/Abilities/GameplayAbilityTargetActor_SingleLineTrace.cpp
--- a/Engine/Source/Runtime/GameplayAbilities/Private/Abilities/GameplayAbilityTargetActor_SingleLineTrace.cpp
+++ b/Engine/Source/Runtime/GameplayAbilities/Private/Abilities/GameplayAbilityTargetActor_SingleLineTrace.cpp
@@ -25,10 +25,9 @@ AGameplayAbilityTargetActor_SingleLineTrace::AGameplayAbilityTargetActor_SingleL
 void AGameplayAbilityTargetActor_SingleLineTrace::GetLifetimeReplicatedProps(TArray< FLifetimeProperty > & OutLifetimeProps) const
 {
 	Super::GetLifetimeReplicatedProps(OutLifetimeProps);
+	//DOREPLIFETIME(AGameplayAbilityTargetActor, StartLocation);
 	DOREPLIFETIME(AGameplayAbilityTargetActor_SingleLineTrace, bDebug);
 	DOREPLIFETIME(AGameplayAbilityTargetActor_SingleLineTrace, SourceActor);
-	DOREPLIFETIME(AGameplayAbilityTargetActor_SingleLineTrace, SourceComponent);
-	DOREPLIFETIME(AGameplayAbilityTargetActor_SingleLineTrace, SourceSocketName);
 }
 
 void AGameplayAbilityTargetActor_SingleLineTrace::EndPlay(const EEndPlayReason::Type EndPlayReason)
@@ -42,7 +41,7 @@ void AGameplayAbilityTargetActor_SingleLineTrace::EndPlay(const EEndPlayReason::
 }
 
 
-FHitResult AGameplayAbilityTargetActor_SingleLineTrace::PerformTrace(AActor* InSourceActor, USkeletalMeshComponent* InSourceComponent, FName InSourceSocketName) const
+FHitResult AGameplayAbilityTargetActor_SingleLineTrace::PerformTrace(AActor* InSourceActor) const
 {
 	static const FName LineTraceSingleName(TEXT("AGameplayAbilityTargetActor_SingleLineTrace"));
 	bool bTraceComplex = false;
@@ -71,11 +70,11 @@ FHitResult AGameplayAbilityTargetActor_SingleLineTrace::PerformTrace(AActor* InS
 	FVector TraceEnd = TraceStart + (AimDirection * 3000.f);
 
 	//If we're using a socket, adjust the starting location and aim direction after the end position has been found. This way we can still aim with the camera, then fire accurately from the socket.
-	if (InSourceComponent)
+	if (MyAbility)		//Server and launching client only
 	{
-		TraceStart = InSourceComponent->GetSocketLocation(InSourceSocketName);		//If socket name is invalid, this just gets the location of the component.
-		AimDirection = (TraceEnd - TraceStart).SafeNormal();
+		TraceStart = StartLocation.GetTargetingTransform(MyAbility->GetCurrentActorInfo()).GetLocation();
 	}
+	AimDirection = (TraceEnd - TraceStart).SafeNormal();
 
 	// ------------------------------------------------------
 
@@ -95,7 +94,7 @@ FGameplayAbilityTargetDataHandle AGameplayAbilityTargetActor_SingleLineTrace::St
 	UAnimInstance* AnimInstance = ActorInfo->AnimInstance.Get();
 	USkeletalMeshComponent* StaticSourceComponent = AnimInstance ? AnimInstance->GetOwningComponent() : NULL;
 
-	return MakeTargetData(PerformTrace(StaticSourceActor, StaticSourceComponent, SourceSocketName));
+	return MakeTargetData(PerformTrace(StaticSourceActor));
 }
 
 void AGameplayAbilityTargetActor_SingleLineTrace::StartTargeting(UGameplayAbility* InAbility)
@@ -103,7 +102,6 @@ void AGameplayAbilityTargetActor_SingleLineTrace::StartTargeting(UGameplayAbilit
 	Ability = InAbility;
 	SourceActor = InAbility->GetCurrentActorInfo()->Actor.Get();
 	UAnimInstance* AnimInstance = InAbility->GetCurrentActorInfo()->AnimInstance.Get();
-	SourceComponent = AnimInstance ? AnimInstance->GetOwningComponent() : NULL;
 
 	// We can bind directly to our ASC's confirm/cancel events, or wait to be told from an outside source to confirm or cancel
 	if (bBindToConfirmCancelInputs && (MasterPC && MasterPC->IsLocalController()))
@@ -130,7 +128,7 @@ void AGameplayAbilityTargetActor_SingleLineTrace::Tick(float DeltaSeconds)
 	// very temp - do a mostly hardcoded trace from the source actor
 	if (SourceActor)
 	{
-		FHitResult HitResult = PerformTrace(SourceActor, SourceComponent, SourceSocketName);
+		FHitResult HitResult = PerformTrace(SourceActor);
 		FVector EndPoint = HitResult.Component.IsValid() ? HitResult.ImpactPoint : HitResult.TraceEnd;
 
 		if (bDebug)
@@ -161,7 +159,7 @@ void AGameplayAbilityTargetActor_SingleLineTrace::ConfirmTargeting()
 	{
 		bDebug = false;
 		bAutoFire = false;
-		FGameplayAbilityTargetDataHandle Handle = MakeTargetData(PerformTrace(SourceActor, SourceComponent, SourceSocketName));
+		FGameplayAbilityTargetDataHandle Handle = MakeTargetData(PerformTrace(SourceActor));
 		TargetDataReadyDelegate.Broadcast(Handle);
 	}
 
@@ -170,20 +168,27 @@ void AGameplayAbilityTargetActor_SingleLineTrace::ConfirmTargeting()
 
 FGameplayAbilityTargetDataHandle AGameplayAbilityTargetActor_SingleLineTrace::MakeTargetData(FHitResult HitResult) const
 {
-	/** Note these are cleaned up by the FGameplayAbilityTargetDataHandle (via an internal TSharedPtr) */
-	if (SourceComponent)
-	{
-		FGameplayAbilityTargetData_Mesh* ReturnData = new FGameplayAbilityTargetData_Mesh();
-		ReturnData->SourceActor = SourceActor;
-		ReturnData->SourceComponent = SourceComponent;
-		ReturnData->SourceSocketName = SourceSocketName;
-		ReturnData->TargetPoint = HitResult.Location;
-		return FGameplayAbilityTargetDataHandle(ReturnData);
-	}
-	else
+	//TODO - Refactor into transmitting a handle built via FGameplayAbilityTargetingLocationInfo
+
+	/** Note: These are cleaned up by the FGameplayAbilityTargetDataHandle (via an internal TSharedPtr) */
+	if (StartLocation.LocationType == EGameplayAbilityTargetingLocationType::Type::SocketTransform)
 	{
-		FGameplayAbilityTargetData_SingleTargetHit* ReturnData = new FGameplayAbilityTargetData_SingleTargetHit();
-		ReturnData->HitResult = HitResult;
-		return FGameplayAbilityTargetDataHandle(ReturnData);
+		const FGameplayAbilityActorInfo* ActorInfo = Ability.IsValid() ? Ability.Get()->GetCurrentActorInfo() : NULL;
+		UAnimInstance* AnimInstance = ActorInfo ? ActorInfo->AnimInstance.Get() : NULL;
+		USkeletalMeshComponent* AISourceComponent = AnimInstance ? AnimInstance->GetOwningComponent() : NULL;
+		AActor* AISourceActor = ActorInfo ? (ActorInfo->Actor.IsValid() ? ActorInfo->Actor.Get() : NULL) : NULL;
+
+		if (AISourceActor && AISourceComponent)
+		{
+			FGameplayAbilityTargetData_Mesh* ReturnData = new FGameplayAbilityTargetData_Mesh();
+			ReturnData->SourceActor = AISourceActor;
+			ReturnData->SourceComponent = AISourceComponent;
+			ReturnData->SourceSocketName = StartLocation.SourceSocketName;
+			ReturnData->TargetPoint = HitResult.Location;
+			return FGameplayAbilityTargetDataHandle(ReturnData);
+		}
 	}
+	FGameplayAbilityTargetData_SingleTargetHit* ReturnData = new FGameplayAbilityTargetData_SingleTargetHit();
+	ReturnData->HitResult = HitResult;
+	return FGameplayAbilityTargetDataHandle(ReturnData);
 }
\ No newline at end of file
diff --git a/Engine/Source/Runtime/GameplayAbilities/Public/Abilities/GameplayAbilityTargetActor.h b/Engine/Source/Runtime/GameplayAbilities/Public/Abilities/GameplayAbilityTargetActor.h
--- a/Engine/Source/Runtime/GameplayAbilities/Public/Abilities/GameplayAbilityTargetActor.h
+++ b/Engine/Source/Runtime/GameplayAbilities/Public/Abilities/GameplayAbilityTargetActor.h
@@ -16,6 +16,11 @@ public:
 	UPROPERTY()
 	bool StaticTargetFunction;
 
+	/** Describes where the targeting action starts, usually the player character or a socket on the player character. */
+	//UPROPERTY(BlueprintReadOnly, meta=(ExposeOnSpawn=true), Category=Targeting)
+	UPROPERTY(BlueprintReadOnly, meta = (ExposeOnSpawn = true), Replicated, Category = Targeting)
+	FGameplayAbilityTargetingLocationInfo StartLocation;
+
 	virtual FGameplayAbilityTargetDataHandle StaticGetTargetData(UWorld * World, const FGameplayAbilityActorInfo* ActorInfo, FGameplayAbilityActivationInfo ActivationInfo) const;
 	
 	/** Initialize and begin targeting logic  */
diff --git a/Engine/Source/Runtime/GameplayAbilities/Public/Abilities/GameplayAbilityTargetActor_SingleLineTrace.h b/Engine/Source/Runtime/GameplayAbilities/Public/Abilities/GameplayAbilityTargetActor_SingleLineTrace.h
--- a/Engine/Source/Runtime/GameplayAbilities/Public/Abilities/GameplayAbilityTargetActor_SingleLineTrace.h
+++ b/Engine/Source/Runtime/GameplayAbilities/Public/Abilities/GameplayAbilityTargetActor_SingleLineTrace.h
@@ -30,12 +30,6 @@ public:
 	UPROPERTY(BlueprintReadOnly, Replicated, Category = Targeting)
 	AActor* SourceActor;
 
-	UPROPERTY(BlueprintReadOnly, Replicated, Category = Targeting)
-	USkeletalMeshComponent* SourceComponent;
-
-	UPROPERTY(EditDefaultsOnly, Replicated, Category = Targeting)
-	FName SourceSocketName;
-
 	UPROPERTY(BlueprintReadOnly, meta=(ExposeOnSpawn=true), Category=Projectile)
 	bool bBindToConfirmCancelInputs;
 
@@ -52,7 +46,7 @@ public:
 	TSubclassOf<AGameplayAbilityWorldReticle> ReticleClass;		//Using a special class for replication purposes
 
 protected:
-	virtual FHitResult PerformTrace(AActor* InSourceActor, USkeletalMeshComponent* InSourceComponent, FName InSourceSocketName) const;
+	virtual FHitResult PerformTrace(AActor* InSourceActor) const;
 
 	FGameplayAbilityTargetDataHandle MakeTargetData(FHitResult HitResult) const;
 
diff --git a/Engine/Source/Runtime/GameplayAbilities/Public/Abilities/GameplayAbilityTypes.h b/Engine/Source/Runtime/GameplayAbilities/Public/Abilities/GameplayAbilityTypes.h
--- a/Engine/Source/Runtime/GameplayAbilities/Public/Abilities/GameplayAbilityTypes.h
+++ b/Engine/Source/Runtime/GameplayAbilities/Public/Abilities/GameplayAbilityTypes.h
@@ -59,6 +59,21 @@ namespace EGameplayAbilityReplicationPolicy
 	};
 }
 
+UENUM(BlueprintType)
+namespace EGameplayAbilityTargetingLocationType
+{
+	/**
+	*	What type of location calculation to use when an ability asks for our transform.
+	*/
+
+	enum Type
+	{
+		LiteralTransform		UMETA(DisplayName = "Literal Transform"),		// We report an actual raw transform. This is also the final fallback if other methods fail.
+		ActorTransform			UMETA(DisplayName = "Actor Transform"),			// We pull the transform from an associated actor directly.
+		SocketTransform			UMETA(DisplayName = "Socket Transform"),		// We aim from a named socket on the player's skeletal mesh component.
+	};
+}
+
 USTRUCT(BlueprintType)
 struct FGameplayAbilityHandle
 {
@@ -180,6 +195,90 @@ struct GAMEPLAYABILITIES_API FGameplayAbilityActorInfo
 	virtual void InitFromActor(AActor *Actor, UAbilitySystemComponent* InAbilitySystemComponent);
 };
 
+/**
+*	FGameplayAbilityActorInfo
+*
+*	Contains data used to determine locations for targeting systems.
+*
+*/
+
+USTRUCT(BlueprintType)
+struct GAMEPLAYABILITIES_API FGameplayAbilityTargetingLocationInfo
+{
+	GENERATED_USTRUCT_BODY()
+
+	FGameplayAbilityTargetingLocationInfo()
+	{
+	};
+
+	//Not sure if we want to distinguish between where we actually aim from (should probably always be the user's camera or actor origin) and where we trace from.
+	FTransform GetTargetingTransform(const FGameplayAbilityActorInfo* ActorInfo) const
+	{
+		//Return or calculate based on LocationType.
+		switch (LocationType)
+		{
+		case EGameplayAbilityTargetingLocationType::ActorTransform:
+			check(ActorInfo);
+			if (ActorInfo->Actor.IsValid())
+			{
+				return ActorInfo->Actor.Get()->GetTransform();
+			}
+			return LiteralTransform;		//Fallback
+		case EGameplayAbilityTargetingLocationType::SocketTransform:
+			check(ActorInfo);
+			check(ActorInfo->AnimInstance.IsValid());
+			if (ActorInfo->AnimInstance.Get()->GetOwningComponent())
+			{
+				return ActorInfo->AnimInstance.Get()->GetOwningComponent()->GetSocketTransform(SourceSocketName);		//Bad socket name will just return component transform anyway, so we're safe
+			}
+			if (ActorInfo->Actor.IsValid())
+			{
+				return ActorInfo->Actor->GetTransform();		//First fallback
+			}
+			return LiteralTransform;		//Second fallback, but not really acceptable since it will almost surely be Identity
+		case EGameplayAbilityTargetingLocationType::LiteralTransform:
+			return LiteralTransform;
+		default:
+			check(false);		//This case should not happen
+			return LiteralTransform;		//Fallback
+		}
+	}
+
+	UPROPERTY(BlueprintReadWrite, meta = (ExposeOnSpawn = true), Category = Targeting)
+	TEnumAsByte<EGameplayAbilityTargetingLocationType::Type> LocationType;
+
+	/** A literal world transform can be used, if one has been calculated outside of the actor using the ability. */
+	UPROPERTY(BlueprintReadWrite, meta = (ExposeOnSpawn = true), Category = Targeting)
+	FTransform LiteralTransform;
+
+	/** Actor who owns the named component. Actor's location is used as start point if component cannot be found. */
+	//UPROPERTY(BlueprintReadWrite, EditAnywhere, Category = Targeting)
+	//AActor* SourceActor;
+
+	/** Local skeletal mesh component that holds the socket. */
+	//UPROPERTY(BlueprintReadWrite, EditAnywhere, Category = Targeting)
+	//USkeletalMeshComponent* SourceComponent;
+
+	/** If SourceActor and SourceComponent are valid, this is the name of the socket that will be used instead of the actor's location. */
+	UPROPERTY(BlueprintReadWrite, meta = (ExposeOnSpawn = true), Category = Targeting)
+	FName SourceSocketName;
+
+	// -------------------------------------
+
+	virtual FString ToString() const
+	{
+		return TEXT("FGameplayAbilityTargetingLocationInfo");
+	}
+
+	//bool NetSerialize(FArchive& Ar, class UPackageMap* Map, bool& bOutSuccess);
+
+	virtual UScriptStruct* GetScriptStruct()
+	{
+		return FGameplayAbilityTargetingLocationInfo::StaticStruct();
+	}
+};
+
+
 /**
 *	FGameplayAbilityActivationInfo
 *
