unity fix

diff --git a/Engine/Source/Runtime/GameplayAbilities/Private/AttributeSet.cpp b/Engine/Source/Runtime/GameplayAbilities/Private/AttributeSet.cpp
--- a/Engine/Source/Runtime/GameplayAbilities/Private/AttributeSet.cpp
+++ b/Engine/Source/Runtime/GameplayAbilities/Private/AttributeSet.cpp
@@ -2,6 +2,7 @@
 
 #include "AbilitySystemPrivatePCH.h"
 #include "AttributeSet.h"
+#include "AbilitySystemComponent.h"
 #include "HAL/OutputDevices.h"
 
 #include "ComponentReregisterContext.h"	
