Fix issues with GameplayCues not working on PlayerState owned AbilitySystemComponents

diff --git a/Engine/Source/Runtime/GameplayAbilities/Private/AbilitySystemBlueprintLibrary.cpp b/Engine/Source/Runtime/GameplayAbilities/Private/AbilitySystemBlueprintLibrary.cpp
--- a/Engine/Source/Runtime/GameplayAbilities/Private/AbilitySystemBlueprintLibrary.cpp
+++ b/Engine/Source/Runtime/GameplayAbilities/Private/AbilitySystemBlueprintLibrary.cpp
@@ -3,6 +3,7 @@
 #include "AbilitySystemPrivatePCH.h"
 #include "AbilitySystemBlueprintLibrary.h"
 #include "AbilitySystemComponent.h"
+#include "AbilitySystemInterface.h"
 #include "Abilities/GameplayAbility.h"
 #include "Abilities/Tasks/AbilityTask_PlayMontageAndWait.h"
 #include "Abilities/Tasks/AbilityTask_WaitMovementModeChange.h"
@@ -17,11 +18,7 @@ UAbilitySystemBlueprintLibrary::UAbilitySystemBlueprintLibrary(const class FPost
 
 UAbilitySystemComponent* UAbilitySystemBlueprintLibrary::GetAbilitySystemComponent(AActor *Actor)
 {
-	/**
-	 *	Fixme: This is supposed to be a small optimization but here we are going through module code, then a global uobject, then finally a virtual function - yuck!
-	 *		-think of a better way to have a global, overridable per project function. Thats it!
-	 */
-	return UAbilitySystemGlobals::Get().GetAbilitySystemComponentFromActor(Actor);
+	return UAbilitySystemGlobals::GetAbilitySystemComponentFromActor(Actor);
 }
 
 void UAbilitySystemBlueprintLibrary::ApplyGameplayEffectToTargetData(FGameplayAbilityTargetDataHandle Target, UGameplayEffect *GameplayEffect, const FGameplayAbilityActorInfo InstigatorInfo)
diff --git a/Engine/Source/Runtime/GameplayAbilities/Private/AbilitySystemGlobals.cpp b/Engine/Source/Runtime/GameplayAbilities/Private/AbilitySystemGlobals.cpp
--- a/Engine/Source/Runtime/GameplayAbilities/Private/AbilitySystemGlobals.cpp
+++ b/Engine/Source/Runtime/GameplayAbilities/Private/AbilitySystemGlobals.cpp
@@ -70,19 +70,6 @@ void UAbilitySystemGlobals::OnTableReimported(UObject* InObject)
 
 #endif
 
-
-UAbilitySystemComponent * UAbilitySystemGlobals::GetAbilitySystemComponentFromActor(AActor *Actor) const
-{
-	if (Actor)
-	{
-		return Actor->FindComponentByClass<UAbilitySystemComponent>();
-	}
-
-	// Caller should check this
-	ensure(false);
-	return NULL;
-}
-
 FGameplayAbilityActorInfo * UAbilitySystemGlobals::AllocAbilityActorInfo() const
 {
 	return new FGameplayAbilityActorInfo();
@@ -98,6 +85,25 @@ UAbilitySystemGlobals& UAbilitySystemGlobals::Get()
 	return *GlobalPtr;
 }
 
+/** Helping function to avoid having to manually cast */
+UAbilitySystemComponent* UAbilitySystemGlobals::GetAbilitySystemComponentFromActor(AActor* Actor)
+{
+	if (Actor == nullptr)
+	{
+		return nullptr;
+	}
+
+	IAbilitySystemInterface* ASI = InterfaceCast<IAbilitySystemInterface>(Actor);
+	if (ASI)
+	{
+		return ASI->GetAbilitySystemComponent();
+	}
+
+	/** This is slow and not desirable */
+	ABILITY_LOG(Warning, TEXT("GetAbilitySystemComponentFromActor called on Actor that is not IAbilitySystemInterface. This slow!"));
+
+	return Actor->FindComponentByClass<UAbilitySystemComponent>();
+}
 
 // --------------------------------------------------------------------
 
diff --git a/Engine/Source/Runtime/GameplayAbilities/Private/GameplayEffect.cpp b/Engine/Source/Runtime/GameplayAbilities/Private/GameplayEffect.cpp
--- a/Engine/Source/Runtime/GameplayAbilities/Private/GameplayEffect.cpp
+++ b/Engine/Source/Runtime/GameplayAbilities/Private/GameplayEffect.cpp
@@ -2057,16 +2057,16 @@ float FGameplayEffectLevelSpec::GetLevel() const
 		UClass * SetClass = Attribute.GetAttributeSetClass();
 		check(SetClass);
 
-		TArray<UObject*> ChildObjects;
-		GetObjectsWithOuter(Source.Get(), ChildObjects);
-
-		for (int32 ChildIndex = 0; ChildIndex < ChildObjects.Num(); ++ChildIndex)
+		UAbilitySystemComponent* Component = UAbilitySystemGlobals::GetAbilitySystemComponentFromActor(Source.Get());
+		if (Component)
 		{
-			UObject *Obj = ChildObjects[ChildIndex];
-			if (Obj->GetClass()->IsChildOf(SetClass))
+			for (const UAttributeSet* Set : Component->SpawnedAttributes)
 			{
-				CachedLevel = Attribute.GetNumericValueChecked(CastChecked<UAttributeSet>(Obj));
-				return CachedLevel;
+				if (Set->GetClass()->IsChildOf(SetClass))
+				{
+					CachedLevel = Attribute.GetNumericValueChecked(CastChecked<UAttributeSet>(Set));
+					return CachedLevel;
+				}
 			}
 		}
 	}
diff --git a/Engine/Source/Runtime/GameplayAbilities/Public/AbilitySystemGlobals.h b/Engine/Source/Runtime/GameplayAbilities/Public/AbilitySystemGlobals.h
--- a/Engine/Source/Runtime/GameplayAbilities/Public/AbilitySystemGlobals.h
+++ b/Engine/Source/Runtime/GameplayAbilities/Public/AbilitySystemGlobals.h
@@ -38,8 +38,6 @@ class GAMEPLAYABILITIES_API UAbilitySystemGlobals : public UObject
 
 	UDataTable* GetGlobalAttributeMetaDataTable();
 
-	
-
 	void AutomationTestOnly_SetGlobalCurveTable(class UCurveTable *InTable)
 	{
 		GlobalCurveTable = InTable;
@@ -52,11 +50,7 @@ class GAMEPLAYABILITIES_API UAbilitySystemGlobals : public UObject
 
 	static UAbilitySystemGlobals& Get();
 
-	/**
-	 *	Games may want to override this in a UMyProjectAbilitySystemsGlobals class and provide
-	 *	a faster lookup for attribute compnents (E.g. Cast<MyProjPawn>(Actor)->AbilitySystemComponent; etc)
-	 */
-	virtual UAbilitySystemComponent * GetAbilitySystemComponentFromActor(AActor *Actor) const;
+	static UAbilitySystemComponent* GetAbilitySystemComponentFromActor(AActor* Actor);
 
 	/** Should allocate a project specific AbilityActorInfo struct. Caller is responsible for dellocation */
 	virtual FGameplayAbilityActorInfo * AllocAbilityActorInfo() const;
