- Attempt to fix a build break by working around missing header file with some selective commenting; Will need to be restored tomorrow

diff --git a/Engine/Source/Runtime/GameplayAbilities/Private/AbilitySystemBlueprintLibrary.cpp b/Engine/Source/Runtime/GameplayAbilities/Private/AbilitySystemBlueprintLibrary.cpp
--- a/Engine/Source/Runtime/GameplayAbilities/Private/AbilitySystemBlueprintLibrary.cpp
+++ b/Engine/Source/Runtime/GameplayAbilities/Private/AbilitySystemBlueprintLibrary.cpp
@@ -3,7 +3,7 @@
 #include "AbilitySystemPrivatePCH.h"
 #include "AbilitySystemBlueprintLibrary.h"
 #include "AbilitySystemComponent.h"
-#include "AbilitySystemInterface.h"
+//#include "AbilitySystemInterface.h"
 #include "Abilities/GameplayAbility.h"
 #include "Abilities/Tasks/AbilityTask_PlayMontageAndWait.h"
 #include "Abilities/Tasks/AbilityTask_WaitMovementModeChange.h"
diff --git a/Engine/Source/Runtime/GameplayAbilities/Private/AbilitySystemGlobals.cpp b/Engine/Source/Runtime/GameplayAbilities/Private/AbilitySystemGlobals.cpp
--- a/Engine/Source/Runtime/GameplayAbilities/Private/AbilitySystemGlobals.cpp
+++ b/Engine/Source/Runtime/GameplayAbilities/Private/AbilitySystemGlobals.cpp
@@ -93,11 +93,11 @@ UAbilitySystemComponent* UAbilitySystemGlobals::GetAbilitySystemComponentFromAct
 		return nullptr;
 	}
 
-	IAbilitySystemInterface* ASI = InterfaceCast<IAbilitySystemInterface>(Actor);
-	if (ASI)
-	{
-		return ASI->GetAbilitySystemComponent();
-	}
+// 	IAbilitySystemInterface* ASI = InterfaceCast<IAbilitySystemInterface>(Actor);
+// 	if (ASI)
+// 	{
+// 		return ASI->GetAbilitySystemComponent();
+// 	}
 
 	/** This is slow and not desirable */
 	ABILITY_LOG(Warning, TEXT("GetAbilitySystemComponentFromActor called on Actor that is not IAbilitySystemInterface. This slow!"));
