Missed files

diff --git a/Engine/Source/Runtime/GameplayAbilities/Private/AbilitySystemBlueprintLibrary.cpp b/Engine/Source/Runtime/GameplayAbilities/Private/AbilitySystemBlueprintLibrary.cpp
--- a/Engine/Source/Runtime/GameplayAbilities/Private/AbilitySystemBlueprintLibrary.cpp
+++ b/Engine/Source/Runtime/GameplayAbilities/Private/AbilitySystemBlueprintLibrary.cpp
@@ -3,7 +3,6 @@
 #include "AbilitySystemPrivatePCH.h"
 #include "AbilitySystemBlueprintLibrary.h"
 #include "AbilitySystemComponent.h"
-//#include "AbilitySystemInterface.h"
 #include "Abilities/GameplayAbility.h"
 #include "Abilities/Tasks/AbilityTask_PlayMontageAndWait.h"
 #include "Abilities/Tasks/AbilityTask_WaitMovementModeChange.h"
diff --git a/Engine/Source/Runtime/GameplayAbilities/Private/AbilitySystemGlobals.cpp b/Engine/Source/Runtime/GameplayAbilities/Private/AbilitySystemGlobals.cpp
--- a/Engine/Source/Runtime/GameplayAbilities/Private/AbilitySystemGlobals.cpp
+++ b/Engine/Source/Runtime/GameplayAbilities/Private/AbilitySystemGlobals.cpp
@@ -3,6 +3,7 @@
 #include "AbilitySystemPrivatePCH.h"
 #include "AttributeSet.h"
 #include "AbilitySystemComponent.h"
+#include "AbilitySystemInterface.h"
 #include "Abilities/GameplayAbilityTypes.h"
 #include "GameplayTagsModule.h"
 #include "GameplayCueInterface.h"
@@ -93,11 +94,11 @@ UAbilitySystemComponent* UAbilitySystemGlobals::GetAbilitySystemComponentFromAct
 		return nullptr;
 	}
 
-// 	IAbilitySystemInterface* ASI = InterfaceCast<IAbilitySystemInterface>(Actor);
-// 	if (ASI)
-// 	{
-// 		return ASI->GetAbilitySystemComponent();
-// 	}
+ 	IAbilitySystemInterface* ASI = InterfaceCast<IAbilitySystemInterface>(Actor);
+ 	if (ASI)
+ 	{
+ 		return ASI->GetAbilitySystemComponent();
+ 	}
 
 	/** This is slow and not desirable */
 	ABILITY_LOG(Warning, TEXT("GetAbilitySystemComponentFromActor called on Actor that is not IAbilitySystemInterface. This slow!"));
diff --git a/Engine/Source/Runtime/GameplayAbilities/Private/AbilitySystemInterface.cpp b/Engine/Source/Runtime/GameplayAbilities/Private/AbilitySystemInterface.cpp
new file mode 100644
--- /dev/null
+++ b/Engine/Source/Runtime/GameplayAbilities/Private/AbilitySystemInterface.cpp
@@ -0,0 +1,9 @@
+// Copyright 1998-2014 Epic Games, Inc. All Rights Reserved.
+
+#include "AbilitySystemPrivatePCH.h"
+#include "AbilitySystemInterface.h"
+
+UAbilitySystemInterface::UAbilitySystemInterface(const FPostConstructInitializeProperties& PCIP)
+: Super(PCIP)
+{
+}
\ No newline at end of file
diff --git a/Engine/Source/Runtime/GameplayAbilities/Public/AbilitySystemInterface.h b/Engine/Source/Runtime/GameplayAbilities/Public/AbilitySystemInterface.h
new file mode 100644
--- /dev/null
+++ b/Engine/Source/Runtime/GameplayAbilities/Public/AbilitySystemInterface.h
@@ -0,0 +1,23 @@
+// Copyright 1998-2014 Epic Games, Inc. All Rights Reserved.
+
+#pragma once
+
+#include "GameplayTagContainer.h"
+#include "GameplayEffect.h"
+#include "AbilitySystemInterface.generated.h"
+
+class UAbilitySystemComponent;
+
+/** Interface for actors that wish to handle GameplayCue events from GameplayEffects */
+UINTERFACE(MinimalAPI)
+class UAbilitySystemInterface : public UInterface
+{
+	GENERATED_UINTERFACE_BODY()
+};
+
+class GAMEPLAYABILITIES_API IAbilitySystemInterface
+{
+	GENERATED_IINTERFACE_BODY()
+
+	virtual UAbilitySystemComponent* GetAbilitySystemComponent() const = 0;
+};
\ No newline at end of file
