- Partially blind fixes for some other CIS errors hanging around

diff --git a/Engine/Source/Runtime/GameplayAbilities/Private/Abilities/GameplayAbilityTargetActor.cpp b/Engine/Source/Runtime/GameplayAbilities/Private/Abilities/GameplayAbilityTargetActor.cpp
--- a/Engine/Source/Runtime/GameplayAbilities/Private/Abilities/GameplayAbilityTargetActor.cpp
+++ b/Engine/Source/Runtime/GameplayAbilities/Private/Abilities/GameplayAbilityTargetActor.cpp
@@ -2,6 +2,7 @@
 
 #include "AbilitySystemPrivatePCH.h"
 #include "GameplayAbilityTargetActor.h"
+#include "Net/UnrealNetwork.h"
 
 // --------------------------------------------------------------------------------------------------------------------------------------------------------
 //
diff --git a/Engine/Source/Runtime/GameplayAbilities/Private/AbilitySystemGlobals.cpp b/Engine/Source/Runtime/GameplayAbilities/Private/AbilitySystemGlobals.cpp
--- a/Engine/Source/Runtime/GameplayAbilities/Private/AbilitySystemGlobals.cpp
+++ b/Engine/Source/Runtime/GameplayAbilities/Private/AbilitySystemGlobals.cpp
@@ -11,7 +11,9 @@
 UAbilitySystemGlobals::UAbilitySystemGlobals(const class FPostConstructInitializeProperties& PCIP)
 : Super(PCIP)
 {
+#if WITH_EDITORONLY_DATA
 	RegisteredReimportCallback = false;
+#endif // #if WITH_EDITORONLY_DATA
 }
 
 void UAbilitySystemGlobals::InitGlobalData()
