Reorg gameplaytag functions to be consistent with IGameplayTagAssetInterface

diff --git a/Engine/Source/Runtime/GameplayAbilities/Private/Abilities/GameplayAbility.cpp b/Engine/Source/Runtime/GameplayAbilities/Private/Abilities/GameplayAbility.cpp
--- a/Engine/Source/Runtime/GameplayAbilities/Private/Abilities/GameplayAbility.cpp
+++ b/Engine/Source/Runtime/GameplayAbilities/Private/Abilities/GameplayAbility.cpp
@@ -379,7 +379,7 @@ bool UGameplayAbility::CheckCooldown(const FGameplayAbilityActorInfo* ActorInfo)
 	if (CooldownGameplayEffect)
 	{
 		check(ActorInfo->AbilitySystemComponent.IsValid());
-		if (ActorInfo->AbilitySystemComponent->HasAnyTags(CooldownGameplayEffect->OwnedTagsContainer))
+		if (ActorInfo->AbilitySystemComponent->HasAnyMatchingGameplayTags(CooldownGameplayEffect->OwnedTagsContainer))
 		{
 			return false;
 		}
diff --git a/Engine/Source/Runtime/GameplayAbilities/Private/AbilitySystemComponent.cpp b/Engine/Source/Runtime/GameplayAbilities/Private/AbilitySystemComponent.cpp
--- a/Engine/Source/Runtime/GameplayAbilities/Private/AbilitySystemComponent.cpp
+++ b/Engine/Source/Runtime/GameplayAbilities/Private/AbilitySystemComponent.cpp
@@ -241,32 +241,40 @@ bool UAbilitySystemComponent::IsGameplayEffectActive(FActiveGameplayEffectHandle
 	return ActiveGameplayEffects.IsGameplayEffectActive(InHandle);
 }
 
-bool UAbilitySystemComponent::HasAnyTags(FGameplayTagContainer &Tags)
+FOnGameplayEffectTagCountChanged& UAbilitySystemComponent::RegisterGameplayTagEvent(FGameplayTag Tag)
 {
-	// Fixme: we could aggregate our current tags into a map to avoid this type of iteration
-	return ActiveGameplayEffects.HasAnyTags(Tags);
+	return ActiveGameplayEffects.RegisterGameplayTagEvent(Tag);
 }
 
-bool UAbilitySystemComponent::HasTag(FGameplayTag Tag)
+FOnGameplayAttributeChange& UAbilitySystemComponent::RegisterGameplayAttributeEvent(FGameplayAttribute Attribute)
 {
-	return ActiveGameplayEffects.HasTag(Tag);
+	return ActiveGameplayEffects.RegisterGameplayAttributeEvent(Attribute);
 }
 
-bool UAbilitySystemComponent::HasAllTags(FGameplayTagContainer &Tags)
+// ------------------------------------------------------------------------
+
+void UAbilitySystemComponent::GetOwnedGameplayTags(FGameplayTagContainer& TagContainer) const
 {
-	return ActiveGameplayEffects.HasAllTags(Tags);
+	return ActiveGameplayEffects.GetOwnedGameplayTags(TagContainer);
 }
 
-FOnGameplayEffectTagCountChanged& UAbilitySystemComponent::RegisterGameplayTagEvent(FGameplayTag Tag)
+bool UAbilitySystemComponent::HasMatchingGameplayTag(FGameplayTag TagToCheck) const
 {
-	return ActiveGameplayEffects.RegisterGameplayTagEvent(Tag);
+	return ActiveGameplayEffects.HasMatchingGameplayTag(TagToCheck);
 }
 
-FOnGameplayAttributeChange& UAbilitySystemComponent::RegisterGameplayAttributeEvent(FGameplayAttribute Attribute)
+bool UAbilitySystemComponent::HasAllMatchingGameplayTags(const FGameplayTagContainer& TagContainer, bool bCountEmptyAsMatch) const
 {
-	return ActiveGameplayEffects.RegisterGameplayAttributeEvent(Attribute);
+	return ActiveGameplayEffects.HasAllMatchingGameplayTags(TagContainer, bCountEmptyAsMatch);
+}
+
+bool UAbilitySystemComponent::HasAnyMatchingGameplayTags(const FGameplayTagContainer& TagContainer, bool bCountEmptyAsMatch) const
+{
+	return ActiveGameplayEffects.HasAnyMatchingGameplayTags(TagContainer, bCountEmptyAsMatch);
 }
 
+// ------------------------------------------------------------------------
+
 void UAbilitySystemComponent::TEMP_ApplyActiveGameplayEffects()
 {
 	for (int32 idx=0; idx < ActiveGameplayEffects.GameplayEffects.Num(); ++idx)
diff --git a/Engine/Source/Runtime/GameplayAbilities/Private/AbilitySystemStats.cpp b/Engine/Source/Runtime/GameplayAbilities/Private/AbilitySystemStats.cpp
--- a/Engine/Source/Runtime/GameplayAbilities/Private/AbilitySystemStats.cpp
+++ b/Engine/Source/Runtime/GameplayAbilities/Private/AbilitySystemStats.cpp
@@ -2,6 +2,7 @@
 
 #include "AbilitySystemPrivatePCH.h"
 
+DEFINE_STAT(STAT_GameplayEffectsGetOwnedTags);
 DEFINE_STAT(STAT_GameplayEffectsHasAllTags);
 DEFINE_STAT(STAT_GameplayEffectsHasAnyTag);
 DEFINE_STAT(STAT_GameplayEffectsTick);
diff --git a/Engine/Source/Runtime/GameplayAbilities/Private/GameplayEffect.cpp b/Engine/Source/Runtime/GameplayAbilities/Private/GameplayEffect.cpp
--- a/Engine/Source/Runtime/GameplayAbilities/Private/GameplayEffect.cpp
+++ b/Engine/Source/Runtime/GameplayAbilities/Private/GameplayEffect.cpp
@@ -1919,27 +1919,43 @@ void FActiveGameplayEffectsContainer::RecalculateStacking()
 	}
 }
 
-bool FActiveGameplayEffectsContainer::HasAnyTags(FGameplayTagContainer &Tags)
+void FActiveGameplayEffectsContainer::GetOwnedGameplayTags(FGameplayTagContainer& TagContainer) const
 {
-	SCOPE_CYCLE_COUNTER(STAT_GameplayEffectsHasAnyTag);
-	for (auto It = Tags.CreateConstIterator(); It; ++It)
+	SCOPE_CYCLE_COUNTER(STAT_GameplayEffectsGetOwnedTags);
+	for (auto It = GameplayTagCountMap.CreateConstIterator(); It; ++It)
 	{
-		if (HasTag(*It))
+		if(It.Value() > 0)
 		{
-			return true;
+			TagContainer.AddTagFast(It.Key());
 		}
 	}
+}
+
+bool FActiveGameplayEffectsContainer::HasMatchingGameplayTag(FGameplayTag TagToCheck) const
+{
+	/**
+	 *	The desired behavior here is:
+	 *		"Do we have tag a.b.c and the GameplayEffect has a.b.c.d -> match!"
+	 *		"Do we have tag a.b.c.d. and the GameplayEffect has a.b.c -> no match!"
+	 */
 
+	const int32* Count = GameplayTagCountMap.Find(TagToCheck);
+	if (Count)
+	{
+		return *Count > 0;
+	}
 	return false;
 }
 
-bool FActiveGameplayEffectsContainer::HasAllTags(FGameplayTagContainer &Tags)
+bool FActiveGameplayEffectsContainer::HasAllMatchingGameplayTags(const FGameplayTagContainer& TagContainer, bool bCountEmptyAsMatch) const
 {
 	SCOPE_CYCLE_COUNTER(STAT_GameplayEffectsHasAllTags);
+	if (TagContainer.Num() == 0)
+		return bCountEmptyAsMatch;
 
-	for (auto It = Tags.CreateConstIterator(); It;  ++It)
+	for (auto It = TagContainer.CreateConstIterator(); It; ++It)
 	{
-		if (!HasTag(*It))
+		if (!HasMatchingGameplayTag(*It))
 		{
 			return false;
 		}
@@ -1948,31 +1964,21 @@ bool FActiveGameplayEffectsContainer::HasAllTags(FGameplayTagContainer &Tags)
 	return true;
 }
 
-bool FActiveGameplayEffectsContainer::HasTag(const FGameplayTag Tag)
+bool FActiveGameplayEffectsContainer::HasAnyMatchingGameplayTags(const FGameplayTagContainer& TagContainer, bool bCountEmptyAsMatch) const
 {
-	/**
-	 *	The desired behavior here is:
-	 *		"Do we have tag a.b.c and the GameplayEffect has a.b.c.d -> match!"
-	 *		"Do we have tag a.b.c.d. and the GameplayEffect has a.b.c -> no match!"
-	 */
-
-	int32* Count = GameplayTagCountMap.Find(Tag);
-	if (Count)
-	{
-		return *Count > 0;
-	}
-	return false;
+	SCOPE_CYCLE_COUNTER(STAT_GameplayEffectsHasAnyTag);
+	if (TagContainer.Num() == 0)
+		return bCountEmptyAsMatch;
 
-#if 0
-	for (FActiveGameplayEffect &ActiveEffect : GameplayEffects)
+	for (auto It = TagContainer.CreateConstIterator(); It; ++It)
 	{
-		if (ActiveEffect.Spec.Def->OwnedTagsContainer.HasTag(Tag, EGameplayTagMatchType::IncludeParentTags, EGameplayTagMatchType::Explicit))
+		if (HasMatchingGameplayTag(*It))
 		{
 			return true;
 		}
 	}
+
 	return false;
-#endif
 }
 
 bool FActiveGameplayEffectsContainer::CanApplyAttributeModifiers(const UGameplayEffect *GameplayEffect, float Level, AActor *Instigator)
diff --git a/Engine/Source/Runtime/GameplayAbilities/Public/AbilitySystemComponent.h b/Engine/Source/Runtime/GameplayAbilities/Public/AbilitySystemComponent.h
--- a/Engine/Source/Runtime/GameplayAbilities/Public/AbilitySystemComponent.h
+++ b/Engine/Source/Runtime/GameplayAbilities/Public/AbilitySystemComponent.h
@@ -205,15 +205,16 @@ class GAMEPLAYABILITIES_API UAbilitySystemComponent : public UActorComponent
 	UFUNCTION(BlueprintCallable, Category = GameplayEffects)
 	bool IsGameplayEffectActive(FActiveGameplayEffectHandle InHandle) const;
 
+	// --------------------------------------------
 	// Tags
-	UFUNCTION(BlueprintCallable, Category=GameplayEffects)
-	bool HasAnyTags(FGameplayTagContainer &Tags);
+	// --------------------------------------------
+	bool HasMatchingGameplayTag(FGameplayTag TagToCheck) const;
 
-	UFUNCTION(BlueprintCallable, Category = GameplayEffects)
-	bool HasTag(const FGameplayTag Tag);
+	bool HasAllMatchingGameplayTags(const FGameplayTagContainer& TagContainer, bool bCountEmptyAsMatch = true) const;
 
-	UFUNCTION(BlueprintCallable, Category = GameplayEffects)
-	bool HasAllTags(FGameplayTagContainer &Tags);
+	bool HasAnyMatchingGameplayTags(const FGameplayTagContainer& TagContainer, bool bCountEmptyAsMatch = true) const;
+
+	void GetOwnedGameplayTags(FGameplayTagContainer& TagContainer) const;
 
 	/** Allow events to be registered for specific gameplay tags being added or removed */
 	FOnGameplayEffectTagCountChanged& RegisterGameplayTagEvent(FGameplayTag Tag);
diff --git a/Engine/Source/Runtime/GameplayAbilities/Public/AbilitySystemStats.h b/Engine/Source/Runtime/GameplayAbilities/Public/AbilitySystemStats.h
--- a/Engine/Source/Runtime/GameplayAbilities/Public/AbilitySystemStats.h
+++ b/Engine/Source/Runtime/GameplayAbilities/Public/AbilitySystemStats.h
@@ -8,6 +8,7 @@ DECLARE_STATS_GROUP(TEXT("AbilitySystem"), STATGROUP_AbilitySystem, STATCAT_Adva
 
 DECLARE_CYCLE_STAT_EXTERN(TEXT("GameplayEffectsHasAllTagsTime"), STAT_GameplayEffectsHasAllTags, STATGROUP_AbilitySystem, );
 DECLARE_CYCLE_STAT_EXTERN(TEXT("GameplayEffectsHasAnyTagTime"), STAT_GameplayEffectsHasAnyTag, STATGROUP_AbilitySystem, );
+DECLARE_CYCLE_STAT_EXTERN(TEXT("GameplayEffectsGetOwnedTags"), STAT_GameplayEffectsGetOwnedTags, STATGROUP_AbilitySystem, );
 DECLARE_CYCLE_STAT_EXTERN(TEXT("GameplayEffectsTick"), STAT_GameplayEffectsTick, STATGROUP_AbilitySystem, );
 DECLARE_CYCLE_STAT_EXTERN(TEXT("CanApplyAttributeModifiers"), STAT_GameplayEffectsCanApplyAttributeModifiers, STATGROUP_AbilitySystem, );
 DECLARE_CYCLE_STAT_EXTERN(TEXT("GetActiveEffectsData"), STAT_GameplayEffectsGetActiveEffectsData, STATGROUP_AbilitySystem, );
diff --git a/Engine/Source/Runtime/GameplayAbilities/Public/GameplayEffect.h b/Engine/Source/Runtime/GameplayAbilities/Public/GameplayEffect.h
--- a/Engine/Source/Runtime/GameplayAbilities/Public/GameplayEffect.h
+++ b/Engine/Source/Runtime/GameplayAbilities/Public/GameplayEffect.h
@@ -1212,11 +1212,17 @@ struct FActiveGameplayEffectsContainer : public FFastArraySerializer
 
 	bool bNeedToRecalculateStacks;
 
-	bool HasAnyTags(FGameplayTagContainer &Tags);
+	// ------------------------------------------------
+
+	void GetOwnedGameplayTags(FGameplayTagContainer& TagContainer) const;
+
+	bool HasMatchingGameplayTag(FGameplayTag TagToCheck) const;
 
-	bool HasAllTags(FGameplayTagContainer &Tags);
+	bool HasAllMatchingGameplayTags(const FGameplayTagContainer& TagContainer, bool bCountEmptyAsMatch = true) const;
 
-	bool HasTag(const FGameplayTag Tag);
+	bool HasAnyMatchingGameplayTags(const FGameplayTagContainer& TagContainer, bool bCountEmptyAsMatch = true) const;
+
+	// ------------------------------------------------
 
 	bool CanApplyAttributeModifiers(const UGameplayEffect *GameplayEffect, float Level, AActor *Instigator);
 	
