Default AbilitySystemComponent's initial ActorInfo to the outer actors.

Fixes automation tests

diff --git a/Engine/Source/Runtime/GameplayAbilities/Private/AbilitySystemComponent_Abilities.cpp b/Engine/Source/Runtime/GameplayAbilities/Private/AbilitySystemComponent_Abilities.cpp
--- a/Engine/Source/Runtime/GameplayAbilities/Private/AbilitySystemComponent_Abilities.cpp
+++ b/Engine/Source/Runtime/GameplayAbilities/Private/AbilitySystemComponent_Abilities.cpp
@@ -25,6 +25,7 @@ void UAbilitySystemComponent::InitializeComponent()
 	
 	// Look for DSO AttributeSets (note we are currently requiring all attribute sets to be subobjects of the same owner. This doesn't *have* to be the case forever.
 	AActor *Owner = GetOwner();
+	InitAbilityActorInfo(Owner);	// Default init to our outer owner
 
 	TArray<UObject*> ChildObjects;
 	GetObjectsWithOuter(Owner, ChildObjects, false, RF_PendingKill);
