Single Line Trace targeting now targets the endpoint of the trace if nothing is hit. This helps with aiming off into the sky, so you don't fire at the origin.

diff --git a/Engine/Source/Runtime/GameplayAbilities/Private/Abilities/GameplayAbilityTargetActor_SingleLineTrace.cpp b/Engine/Source/Runtime/GameplayAbilities/Private/Abilities/GameplayAbilityTargetActor_SingleLineTrace.cpp
--- a/Engine/Source/Runtime/GameplayAbilities/Private/Abilities/GameplayAbilityTargetActor_SingleLineTrace.cpp
+++ b/Engine/Source/Runtime/GameplayAbilities/Private/Abilities/GameplayAbilityTargetActor_SingleLineTrace.cpp
@@ -80,6 +80,11 @@ FHitResult AGameplayAbilityTargetActor_SingleLineTrace::PerformTrace(AActor* InS
 
 	FHitResult ReturnHitResult;
 	InSourceActor->GetWorld()->LineTraceSingle(ReturnHitResult, TraceStart, TraceEnd, ECC_WorldStatic, Params);
+	//Default to end of trace line if we don't hit anything.
+	if (!ReturnHitResult.bBlockingHit)
+	{
+		ReturnHitResult.Location = TraceEnd;
+	}
 	if (AActor* LocalReticleActor = ReticleActor.Get())
 	{
 		LocalReticleActor->SetActorLocation(ReturnHitResult.Location);
