Refactored SingleLineTrace so it doesn't make its own target data handles. Instead, the hit result can be taken in by the starting location (which is the new location data type), and that class can turn a hit result into the appropriate target data handle type.

diff --git a/Engine/Source/Runtime/GameplayAbilities/Private/Abilities/GameplayAbility.cpp b/Engine/Source/Runtime/GameplayAbilities/Private/Abilities/GameplayAbility.cpp
--- a/Engine/Source/Runtime/GameplayAbilities/Private/Abilities/GameplayAbility.cpp
+++ b/Engine/Source/Runtime/GameplayAbilities/Private/Abilities/GameplayAbility.cpp
@@ -648,6 +648,32 @@ FString FGameplayAbilityTargetData::ToString() const
 }
 
 
+FGameplayAbilityTargetDataHandle FGameplayAbilityTargetingLocationInfo::MakeTargetDataHandleFromHitResult(TWeakObjectPtr<UGameplayAbility> Ability, FHitResult HitResult) const
+{
+	/** Note: These are cleaned up by the FGameplayAbilityTargetDataHandle (via an internal TSharedPtr) */
+	if (LocationType == EGameplayAbilityTargetingLocationType::Type::SocketTransform)
+	{
+		const FGameplayAbilityActorInfo* ActorInfo = Ability.IsValid() ? Ability.Get()->GetCurrentActorInfo() : NULL;
+		AActor* AISourceActor = ActorInfo ? (ActorInfo->Actor.IsValid() ? ActorInfo->Actor.Get() : NULL) : NULL;
+		UAnimInstance* AnimInstance = ActorInfo ? ActorInfo->AnimInstance.Get() : NULL;
+		USkeletalMeshComponent* AISourceComponent = AnimInstance ? AnimInstance->GetOwningComponent() : NULL;
+
+		if (AISourceActor && AISourceComponent)
+		{
+			FGameplayAbilityTargetData_Mesh* ReturnData = new FGameplayAbilityTargetData_Mesh();
+			ReturnData->SourceActor = AISourceActor;
+			ReturnData->SourceComponent = AISourceComponent;
+			ReturnData->SourceSocketName = SourceSocketName;
+			ReturnData->TargetPoint = HitResult.Location;
+			return FGameplayAbilityTargetDataHandle(ReturnData);
+		}
+	}
+	FGameplayAbilityTargetData_SingleTargetHit* ReturnData = new FGameplayAbilityTargetData_SingleTargetHit();
+	ReturnData->HitResult = HitResult;
+	return FGameplayAbilityTargetDataHandle(ReturnData);
+}
+
+
 bool FGameplayAbilityTargetDataHandle::NetSerialize(FArchive& Ar, class UPackageMap* Map, bool& bOutSuccess)
 {
 
diff --git a/Engine/Source/Runtime/GameplayAbilities/Private/Abilities/GameplayAbilityTargetActor_SingleLineTrace.cpp b/Engine/Source/Runtime/GameplayAbilities/Private/Abilities/GameplayAbilityTargetActor_SingleLineTrace.cpp
--- a/Engine/Source/Runtime/GameplayAbilities/Private/Abilities/GameplayAbilityTargetActor_SingleLineTrace.cpp
+++ b/Engine/Source/Runtime/GameplayAbilities/Private/Abilities/GameplayAbilityTargetActor_SingleLineTrace.cpp
@@ -173,27 +173,6 @@ void AGameplayAbilityTargetActor_SingleLineTrace::ConfirmTargeting()
 
 FGameplayAbilityTargetDataHandle AGameplayAbilityTargetActor_SingleLineTrace::MakeTargetData(FHitResult HitResult) const
 {
-	//TODO - Refactor into transmitting a handle built via FGameplayAbilityTargetingLocationInfo
-
-	/** Note: These are cleaned up by the FGameplayAbilityTargetDataHandle (via an internal TSharedPtr) */
-	if (StartLocation.LocationType == EGameplayAbilityTargetingLocationType::Type::SocketTransform)
-	{
-		const FGameplayAbilityActorInfo* ActorInfo = Ability.IsValid() ? Ability.Get()->GetCurrentActorInfo() : NULL;
-		UAnimInstance* AnimInstance = ActorInfo ? ActorInfo->AnimInstance.Get() : NULL;
-		USkeletalMeshComponent* AISourceComponent = AnimInstance ? AnimInstance->GetOwningComponent() : NULL;
-		AActor* AISourceActor = ActorInfo ? (ActorInfo->Actor.IsValid() ? ActorInfo->Actor.Get() : NULL) : NULL;
-
-		if (AISourceActor && AISourceComponent)
-		{
-			FGameplayAbilityTargetData_Mesh* ReturnData = new FGameplayAbilityTargetData_Mesh();
-			ReturnData->SourceActor = AISourceActor;
-			ReturnData->SourceComponent = AISourceComponent;
-			ReturnData->SourceSocketName = StartLocation.SourceSocketName;
-			ReturnData->TargetPoint = HitResult.Location;
-			return FGameplayAbilityTargetDataHandle(ReturnData);
-		}
-	}
-	FGameplayAbilityTargetData_SingleTargetHit* ReturnData = new FGameplayAbilityTargetData_SingleTargetHit();
-	ReturnData->HitResult = HitResult;
-	return FGameplayAbilityTargetDataHandle(ReturnData);
+	/** Note: This will be cleaned up by the FGameplayAbilityTargetDataHandle (via an internal TSharedPtr) */
+	return StartLocation.MakeTargetDataHandleFromHitResult(Ability, HitResult);
 }
\ No newline at end of file
diff --git a/Engine/Source/Runtime/GameplayAbilities/Public/Abilities/GameplayAbilityTypes.h b/Engine/Source/Runtime/GameplayAbilities/Public/Abilities/GameplayAbilityTypes.h
--- a/Engine/Source/Runtime/GameplayAbilities/Public/Abilities/GameplayAbilityTypes.h
+++ b/Engine/Source/Runtime/GameplayAbilities/Public/Abilities/GameplayAbilityTypes.h
@@ -196,5 +196,73 @@
 };
 
 /**
+*	Handle for Targeting Data. This servers two main purposes:
+*		-Avoid us having to copy around the full targeting data structure in Blueprints
+*		-Allows us to leverage polymorphism in the target data structure
+*		-Allows us to implement NetSerialize and replicate by value between clients/server
+*
+*		-Avoid using UObjects could be used to give us polymorphism and by reference passing in blueprints.
+*		-However we would still be screwed when it came to replication
+*
+*		-Replication by value
+*		-Pass by reference in blueprints
+*		-Polymophism in TargetData structure
+*
+*/
+
+USTRUCT(BlueprintType)
+struct GAMEPLAYABILITIES_API FGameplayAbilityTargetDataHandle
+{
+	GENERATED_USTRUCT_BODY()
+
+	FGameplayAbilityTargetDataHandle() { }
+	FGameplayAbilityTargetDataHandle(struct FGameplayAbilityTargetData* DataPtr)
+		: Data(DataPtr)
+	{
+
+	}
+
+	TSharedPtr<FGameplayAbilityTargetData>	Data;
+
+	void Clear()
+	{
+		Data.Reset();
+	}
+
+	bool IsValid() const
+	{
+		return Data.IsValid();
+	}
+
+	bool NetSerialize(FArchive& Ar, class UPackageMap* Map, bool& bOutSuccess);
+
+	/** Comparison operator */
+	bool operator==(FGameplayAbilityTargetDataHandle const& Other) const
+	{
+		// Both invalid structs or both valid and Pointer compare (???) // deep comparison equality
+		bool bBothValid = IsValid() && Other.IsValid();
+		bool bBothInvalid = !IsValid() && !Other.IsValid();
+		return (bBothInvalid || (bBothValid && (Data.Get() == Other.Data.Get())));
+	}
+
+	/** Comparison operator */
+	bool operator!=(FGameplayAbilityTargetDataHandle const& Other) const
+	{
+		return !(FGameplayAbilityTargetDataHandle::operator==(Other));
+	}
+};
+
+template<>
+struct TStructOpsTypeTraits<FGameplayAbilityTargetDataHandle> : public TStructOpsTypeTraitsBase
+{
+	enum
+	{
+		WithCopy = true,		// Necessary so that TSharedPtr<FGameplayAbilityTargetData> Data is copied around
+		WithNetSerializer = true,
+		WithIdenticalViaEquality = true,
+	};
+};
+
+/**
 *	FGameplayAbilityActorInfo
 *
@@ -244,6 +312,8 @@
 		}
 	}
 
+	FGameplayAbilityTargetDataHandle MakeTargetDataHandleFromHitResult(TWeakObjectPtr<UGameplayAbility> Ability, FHitResult HitResult) const;
+
 	UPROPERTY(BlueprintReadWrite, meta = (ExposeOnSpawn = true), Category = Targeting)
 	TEnumAsByte<EGameplayAbilityTargetingLocationType::Type> LocationType;
 
@@ -576,74 +646,6 @@
 };
 
 
-/**
-*	Handle for Targeting Data. This servers two main purposes:
-*		-Avoid us having to copy around the full targeting data structure in Blueprints
-*		-Allows us to leverage polymorphism in the target data structure
-*		-Allows us to implement NetSerialize and replicate by value between clients/server
-*
-*		-Avoid using UObjects could be used to give us polymorphism and by reference passing in blueprints.
-*		-However we would still be screwed when it came to replication
-*
-*		-Replication by value
-*		-Pass by reference in blueprints
-*		-Polymophism in TargetData structure
-*
-*/
-
 USTRUCT(BlueprintType)
-struct GAMEPLAYABILITIES_API FGameplayAbilityTargetDataHandle
-{
-	GENERATED_USTRUCT_BODY()
-
-	FGameplayAbilityTargetDataHandle() { }
-	FGameplayAbilityTargetDataHandle(FGameplayAbilityTargetData* DataPtr)
-		: Data(DataPtr)
-	{
-		
-	}
-
-	TSharedPtr<FGameplayAbilityTargetData>	Data;
-
-	void Clear()
-	{
-		Data.Reset();
-	}
-
-	bool IsValid() const
-	{
-		return Data.IsValid();
-	}
-
-	bool NetSerialize(FArchive& Ar, class UPackageMap* Map, bool& bOutSuccess);
-
-	/** Comparison operator */
-	bool operator==(FGameplayAbilityTargetDataHandle const& Other) const
-	{
-		// Both invalid structs or both valid and Pointer compare (???) // deep comparison equality
-		bool bBothValid = IsValid() && Other.IsValid();
-		bool bBothInvalid = !IsValid() && !Other.IsValid();
-		return (bBothInvalid || (bBothValid && (Data.Get() == Other.Data.Get())));
-	}
-
-	/** Comparison operator */
-	bool operator!=(FGameplayAbilityTargetDataHandle const& Other) const
-	{
-		return !(FGameplayAbilityTargetDataHandle::operator==(Other));
-	}
-};
-
-template<>
-struct TStructOpsTypeTraits<FGameplayAbilityTargetDataHandle> : public TStructOpsTypeTraitsBase
-{
-	enum
-	{
-		WithCopy = true,		// Necessary so that TSharedPtr<FGameplayAbilityTargetData> Data is copied around
-		WithNetSerializer = true,
-		WithIdenticalViaEquality = true,
-	};
-};
-
-USTRUCT(BlueprintType)
 struct GAMEPLAYABILITIES_API FGameplayEventData
 {
