MissileSwarm rewritten to use new Wait task.

diff --git a/Engine/Source/Runtime/GameplayAbilities/Private/Abilities/Tasks/AbilityTask_Wait.cpp b/Engine/Source/Runtime/GameplayAbilities/Private/Abilities/Tasks/AbilityTask_Wait.cpp
new file mode 100644
--- /dev/null
+++ b/Engine/Source/Runtime/GameplayAbilities/Private/Abilities/Tasks/AbilityTask_Wait.cpp
@@ -0,0 +1,103 @@
+
+#include "AbilitySystemPrivatePCH.h"
+#include "GameplayAbilityTargetActor.h"
+#include "AbilitySystemComponent.h"
+#include "Abilities/Tasks/AbilityTask_Wait.h"
+
+UAbilityTask_Wait::UAbilityTask_Wait(const class FPostConstructInitializeProperties& PCIP)
+	: Super(PCIP)
+{
+	
+}
+
+UAbilityTask_Wait* UAbilityTask_Wait::Wait(UObject* WorldContextObject, TSubclassOf<AGameplayAbilityTargetActor> InTargetClass, float WaitDuration)
+{
+	auto MyObj = NewTask<UAbilityTask_Wait>(WorldContextObject);		//Register for abort list here, providing a given FName as a key
+	MyObj->TargetClass = InTargetClass;
+	if (WaitDuration > 0.0f)
+	{
+		MyObj->GetWorld()->GetTimerManager().SetTimer(MyObj, &UAbilityTask_Wait::OnWaitDurationComplete, WaitDuration, false);
+	}
+	return MyObj;
+}
+
+// ---------------------------------------------------------------------------------------
+
+bool UAbilityTask_Wait::BeginSpawningActor(UObject* WorldContextObject, TSubclassOf<AGameplayAbilityTargetActor> TargetClass, AGameplayAbilityTargetActor*& SpawnedActor)
+{
+	SpawnedActor = nullptr;
+
+	UGameplayAbility* MyAbility = Ability.Get();
+	if (MyAbility)
+	{
+		const AGameplayAbilityTargetActor* CDO = CastChecked<AGameplayAbilityTargetActor>(TargetClass->GetDefaultObject());
+		MyTargetActor = CDO;
+
+		bool Replicates = CDO->GetReplicates();
+		bool StaticFunc = CDO->StaticTargetFunction;
+		bool IsLocallyControlled = MyAbility->GetCurrentActorInfo()->IsLocallyControlled();
+
+		if (Replicates && StaticFunc)
+		{
+			// We can't replicate a staticFunc target actor, since we are just calling a static function and not spawning an actor at all!
+			ABILITY_LOG(Fatal, TEXT("AbilityTargetActor class %s can't be Replicating and Static"), *TargetClass->GetName());
+			Replicates = false;
+		}
+
+		// Spawn the actor if this is a locally controlled ability (always) or if this is a replicating targeting mode.
+		// (E.g., server will spawn this target actor to replicate to all non owning clients)
+		if (Replicates || IsLocallyControlled)
+		{
+			UClass* Class = *TargetClass;
+			if (Class != NULL)
+			{
+				UWorld* World = GEngine->GetWorldFromContextObject(WorldContextObject);
+				SpawnedActor = World->SpawnActorDeferred<AGameplayAbilityTargetActor>(Class, FVector::ZeroVector, FRotator::ZeroRotator, NULL, NULL, true);
+			}
+
+			MyTargetActor = SpawnedActor;
+
+			AGameplayAbilityTargetActor* TargetActor = CastChecked<AGameplayAbilityTargetActor>(SpawnedActor);
+			if (TargetActor)
+			{
+				TargetActor->MasterPC = MyAbility->GetCurrentActorInfo()->PlayerController.Get();
+			}
+		}
+	}
+	return (SpawnedActor != nullptr);
+}
+
+void UAbilityTask_Wait::FinishSpawningActor(UObject* WorldContextObject, AGameplayAbilityTargetActor* SpawnedActor)
+{
+	if (SpawnedActor)
+	{
+		check(MyTargetActor == SpawnedActor);
+
+		const FTransform SpawnTransform = AbilitySystemComponent->GetOwner()->GetTransform();
+
+		SpawnedActor->FinishSpawning(SpawnTransform);
+		AbilitySystemComponent->SpawnedTargetActors.Push(SpawnedActor);
+		SpawnedActor->StartTargeting(Ability.Get());
+	}
+}
+
+// ---------------------------------------------------------------------------------------
+
+void UAbilityTask_Wait::OnDestroy(bool AbilityEnded)
+{
+	if (MyTargetActor.IsValid() && !MyTargetActor->HasAnyFlags(RF_ClassDefaultObject))
+	{
+		MyTargetActor->Destroy();
+	}
+
+	Super::OnDestroy(AbilityEnded);
+}
+
+
+// --------------------------------------------------------------------------------------
+
+void UAbilityTask_Wait::OnWaitDurationComplete()
+{
+	DoneWaiting.Broadcast();
+	EndTask();
+}
diff --git a/Engine/Source/Runtime/GameplayAbilities/Public/Abilities/Tasks/AbilityTask_Wait.h b/Engine/Source/Runtime/GameplayAbilities/Public/Abilities/Tasks/AbilityTask_Wait.h
new file mode 100644
--- /dev/null
+++ b/Engine/Source/Runtime/GameplayAbilities/Public/Abilities/Tasks/AbilityTask_Wait.h
@@ -0,0 +1,49 @@
+// Copyright 1998-2014 Epic Games, Inc. All Rights Reserved.
+#pragma once
+
+#include "AbilityTask.h"
+#include "AbilityTask_Wait.generated.h"
+
+class AGameplayAbilityTargetActor;
+
+DECLARE_DYNAMIC_MULTICAST_DELEGATE(FWaitDelegate);
+
+UCLASS(MinimalAPI)
+class UAbilityTask_Wait: public UAbilityTask
+{
+	GENERATED_UCLASS_BODY()
+
+	UPROPERTY(BlueprintAssignable)
+	FWaitDelegate	DoneWaiting;
+
+	void OnWaitDurationComplete();
+
+	/** Spawns Targeting actor. */
+	UFUNCTION(BlueprintCallable, meta=(HidePin = "WorldContextObject", DefaultToSelf = "WorldContextObject", BlueprintInternalUseOnly = "true", HideSpawnParms="Instigator"), Category="Ability|Tasks")
+	static UAbilityTask_Wait* Wait(UObject* WorldContextObject, TSubclassOf<AGameplayAbilityTargetActor> Class, float InWaitDuration = -1.0f);
+
+	UFUNCTION(BlueprintCallable, meta = (HidePin = "WorldContextObject", DefaultToSelf = "WorldContextObject", BlueprintInternalUseOnly = "true"), Category = "Abilities")
+	bool BeginSpawningActor(UObject* WorldContextObject, TSubclassOf<AGameplayAbilityTargetActor> Class, AGameplayAbilityTargetActor*& SpawnedActor);
+
+	UFUNCTION(BlueprintCallable, meta = (HidePin = "WorldContextObject", DefaultToSelf = "WorldContextObject", BlueprintInternalUseOnly = "true"), Category = "Abilities")
+	void FinishSpawningActor(UObject* WorldContextObject, AGameplayAbilityTargetActor* SpawnedActor);
+
+protected:
+
+	virtual void OnDestroy(bool AbilityEnded) override;
+
+	TSubclassOf<AGameplayAbilityTargetActor> TargetClass;
+
+	/** The TargetActor that we spawned, or the class CDO if this is a static targeting task */
+	TWeakObjectPtr<AGameplayAbilityTargetActor>	MyTargetActor;
+};
+
+
+/**
+*	Requirements for using Begin/Finish SpawningActor functionality:
+*		-Have a parameters named 'Class' in your Proxy factor function (E.g., WaitTargetdata)
+*		-Have a function named BeginSpawningActor w/ the same Class parameter
+*			-This function should spawn the actor with SpawnActorDeferred and return true/false if it spawned something.
+*		-Have a function named FinishSpawningActor w/ an AActor* of the class you spawned
+*			-This function *must* call ExecuteConstruction + PostActorConstruction
+*/
\ No newline at end of file
