Added support for chain of prediction keys.

This allows predictive abilities that are triggered by other predictive abilities to be correctly handled if the earlier ability fails on the server.

diff --git a/Engine/Source/Runtime/GameplayAbilities/Private/Abilities/GameplayAbility.cpp b/Engine/Source/Runtime/GameplayAbilities/Private/Abilities/GameplayAbility.cpp
--- a/Engine/Source/Runtime/GameplayAbilities/Private/Abilities/GameplayAbility.cpp
+++ b/Engine/Source/Runtime/GameplayAbilities/Private/Abilities/GameplayAbility.cpp
@@ -49,17 +49,67 @@ UGameplayAbility::UGameplayAbility(const class FPostConstructInitializePropertie
 #endif
 }
 
-void UGameplayAbility::TriggerAbilityFromGameplayEvent(FGameplayAbilityActorInfo* ActorInfo, FGameplayTag EventTag, FGameplayEventData* Payload)
+void UGameplayAbility::TriggerAbilityFromGameplayEvent(FGameplayAbilityActorInfo* ActorInfo, FGameplayTag EventTag, FGameplayEventData* Payload, UAbilitySystemComponent& Component)
 {
 	if (ShouldAbilityRespondToEvent(EventTag, Payload))
 	{
-		TryActivateAbility(ActorInfo);
+		int32 ExecutingAbilityIndex = -1;
+
+		// if we're the server and this is coming from a predicted event we should check if the client has already predicted it
+		if (Payload->CurrPredictionKey 
+			&& NetExecutionPolicy == EGameplayAbilityNetExecutionPolicy::Predictive
+			&& ActorInfo->Actor->Role == ROLE_Authority)
+		{
+			bool bPendingClientAbilityFound = false;
+			for (auto PendingAbilityInfo : Component.PendingClientAbilities)
+			{
+				if (Payload->CurrPredictionKey == PendingAbilityInfo.PrevPredictionKey && this == PendingAbilityInfo.Ability) // found a match
+				{
+					Payload->PrevPredictionKey = PendingAbilityInfo.PrevPredictionKey;
+					Payload->CurrPredictionKey = PendingAbilityInfo.CurrPredictionKey;
+
+					Component.PendingClientAbilities.RemoveSingleSwap(PendingAbilityInfo);
+					bPendingClientAbilityFound = true;
+					break;
+				}
+			}
+
+			// we haven't received the client's copy of the triggered ability
+			// keep track of this so we can associate the prediction keys when it comes in
+			if (bPendingClientAbilityFound == false)
+			{
+				UAbilitySystemComponent::FExecutingAbilityInfo Info;
+				Info.CurrPredictionKey = Payload->CurrPredictionKey;
+				Info.Ability = this;
+
+				ExecutingAbilityIndex = Component.ExecutingServerAbilities.Add(Info);
+			}
+		}
+
+		if (TryActivateAbility(ActorInfo, Payload->PrevPredictionKey, Payload->CurrPredictionKey))
+		{
+			if (ExecutingAbilityIndex >= 0)
+			{
+				Component.ExecutingServerAbilities[ExecutingAbilityIndex].State = UAbilitySystemComponent::EAbilityExecutionState::Succeeded;
+			}
+		}
+		else if (ExecutingAbilityIndex >= 0)
+		{
+			Component.ExecutingServerAbilities[ExecutingAbilityIndex].State = UAbilitySystemComponent::EAbilityExecutionState::Failed;
+		}
 	}
 }
 
 // TODO: polymorphic payload
 void UGameplayAbility::SendGameplayEvent(FGameplayTag EventTag, FGameplayEventData Payload)
 {
+	if (NetExecutionPolicy == EGameplayAbilityNetExecutionPolicy::Predictive)
+	{
+		check(CurrentActivationInfo.CurrPredictionKey != 0);
+		Payload.PrevPredictionKey = CurrentActivationInfo.PrevPredictionKey;
+		Payload.CurrPredictionKey = CurrentActivationInfo.CurrPredictionKey;
+	}
+
 	AActor *OwnerActor = Cast<AActor>(GetOuter());
 	if (OwnerActor)
 	{
@@ -207,7 +257,7 @@ bool UGameplayAbility::CommitCheck(const FGameplayAbilityActorInfo* ActorInfo, c
 
 void UGameplayAbility::CommitExecute(const FGameplayAbilityActorInfo* ActorInfo, const FGameplayAbilityActivationInfo ActivationInfo)
 {
-	if (!ActorInfo->AbilitySystemComponent->IsNetSimulating() || ActivationInfo.PredictionKey != 0)
+	if (!ActorInfo->AbilitySystemComponent->IsNetSimulating() || ActivationInfo.CurrPredictionKey != 0)
 	{
 		ApplyCooldown(ActorInfo, ActivationInfo);
 
@@ -215,13 +265,13 @@ void UGameplayAbility::CommitExecute(const FGameplayAbilityActorInfo* ActorInfo,
 	}
 }
 
-bool UGameplayAbility::TryActivateAbility(const FGameplayAbilityActorInfo* ActorInfo, int32 PredictionKey, UGameplayAbility ** OutInstancedAbility)
+bool UGameplayAbility::TryActivateAbility(const FGameplayAbilityActorInfo* ActorInfo, uint32 PrevPredictionKey, uint32 CurrPredictionKey, UGameplayAbility ** OutInstancedAbility)
 {	
 	// This should only come from button presses/local instigation
 	ENetRole NetMode = ActorInfo->Actor->Role;
 	ensure(NetMode != ROLE_SimulatedProxy);
 	
-	FGameplayAbilityActivationInfo	ActivationInfo(ActorInfo->Actor.Get(), PredictionKey);
+	FGameplayAbilityActivationInfo	ActivationInfo(ActorInfo->Actor.Get(), PrevPredictionKey, CurrPredictionKey);
 
 	// Always do a non instanced CanActivate check
 	if (!CanActivateAbility(ActorInfo))
@@ -356,7 +406,7 @@ void UGameplayAbility::ConfirmActivateSucceed()
 
 void UGameplayAbility::ServerTryActivateAbility(const FGameplayAbilityActorInfo* ActorInfo, const FGameplayAbilityActivationInfo ActivationInfo)
 {
-	ActorInfo->AbilitySystemComponent->ServerTryActivateAbility(this, ActivationInfo.PredictionKey);
+	ActorInfo->AbilitySystemComponent->ServerTryActivateAbility(this, ActivationInfo.PrevPredictionKey, ActivationInfo.CurrPredictionKey);
 }
 
 void UGameplayAbility::ClientActivateAbility(const FGameplayAbilityActorInfo* ActorInfo, const FGameplayAbilityActivationInfo ActivationInfo)
@@ -392,7 +442,7 @@ void UGameplayAbility::ApplyCooldown(const FGameplayAbilityActorInfo* ActorInfo,
 	if (CooldownGameplayEffect)
 	{
 		check(ActorInfo->AbilitySystemComponent.IsValid());
-		ActorInfo->AbilitySystemComponent->ApplyGameplayEffectToSelf(CooldownGameplayEffect, 1.f, ActorInfo->Actor.Get(), FModifierQualifier().PredictionKey(ActivationInfo.PredictionKey));
+		ActorInfo->AbilitySystemComponent->ApplyGameplayEffectToSelf(CooldownGameplayEffect, 1.f, ActorInfo->Actor.Get(), FModifierQualifier().PredictionKeys(ActivationInfo.PrevPredictionKey, ActivationInfo.CurrPredictionKey));
 	}
 }
 
@@ -411,7 +461,7 @@ void UGameplayAbility::ApplyCost(const FGameplayAbilityActorInfo* ActorInfo, con
 	check(ActorInfo->AbilitySystemComponent.IsValid());
 	if (CostGameplayEffect)
 	{
-		ActorInfo->AbilitySystemComponent->ApplyGameplayEffectToSelf(CostGameplayEffect, 1.f, ActorInfo->Actor.Get(), FModifierQualifier().PredictionKey(ActivationInfo.PredictionKey));
+		ActorInfo->AbilitySystemComponent->ApplyGameplayEffectToSelf(CostGameplayEffect, 1.f, ActorInfo->Actor.Get(), FModifierQualifier().PredictionKeys(ActivationInfo.PrevPredictionKey, ActivationInfo.CurrPredictionKey));
 	}
 }
 
@@ -593,18 +643,13 @@ bool FGameplayAbilityActorInfo::IsNetAuthority() const
 void FGameplayAbilityActivationInfo::GeneratePredictionKey(UAbilitySystemComponent * Component) const
 {
 	check(Component);
-	check(PredictionKey == 0);
 	check(ActivationMode == EGameplayAbilityActivationMode::NonAuthority);
 
-	PredictionKey = Component->GetNextPredictionKey();
+	PrevPredictionKey = CurrPredictionKey;
+	CurrPredictionKey = Component->GetNextPredictionKey();
 	ActivationMode = EGameplayAbilityActivationMode::Predicting;
 }
 
-void FGameplayAbilityActivationInfo::SetPredictionKey(int32 InPredictionKey)
-{
-	PredictionKey = InPredictionKey;
-}
-
 void FGameplayAbilityActivationInfo::SetActivationConfirmed()
 {
 	//check(ActivationMode == EGameplayAbilityActivationMode::NonAuthority);
diff --git a/Engine/Source/Runtime/GameplayAbilities/Private/Abilities/GameplayAbility_CharacterJump.cpp b/Engine/Source/Runtime/GameplayAbilities/Private/Abilities/GameplayAbility_CharacterJump.cpp
--- a/Engine/Source/Runtime/GameplayAbilities/Private/Abilities/GameplayAbility_CharacterJump.cpp
+++ b/Engine/Source/Runtime/GameplayAbilities/Private/Abilities/GameplayAbility_CharacterJump.cpp
@@ -37,7 +37,7 @@ void UGameplayAbility_CharacterJump::InputReleased(const FGameplayAbilityActorIn
 	// For jump, we are just going to call CancelAbility with an authority ActivationInfo. This is all jump needs, but more complex
 	// abilities may need more work.
 
-	FGameplayAbilityActivationInfo	ActivationInfo(EGameplayAbilityActivationMode::Authority, 0);
+	FGameplayAbilityActivationInfo	ActivationInfo(EGameplayAbilityActivationMode::Authority);
 
 	CancelAbility(ActorInfo, ActivationInfo);
 }
diff --git a/Engine/Source/Runtime/GameplayAbilities/Private/AbilitySystemComponent.cpp b/Engine/Source/Runtime/GameplayAbilities/Private/AbilitySystemComponent.cpp
--- a/Engine/Source/Runtime/GameplayAbilities/Private/AbilitySystemComponent.cpp
+++ b/Engine/Source/Runtime/GameplayAbilities/Private/AbilitySystemComponent.cpp
@@ -303,7 +303,7 @@ FActiveGameplayEffectHandle UAbilitySystemComponent::ApplyGameplayEffectSpecToSe
 {
 	// Temp, only non instant, non periodic GEs can be predictive
 	// Effects with other effects may be a mix so go with non-predictive
-	check((BaseQualifier.PredictionKey() == 0) || (Spec.GetDuration() != UGameplayEffect::INSTANT_APPLICATION && Spec.GetPeriod() == UGameplayEffect::NO_PERIOD));
+	check((BaseQualifier.CurrentPredictionKey() == 0) || (Spec.GetDuration() != UGameplayEffect::INSTANT_APPLICATION && Spec.GetPeriod() == UGameplayEffect::NO_PERIOD));
 
 	// check if the effect being applied actually succeeds
 	float ChanceToApply = Spec.GetChanceToApplyToTarget();
@@ -329,7 +329,7 @@ FActiveGameplayEffectHandle UAbilitySystemComponent::ApplyGameplayEffectSpecToSe
 			{
 				ActiveGameplayEffects.StacksNeedToRecalculate();
 			}
-			FActiveGameplayEffect &NewActiveEffect = ActiveGameplayEffects.CreateNewActiveGameplayEffect(Spec, BaseQualifier.PredictionKey());
+			FActiveGameplayEffect &NewActiveEffect = ActiveGameplayEffects.CreateNewActiveGameplayEffect(Spec, BaseQualifier.PreviousPredictionKey(), BaseQualifier.CurrentPredictionKey());
 			MyHandle = NewActiveEffect.Handle;
 			OurCopyOfSpec = &NewActiveEffect.Spec;
 
@@ -557,7 +557,7 @@ void UAbilitySystemComponent::OnRep_PredictionKey()
 	{
 		if (PredictionDelegates[idx].Key <= ReplicatedPredictionKey)
 		{
-			PredictionDelegates[idx].Value.Broadcast();
+			PredictionDelegates[idx].Value.PredictionKeyClearDelegate.Broadcast();
 		}
 		else
 		{
@@ -571,10 +571,10 @@ void UAbilitySystemComponent::OnRep_PredictionKey()
 	}
 }
 
-FAbilitySystemComponentPredictionKeyClear& UAbilitySystemComponent::GetPredictionKeyDelegate(int32 PredictionKey)
+UAbilitySystemComponent::FPredictionInfo& UAbilitySystemComponent::GetPredictionKeyDelegate(uint32 PredictionKey)
 {
 	// See if we already have one for this key
-	for (TPair<int32, FAbilitySystemComponentPredictionKeyClear> & Item : PredictionDelegates)
+	for (auto & Item : PredictionDelegates)
 	{
 		if (Item.Key == PredictionKey)
 		{
@@ -583,15 +583,16 @@ FAbilitySystemComponentPredictionKeyClear& UAbilitySystemComponent::GetPredictio
 	}
 
 	// Create a new one
-	TPair<int32, FAbilitySystemComponentPredictionKeyClear> NewPair;
+	TPair<uint32, FPredictionInfo> NewPair;
 	NewPair.Key = PredictionKey;
 	
 	PredictionDelegates.Add(NewPair);
 	return PredictionDelegates.Last().Value;
 }
 
-int32 UAbilitySystemComponent::GetNextPredictionKey()
+uint32 UAbilitySystemComponent::GetNextPredictionKey()
 {
+	LocalPredictionKey = FMath::Max(LocalPredictionKey, ReplicatedPredictionKey);
 	return ++LocalPredictionKey;
 }
 
diff --git a/Engine/Source/Runtime/GameplayAbilities/Private/AbilitySystemComponent_Abilities.cpp b/Engine/Source/Runtime/GameplayAbilities/Private/AbilitySystemComponent_Abilities.cpp
--- a/Engine/Source/Runtime/GameplayAbilities/Private/AbilitySystemComponent_Abilities.cpp
+++ b/Engine/Source/Runtime/GameplayAbilities/Private/AbilitySystemComponent_Abilities.cpp
@@ -210,13 +210,13 @@ void UAbilitySystemComponent::OnRep_ActivateAbilities()
 	}
 }
 
-void UAbilitySystemComponent::ServerTryActivateAbility_Implementation(UGameplayAbility * AbilityToActivate, int32 PredictionKey)
+void UAbilitySystemComponent::ServerTryActivateAbility_Implementation(UGameplayAbility * AbilityToActivate, uint32 PrevPredictionKey, uint32 CurrPredictionKey)
 {
 #if !(UE_BUILD_SHIPPING || UE_BUILD_TEST)
 	if (DenyClientActivation > 0)
 	{
 		DenyClientActivation--;
-		ClientActivateAbilityFailed(AbilityToActivate, PredictionKey);
+		ClientActivateAbilityFailed(AbilityToActivate, CurrPredictionKey);
 		return;
 	}
 #endif
@@ -224,51 +224,91 @@ void UAbilitySystemComponent::ServerTryActivateAbility_Implementation(UGameplayA
 	ensure(AbilityToActivate);
 	ensure(AbilityActorInfo.IsValid());
 
+	// if this was triggered by a predicted ability the server triggered copy should be the canonical one
+	if (PrevPredictionKey != 0)
+	{
+		// first check if the server has already started this ability
+		for (auto ExecutingAbilityInfo : ExecutingServerAbilities)
+		{
+			if (ExecutingAbilityInfo.CurrPredictionKey == PrevPredictionKey && ExecutingAbilityInfo.Ability == AbilityToActivate)
+			{
+				switch(ExecutingAbilityInfo.State)
+				{
+					case EAbilityExecutionState::Executing:
+						ExecutingAbilityInfo.Ability->CurrentActivationInfo.PrevPredictionKey = PrevPredictionKey;
+						ExecutingAbilityInfo.Ability->CurrentActivationInfo.CurrPredictionKey = CurrPredictionKey;
+						break;
+					case EAbilityExecutionState::Failed:
+						ClientActivateAbilityFailed(AbilityToActivate, CurrPredictionKey);
+						break;
+					case EAbilityExecutionState::Succeeded:
+						ClientActivateAbilitySucceed(AbilityToActivate, CurrPredictionKey);
+						break;
+				}
+
+				ExecutingServerAbilities.RemoveSingleSwap(ExecutingAbilityInfo);
+				return;
+			}
+		}
+
+		FPendingAbilityInfo AbilityInfo;
+		AbilityInfo.PrevPredictionKey = PrevPredictionKey;
+		AbilityInfo.CurrPredictionKey = CurrPredictionKey;
+		AbilityInfo.Ability = AbilityToActivate;
+
+		PendingClientAbilities.Add(AbilityInfo);
+		return;
+	}
+
 	UGameplayAbility *InstancedAbility = NULL;
 
-	if (AbilityToActivate->TryActivateAbility(AbilityActorInfo.Get(), PredictionKey, &InstancedAbility))
+	if (AbilityToActivate->TryActivateAbility(AbilityActorInfo.Get(), PrevPredictionKey, CurrPredictionKey, &InstancedAbility))
 	{
 		if (InstancedAbility && InstancedAbility->GetReplicationPolicy() != EGameplayAbilityReplicationPolicy::ReplicateNone)
 		{
-			InstancedAbility->ClientActivateAbilitySucceed(PredictionKey);
+			InstancedAbility->ClientActivateAbilitySucceed(CurrPredictionKey);
 		}
 		else
 		{
-			ClientActivateAbilitySucceed(AbilityToActivate, PredictionKey);
+			ClientActivateAbilitySucceed(AbilityToActivate, CurrPredictionKey);
 		}
 
 		// Update our ReplicatedPredictionKey. When the client gets value, he will know his state (actor+all components/subobjects) are up to do date and he can
 		// remove any necessary predictive work.
 
-		if (PredictionKey > 0)
+		if (CurrPredictionKey > 0)
 		{
-			ensure(PredictionKey > ReplicatedPredictionKey);
-			ReplicatedPredictionKey = PredictionKey;
+			ensure(CurrPredictionKey > ReplicatedPredictionKey);
+			ReplicatedPredictionKey = CurrPredictionKey;
 		}
 	}
 	else
 	{
-		ClientActivateAbilityFailed(AbilityToActivate, PredictionKey);
+		ClientActivateAbilityFailed(AbilityToActivate, CurrPredictionKey);
 	}
 }
 
-bool UAbilitySystemComponent::ServerTryActivateAbility_Validate(UGameplayAbility * AbilityToActivate, int32 PredictionKey)
+bool UAbilitySystemComponent::ServerTryActivateAbility_Validate(UGameplayAbility * AbilityToActivate, uint32 PrevPredictionKey, uint32 CurrPredictionKey)
 {
 	return true;
 }
 
-void UAbilitySystemComponent::ClientActivateAbilityFailed_Implementation(UGameplayAbility * AbilityToActivate, int32 PredictionKey)
+void UAbilitySystemComponent::ClientActivateAbilityFailed_Implementation(UGameplayAbility * AbilityToActivate, uint32 PredictionKey)
 {
 	if (PredictionKey > 0)
 	{
 		for (int32 idx = 0; idx < PredictionDelegates.Num(); ++idx)
 		{
-			int32 Key = PredictionDelegates[idx].Key;
+			uint32 Key = PredictionDelegates[idx].Key;
 			if (Key == PredictionKey)
 			{
 				ABILITY_LOG(Warning, TEXT("Failed ActivateAbility, clearing prediction data %d"), PredictionKey);
 
-				PredictionDelegates[idx].Value.Broadcast();
+				PredictionDelegates[idx].Value.PredictionKeyClearDelegate.Broadcast();
+				for (uint32 DependentKey : PredictionDelegates[idx].Value.DependentPredictionKeys)
+				{
+					ClientActivateAbilityFailed(AbilityToActivate, DependentKey);
+				}
 				PredictionDelegates.RemoveAt(idx);
 				break;
 			}
@@ -280,24 +320,21 @@ void UAbilitySystemComponent::ClientActivateAbilityFailed_Implementation(UGamepl
 	}
 }
 
-void UAbilitySystemComponent::ClientActivateAbilitySucceed_Implementation(UGameplayAbility* AbilityToActivate, int32 PredictionKey)
+void UAbilitySystemComponent::ClientActivateAbilitySucceed_Implementation(UGameplayAbility* AbilityToActivate, uint32 PredictionKey)
 {
 	ensure(AbilityToActivate);
 	ensure(AbilityActorInfo.IsValid());
 
-	// Fixme: We need a better way to link up/reconcile preditive replicated abilities. It would be ideal if we could predictively spawn an
+	// Fixme: We need a better way to link up/reconcile predictive replicated abilities. It would be ideal if we could predictively spawn an
 	// ability and then replace/link it with the server spawned one once the server has confirmed it.
 
-	FGameplayAbilityActivationInfo	ActivationInfo(EGameplayAbilityActivationMode::Confirmed, PredictionKey);
-
-
 	if (AbilityToActivate->NetExecutionPolicy == EGameplayAbilityNetExecutionPolicy::Predictive)
 	{
 		// Find the one we predictively spawned, tell them we are confirmed
 		bool found = false;
 		for (UGameplayAbility* LocalAbility : NonReplicatedInstancedAbilities)				// Fixme: this has to be updated once predictive abilities can replicate
 		{
-			if (LocalAbility->GetCurrentActivationInfo().PredictionKey == PredictionKey)
+			if (LocalAbility->GetCurrentActivationInfo().CurrPredictionKey == PredictionKey)
 			{
 				LocalAbility->ConfirmActivateSucceed();
 				found = true;
@@ -312,6 +349,7 @@ void UAbilitySystemComponent::ClientActivateAbilitySucceed_Implementation(UGamep
 	}
 	else
 	{
+		FGameplayAbilityActivationInfo	ActivationInfo(EGameplayAbilityActivationMode::Confirmed, 0, PredictionKey);
 		// We haven't already executed this ability at all, so kick it off.
 		if (AbilityToActivate->GetInstancingPolicy() == EGameplayAbilityInstancingPolicy::InstancedPerExecution)
 		{
@@ -381,7 +419,7 @@ void UAbilitySystemComponent::HandleGameplayEvent(FGameplayTag EventTag, FGamepl
 		{
 			if (Ability.IsValid())
 			{
-				Ability.Get()->TriggerAbilityFromGameplayEvent(AbilityActorInfo.Get(), EventTag, Payload);
+				Ability.Get()->TriggerAbilityFromGameplayEvent(AbilityActorInfo.Get(), EventTag, Payload, *this);
 			}
 		}
 	}
diff --git a/Engine/Source/Runtime/GameplayAbilities/Private/AbilitySystemDebugHUD.cpp b/Engine/Source/Runtime/GameplayAbilities/Private/AbilitySystemDebugHUD.cpp
--- a/Engine/Source/Runtime/GameplayAbilities/Private/AbilitySystemDebugHUD.cpp
+++ b/Engine/Source/Runtime/GameplayAbilities/Private/AbilitySystemDebugHUD.cpp
@@ -114,7 +114,7 @@ void AAbilitySystemDebugHUD::DrawDebugAbilitySystemComponent(UAbilitySystemCompo
 		// Draw Active GameplayEffect
 		for (FActiveGameplayEffect &Effect : Component->ActiveGameplayEffects.GameplayEffects)
 		{
-			String = FString::Printf(TEXT("%s. [%d] %.2f"), *Effect.Spec.ToSimpleString(), Effect.PredictionKey, Effect.GetTimeRemaining(GameWorldTime));
+			String = FString::Printf(TEXT("%s. [%d, %d] %.2f"), *Effect.Spec.ToSimpleString(), Effect.PrevPredictionKey, Effect.CurrPredictionKey, Effect.GetTimeRemaining(GameWorldTime));
 			DrawWithBackground(Font, String, Color, EAlignHorizontal::Left, X, EAlignVertical::Top, Y);	
 		}
 	}
diff --git a/Engine/Source/Runtime/GameplayAbilities/Private/GameplayEffect.cpp b/Engine/Source/Runtime/GameplayAbilities/Private/GameplayEffect.cpp
--- a/Engine/Source/Runtime/GameplayAbilities/Private/GameplayEffect.cpp
+++ b/Engine/Source/Runtime/GameplayAbilities/Private/GameplayEffect.cpp
@@ -1561,12 +1561,12 @@ void FActiveGameplayEffectsContainer::StacksNeedToRecalculate()
 	}
 }
 
-FActiveGameplayEffect & FActiveGameplayEffectsContainer::CreateNewActiveGameplayEffect(const FGameplayEffectSpec &Spec, int32 InPredictionKey)
+FActiveGameplayEffect & FActiveGameplayEffectsContainer::CreateNewActiveGameplayEffect(const FGameplayEffectSpec &Spec, uint32 InPrevPredictionKey, uint32 InCurrPredictionKey)
 {
 	SCOPE_CYCLE_COUNTER(STAT_CreateNewActiveGameplayEffect);
 
 	LastAssignedHandle = LastAssignedHandle.GetNextHandle();
-	FActiveGameplayEffect & NewEffect = *new (GameplayEffects)FActiveGameplayEffect(LastAssignedHandle, Spec, GetWorldTime(), GetGameStateTime(), InPredictionKey);
+	FActiveGameplayEffect & NewEffect = *new (GameplayEffects)FActiveGameplayEffect(LastAssignedHandle, Spec, GetWorldTime(), GetGameStateTime(), InPrevPredictionKey, InCurrPredictionKey);
 
 	// register callbacks with the timer manager
 	if (Owner)
@@ -1610,7 +1610,7 @@ FActiveGameplayEffect & FActiveGameplayEffectsContainer::CreateNewActiveGameplay
 		NewEffect.Spec.Duration.Get()->OnDirty = FAggregator::FOnDirty::CreateRaw(this, &FActiveGameplayEffectsContainer::OnDurationAggregatorDirty, Owner, LastAssignedHandle);
 	}
 	
-	if (InPredictionKey == 0 || IsNetAuthority())	// Clients predicting a GameplayEffect must not call MarkItemDirty
+	if (InCurrPredictionKey == 0 || IsNetAuthority())	// Clients predicting a GameplayEffect must not call MarkItemDirty
 	{
 		MarkItemDirty(NewEffect);
 	}
@@ -1619,7 +1619,13 @@ FActiveGameplayEffect & FActiveGameplayEffectsContainer::CreateNewActiveGameplay
 		// Clients predicting should call MarkArrayDirty to force the internal replication map to be rebuilt.
 		MarkArrayDirty();
 
-		Owner->GetPredictionKeyDelegate(InPredictionKey).AddUObject(Owner, &UAbilitySystemComponent::RemoveActiveGameplayEffect_NoReturn, NewEffect.Handle);
+		UAbilitySystemComponent::FPredictionInfo PredictionInfo = Owner->GetPredictionKeyDelegate(InCurrPredictionKey);
+		PredictionInfo.PredictionKeyClearDelegate.AddUObject(Owner, &UAbilitySystemComponent::RemoveActiveGameplayEffect_NoReturn, NewEffect.Handle);
+
+		if (InPrevPredictionKey)
+		{
+			PredictionInfo.DependentPredictionKeys.AddUnique(InCurrPredictionKey);
+		}
 	}
 
 	InternalOnActiveGameplayEffectAdded(NewEffect);	
diff --git a/Engine/Source/Runtime/GameplayAbilities/Public/Abilities/GameplayAbility.h b/Engine/Source/Runtime/GameplayAbilities/Public/Abilities/GameplayAbility.h
--- a/Engine/Source/Runtime/GameplayAbilities/Public/Abilities/GameplayAbility.h
+++ b/Engine/Source/Runtime/GameplayAbilities/Public/Abilities/GameplayAbility.h
@@ -119,7 +119,7 @@ public:
 	 *	-This function handles networking and prediction
 	 *	-If all goes well, CallActivateAbility is called next.
 	 */
-	bool TryActivateAbility(const FGameplayAbilityActorInfo* ActorInfo, int32 PredictionKey = 0, UGameplayAbility ** OutInstancedAbility = NULL);
+	bool TryActivateAbility(const FGameplayAbilityActorInfo* ActorInfo, uint32 PrevPredictionKey = 0, uint32 CurrPredictionKey = 0, UGameplayAbility ** OutInstancedAbility = NULL);
 
 	/** Returns true if this ability can be activated right now. Has no side effects */
 	virtual bool CanActivateAbility(const FGameplayAbilityActorInfo* ActorInfo) const;
@@ -224,10 +224,10 @@ protected:
 	/** Called on a predictive ability when the server confirms its execution */
 	void ConfirmActivateSucceed();
 
-	virtual void TriggerAbilityFromGameplayEvent(FGameplayAbilityActorInfo* ActorInfo, FGameplayTag Tag, FGameplayEventData* Payload);
+	virtual void TriggerAbilityFromGameplayEvent(FGameplayAbilityActorInfo* ActorInfo, FGameplayTag Tag, FGameplayEventData* Payload, UAbilitySystemComponent& Component);
 
 	UFUNCTION(BlueprintCallable, Category = "Abilities")
-		void SendGameplayEvent(FGameplayTag EventTag, FGameplayEventData Payload);
+	void SendGameplayEvent(FGameplayTag EventTag, FGameplayEventData Payload);
 
 	// --------------------------------------
 	//	CommitAbility
diff --git a/Engine/Source/Runtime/GameplayAbilities/Public/Abilities/GameplayAbilityTypes.h b/Engine/Source/Runtime/GameplayAbilities/Public/Abilities/GameplayAbilityTypes.h
--- a/Engine/Source/Runtime/GameplayAbilities/Public/Abilities/GameplayAbilityTypes.h
+++ b/Engine/Source/Runtime/GameplayAbilities/Public/Abilities/GameplayAbilityTypes.h
@@ -354,49 +354,52 @@ struct GAMEPLAYABILITIES_API FGameplayAbilityTargetingLocationInfo
 *
 *	Data tied to a specific activation of an ability.
 *		-Tell us whether we are the authority, if we are predicting, confirmed, etc.
-*		-Holds PredictionKey
+*		-Holds current and previous PredictionKey
 *		-Generally not meant to be subclassed in projects.
 *		-Passed around by value since the struct is small.
 */
-
 USTRUCT(BlueprintType)
 struct GAMEPLAYABILITIES_API FGameplayAbilityActivationInfo
 {
 	GENERATED_USTRUCT_BODY()
 
 	FGameplayAbilityActivationInfo()
-	: ActivationMode(EGameplayAbilityActivationMode::Authority),
-	PredictionKey(0)
+	: ActivationMode(EGameplayAbilityActivationMode::Authority)
+	, PrevPredictionKey(0)
+	, CurrPredictionKey(0)
 	{
 
 	}
 
-	FGameplayAbilityActivationInfo(AActor * InActor, int32 InPredictionKey)
-		: PredictionKey(InPredictionKey)
+	FGameplayAbilityActivationInfo(AActor * InActor, uint32 InPrevPredictionKey, uint32 InCurrPredictionKey)
+		: PrevPredictionKey(InPrevPredictionKey)
+		, CurrPredictionKey(InCurrPredictionKey)
 	{
 		// On Init, we are either Authority or NonAuthority. We haven't been given a PredictionKey and we haven't been confirmed.
 		// NonAuthority essentially means 'I'm not sure what how I'm going to do this yet'.
 		ActivationMode = (InActor->Role == ROLE_Authority ? EGameplayAbilityActivationMode::Authority : EGameplayAbilityActivationMode::NonAuthority);
 	}
 
-	FGameplayAbilityActivationInfo(EGameplayAbilityActivationMode::Type InType, int32 InPredictionKey)
+	FGameplayAbilityActivationInfo(EGameplayAbilityActivationMode::Type InType, uint32 InPrevPredictionKey = 0, uint32 InCurrPredictionKey = 0)
 		: ActivationMode(InType)
-		, PredictionKey(InPredictionKey)
+		, PrevPredictionKey(InPrevPredictionKey)
+		, CurrPredictionKey(InCurrPredictionKey)
 	{
 	}
 
 
 	void GeneratePredictionKey(UAbilitySystemComponent * Component) const;
 
-	void SetPredictionKey(int32 InPredictionKey);
-
 	void SetActivationConfirmed();
 
 	UPROPERTY(BlueprintReadOnly, Category = "ActorInfo")
 	mutable TEnumAsByte<EGameplayAbilityActivationMode::Type>	ActivationMode;
 
 	UPROPERTY()
-	mutable int32 PredictionKey;
+	mutable uint32 PrevPredictionKey;
+
+	UPROPERTY()
+	mutable uint32 CurrPredictionKey;
 };
 
 
@@ -650,6 +653,15 @@ USTRUCT(BlueprintType)
 struct GAMEPLAYABILITIES_API FGameplayEventData
 {
 	GENERATED_USTRUCT_BODY()
+
+	FGameplayEventData()
+	: Instigator(NULL)
+	, Target(NULL)
+	, PrevPredictionKey(0)
+	, CurrPredictionKey(0)
+	{
+		// do nothing
+	}
 	
 	UPROPERTY(EditAnywhere, BlueprintReadWrite, Category = GameplayAbilityTriggerPayload)
 	AActor* Instigator;
@@ -661,6 +673,9 @@ struct GAMEPLAYABILITIES_API FGameplayEventData
 	float Var2;
 	UPROPERTY(EditAnywhere, BlueprintReadWrite, Category = GameplayAbilityTriggerPayload)
 	float Var3;
+
+	uint32 PrevPredictionKey;
+	uint32 CurrPredictionKey;
 };
 
 /** 
diff --git a/Engine/Source/Runtime/GameplayAbilities/Public/AbilitySystemComponent.h b/Engine/Source/Runtime/GameplayAbilities/Public/AbilitySystemComponent.h
--- a/Engine/Source/Runtime/GameplayAbilities/Public/AbilitySystemComponent.h
+++ b/Engine/Source/Runtime/GameplayAbilities/Public/AbilitySystemComponent.h
@@ -146,20 +146,67 @@ class GAMEPLAYABILITIES_API UAbilitySystemComponent : public UActorComponent
 	 */
 
 	UPROPERTY(ReplicatedUsing=OnRep_PredictionKey)
-	int32	ReplicatedPredictionKey;
+	uint32	ReplicatedPredictionKey;
 
 	UPROPERTY(transient)
-	int32	LocalPredictionKey;
+	uint32	LocalPredictionKey;
 
-	int32	GetNextPredictionKey();
+	uint32	GetNextPredictionKey();
 
 	UFUNCTION()
 	void OnRep_PredictionKey();
 
-	TArray<TPair<int32, FAbilitySystemComponentPredictionKeyClear> > PredictionDelegates;
+	struct FPredictionInfo
+	{
+		FAbilitySystemComponentPredictionKeyClear PredictionKeyClearDelegate;
+		TArray<uint32> DependentPredictionKeys;
+	};
 
-	FAbilitySystemComponentPredictionKeyClear &	GetPredictionKeyDelegate(int32 PredictionKey);
-	
+	TArray<TPair<uint32, FPredictionInfo > > PredictionDelegates;
+
+	FPredictionInfo &	GetPredictionKeyDelegate(uint32 PredictionKey);
+
+	struct FPendingAbilityInfo
+	{
+		bool operator==(const FPendingAbilityInfo& Other)
+		{
+			return PrevPredictionKey == Other.PrevPredictionKey
+				&& CurrPredictionKey == Other.CurrPredictionKey
+				&& Ability == Other.Ability;
+		}
+
+		uint32 PrevPredictionKey;
+		uint32 CurrPredictionKey;
+		UGameplayAbility* Ability;
+	};
+
+	// This is a list of GameplayAbilities that are predicted by the client and were triggered by abilities that were also predicted by the client
+	// When the server version of the predicted ability executes it should trigger copies of these and the copies will be associated with the correct prediction keys
+	TArray<FPendingAbilityInfo> PendingClientAbilities;
+
+	enum class EAbilityExecutionState : uint8
+	{
+		Executing,
+		Succeeded,
+		Failed,
+	};
+
+	struct FExecutingAbilityInfo
+	{
+		FExecutingAbilityInfo() : State(EAbilityExecutionState::Executing) {};
+
+		bool operator==(const FExecutingAbilityInfo& Other)
+		{
+			return CurrPredictionKey == Other.CurrPredictionKey
+				&& State == Other.State;
+		}
+
+		uint32 CurrPredictionKey;
+		EAbilityExecutionState State;
+		UGameplayAbility* Ability;
+	};
+
+	TArray<FExecutingAbilityInfo> ExecutingServerAbilities;
 
 	// ----------------------------------------------------------------------------------------------------------------
 	//
@@ -332,13 +379,13 @@ class GAMEPLAYABILITIES_API UAbilitySystemComponent : public UActorComponent
 	void	OnRep_ActivateAbilities();
 
 	UFUNCTION(Server, reliable, WithValidation)
-	void	ServerTryActivateAbility(class UGameplayAbility* AbilityToActivate, int32 PredictionKey);
+	void	ServerTryActivateAbility(class UGameplayAbility* AbilityToActivate, uint32 PrevPredictionKey, uint32 CurrPredictionKey);
 
 	UFUNCTION(Client, Reliable)
-	void	ClientActivateAbilityFailed(class UGameplayAbility* AbilityToActivate, int32 PredictionKey);
+	void	ClientActivateAbilityFailed(class UGameplayAbility* AbilityToActivate, uint32 PredictionKey);
 
 	UFUNCTION(Client, Reliable)
-	void	ClientActivateAbilitySucceed(class UGameplayAbility* AbilityToActivate, int32 PredictionKey);
+	void	ClientActivateAbilitySucceed(class UGameplayAbility* AbilityToActivate, uint32 PredictionKey);
 
 	// ----------------------------------------------------------------------------------------------------------------
 
diff --git a/Engine/Source/Runtime/GameplayAbilities/Public/GameplayEffect.h b/Engine/Source/Runtime/GameplayAbilities/Public/GameplayEffect.h
--- a/Engine/Source/Runtime/GameplayAbilities/Public/GameplayEffect.h
+++ b/Engine/Source/Runtime/GameplayAbilities/Public/GameplayEffect.h
@@ -815,7 +815,8 @@ struct FModifierQualifier
 
 	FModifierQualifier()
 		: MyType(EGameplayMod::Max)
-		, MyPredictionKey(0)
+		, MyPrevPredictionKey(0)
+		, MyCurrPredictionKey(0)
 	{
 	}
 
@@ -835,14 +836,28 @@ struct FModifierQualifier
 	// ----------------------------------
 	// PredictionKey is used by networking/replication to specify that a GameplayEffect has been predictively created/added.
 
-	int32 PredictionKey() const
+	uint32 PreviousPredictionKey() const
 	{
-		return MyPredictionKey;
+		return MyPrevPredictionKey;
 	}
 
-	FModifierQualifier& PredictionKey(int32 InPredictionKey)
+	uint32 CurrentPredictionKey() const
 	{
-		MyPredictionKey = InPredictionKey;
+		return MyCurrPredictionKey;
+	}
+
+	FModifierQualifier& PredictionKey(uint32 InPredictionKey)
+	{
+		MyPrevPredictionKey = MyCurrPredictionKey;
+		MyCurrPredictionKey = InPredictionKey;
+		return *this;
+	}
+
+	// this has the same result as calling PredictionKey twice in a row to set both keys
+	FModifierQualifier& PredictionKeys(uint32 InPrevPredictionKey, uint32 InCurrPredictionKey)
+	{
+		MyPrevPredictionKey = InPrevPredictionKey;
+		MyCurrPredictionKey = InCurrPredictionKey;
 		return *this;
 	}
 
@@ -902,7 +917,8 @@ private:
 	EGameplayMod::Type	MyType;
 	FActiveGameplayEffectHandle MyIgnoreHandle;		// Do not modify this handle
 	FActiveGameplayEffectHandle MyExclusiveTargetHandle;	// Only modify this handle
-	int32 MyPredictionKey;
+	uint32 MyPrevPredictionKey;
+	uint32 MyCurrPredictionKey;
 };
 
 /**
@@ -1047,16 +1063,18 @@ struct FActiveGameplayEffect : public FFastArraySerializerItem
 	FActiveGameplayEffect()
 		: StartGameStateTime(0)
 		, StartWorldTime(0.f)
-		, PredictionKey(0)
+		, PrevPredictionKey(0)
+		, CurrPredictionKey(0)
 	{
 	}
 
-	FActiveGameplayEffect(FActiveGameplayEffectHandle InHandle, const FGameplayEffectSpec &InSpec, float CurrentWorldTime, int32 InStartGameStateTime, int32 InPredictionKey)
+	FActiveGameplayEffect(FActiveGameplayEffectHandle InHandle, const FGameplayEffectSpec &InSpec, float CurrentWorldTime, int32 InStartGameStateTime, uint32 InPrevPredictionKey, uint32 InCurrPredictionKey)
 		: Handle(InHandle)
 		, Spec(InSpec)
 		, StartGameStateTime(InStartGameStateTime)
 		, StartWorldTime(CurrentWorldTime)
-		, PredictionKey(InPredictionKey)
+		, PrevPredictionKey(InPrevPredictionKey)
+		, CurrPredictionKey(InCurrPredictionKey)
 	{
 		for (FModifierSpec &Mod : Spec.Modifiers)
 		{
@@ -1082,7 +1100,8 @@ struct FActiveGameplayEffect : public FFastArraySerializerItem
 
 	FTimerHandle DurationHandle;
 
-	int32 PredictionKey;
+	uint32 PrevPredictionKey;
+	uint32 CurrPredictionKey;
 
 	float GetTimeRemaining(float WorldTime)
 	{
@@ -1158,7 +1177,7 @@ struct FActiveGameplayEffectsContainer : public FFastArraySerializer
 
 	UAbilitySystemComponent *Owner;
 	
-	FActiveGameplayEffect& CreateNewActiveGameplayEffect(const FGameplayEffectSpec &Spec, int32 InPredictionKey);
+	FActiveGameplayEffect& CreateNewActiveGameplayEffect(const FGameplayEffectSpec &Spec, uint32 InPrevPredictionKey, uint32 InCurrPredictionKey);
 
 	FActiveGameplayEffect* GetActiveGameplayEffect(const FActiveGameplayEffectHandle Handle);
 
@@ -1332,4 +1351,4 @@ struct TStructOpsTypeTraits<FGameplayEffectSpecHandle> : public TStructOpsTypeTr
 		WithNetSerializer = true,
 		WithIdenticalViaEquality = true,
 	};
-};
\ No newline at end of file
+};
