#UE4 - Added support for ticking ability tasks

diff --git a/Engine/Source/Runtime/GameplayAbilities/Private/Abilities/Tasks/AbilityTask.cpp b/Engine/Source/Runtime/GameplayAbilities/Private/Abilities/Tasks/AbilityTask.cpp
--- a/Engine/Source/Runtime/GameplayAbilities/Private/Abilities/Tasks/AbilityTask.cpp
+++ b/Engine/Source/Runtime/GameplayAbilities/Private/Abilities/Tasks/AbilityTask.cpp
@@ -19,6 +19,12 @@ void UAbilityTask::InitTask(UGameplayAbility* InAbility)
 	Ability = InAbility;
 	AbilitySystemComponent = InAbility->GetCurrentActorInfo()->AbilitySystemComponent;
 	InAbility->TaskStarted(this);
+
+	// If this task requires ticking, register it with AbilitySystemComponent
+	if (bTickingTask && AbilitySystemComponent.IsValid())
+	{
+		AbilitySystemComponent->TickingTaskStarted(this);
+	}
 }
 
 UWorld* UAbilityTask::GetWorld() const
@@ -65,6 +71,12 @@ void UAbilityTask::OnDestroy(bool AbilityIsEnding)
 {
 	ensure(!IsPendingKill());
 
+	// If this task required ticking, unregister it with AbilitySystemComponent
+	if (bTickingTask && AbilitySystemComponent.IsValid())
+	{
+		AbilitySystemComponent->TickingTaskEnded(this);
+	}
+
 	// Remove ourselves from the owning Ability's task list, if the ability isn't ending
 	if (!AbilityIsEnding && Ability.IsValid())
 	{
diff --git a/Engine/Source/Runtime/GameplayAbilities/Private/AbilitySystemComponent.cpp b/Engine/Source/Runtime/GameplayAbilities/Private/AbilitySystemComponent.cpp
--- a/Engine/Source/Runtime/GameplayAbilities/Private/AbilitySystemComponent.cpp
+++ b/Engine/Source/Runtime/GameplayAbilities/Private/AbilitySystemComponent.cpp
@@ -497,6 +497,28 @@ void UAbilitySystemComponent::OnRestackGameplayEffects()
 
 // ---------------------------------------------------------------------------------------
 
+void UAbilitySystemComponent::TickingTaskStarted(UAbilityTask* NewTask)
+{
+	// If this is our first ticking task, set this component as active so it begins ticking
+	if (TickingTasks.Num() == 0)
+	{
+		SetActive(true);
+	}
+	TickingTasks.AddUnique(NewTask);
+}
+
+void UAbilitySystemComponent::TickingTaskEnded(UAbilityTask* Task)
+{
+	// If we are removing our last ticking task, set this component as inactive so it stops ticking
+	TickingTasks.RemoveSingleSwap(Task);
+	if (TickingTasks.Num() == 0)
+	{
+		SetActive(false);
+	}
+}
+
+// ---------------------------------------------------------------------------------------
+
 void UAbilitySystemComponent::GetLifetimeReplicatedProps(TArray< FLifetimeProperty > & OutLifetimeProps) const
 {
 	Super::GetLifetimeReplicatedProps(OutLifetimeProps);
diff --git a/Engine/Source/Runtime/GameplayAbilities/Private/AbilitySystemComponent_Abilities.cpp b/Engine/Source/Runtime/GameplayAbilities/Private/AbilitySystemComponent_Abilities.cpp
--- a/Engine/Source/Runtime/GameplayAbilities/Private/AbilitySystemComponent_Abilities.cpp
+++ b/Engine/Source/Runtime/GameplayAbilities/Private/AbilitySystemComponent_Abilities.cpp
@@ -40,6 +40,19 @@ void UAbilitySystemComponent::InitializeComponent()
 	}
 }
 
+void UAbilitySystemComponent::TickComponent(float DeltaTime, enum ELevelTick TickType, FActorComponentTickFunction *ThisTickFunction)
+{
+	Super::TickComponent(DeltaTime, TickType, ThisTickFunction);
+
+	for (TWeakObjectPtr<UAbilityTask> TickingTask : TickingTasks)
+	{
+		if (TickingTask.IsValid())
+		{
+			TickingTask->TickTask(DeltaTime);
+		}
+	}
+}
+
 void UAbilitySystemComponent::InitAbilityActorInfo(AActor* AvatarActor)
 {
 	check(AbilityActorInfo.IsValid());
@@ -115,6 +128,12 @@ UGameplayAbility* UAbilitySystemComponent::CreateNewInstanceOfAbility(UGameplayA
 void UAbilitySystemComponent::NotifyAbilityEnded(UGameplayAbility* Ability)
 {
 	check(Ability);
+
+	// If AnimatingAbility ended, clear the pointer
+	if (AnimatingAbility == Ability)
+	{
+		AnimatingAbility = NULL;
+	}
 	
 	/** If this is instanced per execution, mark pending kill and remove it from our instanced lists if we are the authority */
 	if (Ability->GetInstancingPolicy() == EGameplayAbilityInstancingPolicy::InstancedPerExecution)
diff --git a/Engine/Source/Runtime/GameplayAbilities/Public/Abilities/Tasks/AbilityTask.h b/Engine/Source/Runtime/GameplayAbilities/Public/Abilities/Tasks/AbilityTask.h
--- a/Engine/Source/Runtime/GameplayAbilities/Public/Abilities/Tasks/AbilityTask.h
+++ b/Engine/Source/Runtime/GameplayAbilities/Public/Abilities/Tasks/AbilityTask.h
@@ -70,6 +70,9 @@ class GAMEPLAYABILITIES_API UAbilityTask : public UObject
 	/** Initailizes the task with the owning GameplayAbility but does not actviate until Activate() is called */
 	virtual void InitTask(UGameplayAbility* InAbility);
 
+	/** Tick function for this task, if bTickingTask == true */
+	virtual void TickTask(float DeltaTime) {}
+
 	
 	
 	/** GameplayAbility that created us */
@@ -105,4 +108,7 @@ protected:
 
 	/** End and CleanUp the task - may be called by the task itself or by the owning ability if the ability is ending. Do NOT call directly! Call EndTask() or AbilityEnded() */
 	virtual void OnDestroy(bool AbilityIsEnding);
+
+	/** If true, this task will receive TickTask calls from AbilitySystemComponent */
+	bool bTickingTask;
 };
\ No newline at end of file
diff --git a/Engine/Source/Runtime/GameplayAbilities/Public/AbilitySystemComponent.h b/Engine/Source/Runtime/GameplayAbilities/Public/AbilitySystemComponent.h
--- a/Engine/Source/Runtime/GameplayAbilities/Public/AbilitySystemComponent.h
+++ b/Engine/Source/Runtime/GameplayAbilities/Public/AbilitySystemComponent.h
@@ -64,6 +64,7 @@ class GAMEPLAYABILITIES_API UAbilitySystemComponent : public UActorComponent
 	virtual ~UAbilitySystemComponent();
 
 	virtual void InitializeComponent() override;
+	virtual void TickComponent(float DeltaTime, enum ELevelTick TickType, FActorComponentTickFunction *ThisTickFunction) override;
 
 	/** Finds existing AttributeSet */
 	template <class T >
@@ -448,6 +449,12 @@ class GAMEPLAYABILITIES_API UAbilitySystemComponent : public UActorComponent
 	/** This is temp for testing. We want to think a bit more on the API for outside stuff activating abilities */
 	bool	ActivateAbility(TWeakObjectPtr<UGameplayAbility> Ability);
 
+	/** Adds a UAbilityTask task to the list of tasks to be ticked */
+	void TickingTaskStarted(UAbilityTask* NewTask);
+
+	/** Removes a UAbilityTask task from the list of tasks to be ticked */
+	void TickingTaskEnded(UAbilityTask* Task);
+
 	/** There needs to be a concept of an animating ability. Only one may exist at a time. New requests can be queued up, overridden, or ignored. */
 	UPROPERTY()
 	UGameplayAbility*	AnimatingAbility;
@@ -551,6 +558,8 @@ private:
 	
 	// ---------------------------------------------
 
+	/** Array of currently active UAbilityTasks that require ticking */
+	TArray<TWeakObjectPtr<UAbilityTask> >	TickingTasks;
 	
 
 protected:
