Fix non-unity for ability system

diff --git a/Engine/Source/Runtime/GameplayAbilities/Private/AbilitySystemComponent.cpp b/Engine/Source/Runtime/GameplayAbilities/Private/AbilitySystemComponent.cpp
--- a/Engine/Source/Runtime/GameplayAbilities/Private/AbilitySystemComponent.cpp
+++ b/Engine/Source/Runtime/GameplayAbilities/Private/AbilitySystemComponent.cpp
@@ -5,7 +5,7 @@
 #include "AbilitySystemComponent.h"
 #include "GameplayCueInterface.h"
 #include "Abilities/GameplayAbility.h"
-
+#include "Abilities/Tasks/AbilityTask.h"
 
 #include "Net/UnrealNetwork.h"
 #include "Engine/ActorChannel.h"
diff --git a/Engine/Source/Runtime/GameplayAbilities/Private/AbilitySystemComponent_Abilities.cpp b/Engine/Source/Runtime/GameplayAbilities/Private/AbilitySystemComponent_Abilities.cpp
--- a/Engine/Source/Runtime/GameplayAbilities/Private/AbilitySystemComponent_Abilities.cpp
+++ b/Engine/Source/Runtime/GameplayAbilities/Private/AbilitySystemComponent_Abilities.cpp
@@ -5,6 +5,7 @@
 #include "AbilitySystemComponent.h"
 #include "Abilities/GameplayAbility.h"
 #include "Abilities/GameplayAbilityTargetActor.h"
+#include "Abilities/Tasks/AbilityTask.h"
 
 #include "Net/UnrealNetwork.h"
 #include "MessageLog.h"
