Cleanup of commented-out code.

diff --git a/Engine/Source/Runtime/GameplayAbilities/Public/Abilities/GameplayAbilityTypes.h b/Engine/Source/Runtime/GameplayAbilities/Public/Abilities/GameplayAbilityTypes.h
--- a/Engine/Source/Runtime/GameplayAbilities/Public/Abilities/GameplayAbilityTypes.h
+++ b/Engine/Source/Runtime/GameplayAbilities/Public/Abilities/GameplayAbilityTypes.h
@@ -506,8 +506,6 @@ struct GAMEPLAYABILITIES_API FGameplayAbilityTargetData_Mesh : public FGameplayA
 	{}
 
 
-	//UPROPERTY(BlueprintReadWrite, EditAnywhere, Category = Projectile)
-
 	/** Actor who owns the named component. Actor's location is used as start point if component cannot be found. */
 	UPROPERTY(BlueprintReadWrite, EditAnywhere, Category = Targeting)
 	AActor* SourceActor;
@@ -524,13 +522,6 @@ struct GAMEPLAYABILITIES_API FGameplayAbilityTargetData_Mesh : public FGameplayA
 	UPROPERTY(BlueprintReadWrite, EditAnywhere, Category = Targeting)
 	FVector_NetQuantize TargetPoint;
 
-	//virtual TArray<AActor*>	GetActors() const
-
-	// -------------------------------------
-
-	//virtual bool HasHitResult() const
-	//virtual const FHitResult* GetHitResult() const
-
 	// -------------------------------------
 
 	virtual bool HasOrigin() const
