Renamed Wait ability task to VisualizeTargeting. Added task-cancel system. Updated MissileSwarm to use the new features.

diff --git a/Engine/Source/Runtime/GameplayAbilities/Private/Abilities/GameplayAbility.cpp b/Engine/Source/Runtime/GameplayAbilities/Private/Abilities/GameplayAbility.cpp
--- a/Engine/Source/Runtime/GameplayAbilities/Private/Abilities/GameplayAbility.cpp
+++ b/Engine/Source/Runtime/GameplayAbilities/Private/Abilities/GameplayAbility.cpp
@@ -572,6 +572,18 @@ void UGameplayAbility::TaskStarted(UAbilityTask* NewTask)
 	ActiveTasks.Add(NewTask);
 }
 
+void UGameplayAbility::CancelTaskByInstanceName(FName InstanceName)
+{
+	TArray<TWeakObjectPtr<UAbilityTask>> NamedTasks = ActiveTasks.FilterByPredicate<FAbilityInstanceNamePredicate>(FAbilityInstanceNamePredicate(InstanceName));
+	for (int32 i = NamedTasks.Num() - 1; i >= 0; --i)
+	{
+		if (NamedTasks[i].IsValid())
+		{
+			NamedTasks[i].Get()->EndTask();
+		}
+	}
+}
+
 void UGameplayAbility::TaskEnded(UAbilityTask* Task)
 {
 	ActiveTasks.Remove(Task);
diff --git a/Engine/Source/Runtime/GameplayAbilities/Private/Abilities/Tasks/AbilityTask_VisualizeTargeting.cpp b/Engine/Source/Runtime/GameplayAbilities/Private/Abilities/Tasks/AbilityTask_VisualizeTargeting.cpp
new file mode 100644
--- /dev/null
+++ b/Engine/Source/Runtime/GameplayAbilities/Private/Abilities/Tasks/AbilityTask_VisualizeTargeting.cpp
@@ -0,0 +1,103 @@
+
+#include "AbilitySystemPrivatePCH.h"
+#include "GameplayAbilityTargetActor.h"
+#include "AbilitySystemComponent.h"
+#include "Abilities/Tasks/AbilityTask_VisualizeTargeting.h"
+
+UAbilityTask_VisualizeTargeting::UAbilityTask_VisualizeTargeting(const class FPostConstructInitializeProperties& PCIP)
+	: Super(PCIP)
+{
+
+}
+
+UAbilityTask_VisualizeTargeting* UAbilityTask_VisualizeTargeting::VisualizeTargeting(UObject* WorldContextObject, TSubclassOf<AGameplayAbilityTargetActor> InTargetClass, FName TaskInstanceName, float Duration)
+{
+	auto MyObj = NewTask<UAbilityTask_VisualizeTargeting>(WorldContextObject, TaskInstanceName);		//Register for abort list here, providing a given FName as a key
+	MyObj->TargetClass = InTargetClass;
+	if (Duration > 0.0f)
+	{
+		MyObj->GetWorld()->GetTimerManager().SetTimer(MyObj, &UAbilityTask_VisualizeTargeting::OnTimeElapsed, Duration, false);
+	}
+	return MyObj;
+}
+
+// ---------------------------------------------------------------------------------------
+
+bool UAbilityTask_VisualizeTargeting::BeginSpawningActor(UObject* WorldContextObject, TSubclassOf<AGameplayAbilityTargetActor> TargetClass, AGameplayAbilityTargetActor*& SpawnedActor)
+{
+	SpawnedActor = nullptr;
+
+	UGameplayAbility* MyAbility = Ability.Get();
+	if (MyAbility)
+	{
+		const AGameplayAbilityTargetActor* CDO = CastChecked<AGameplayAbilityTargetActor>(TargetClass->GetDefaultObject());
+		MyTargetActor = CDO;
+
+		bool Replicates = CDO->GetReplicates();
+		bool StaticFunc = CDO->StaticTargetFunction;
+		bool IsLocallyControlled = MyAbility->GetCurrentActorInfo()->IsLocallyControlled();
+
+		if (Replicates && StaticFunc)
+		{
+			// We can't replicate a staticFunc target actor, since we are just calling a static function and not spawning an actor at all!
+			ABILITY_LOG(Fatal, TEXT("AbilityTargetActor class %s can't be Replicating and Static"), *TargetClass->GetName());
+			Replicates = false;
+		}
+
+		// Spawn the actor if this is a locally controlled ability (always) or if this is a replicating targeting mode.
+		// (E.g., server will spawn this target actor to replicate to all non owning clients)
+		if (Replicates || IsLocallyControlled)
+		{
+			UClass* Class = *TargetClass;
+			if (Class != NULL)
+			{
+				UWorld* World = GEngine->GetWorldFromContextObject(WorldContextObject);
+				SpawnedActor = World->SpawnActorDeferred<AGameplayAbilityTargetActor>(Class, FVector::ZeroVector, FRotator::ZeroRotator, NULL, NULL, true);
+			}
+
+			MyTargetActor = SpawnedActor;
+
+			AGameplayAbilityTargetActor* TargetActor = CastChecked<AGameplayAbilityTargetActor>(SpawnedActor);
+			if (TargetActor)
+			{
+				TargetActor->MasterPC = MyAbility->GetCurrentActorInfo()->PlayerController.Get();
+			}
+		}
+	}
+	return (SpawnedActor != nullptr);
+}
+
+void UAbilityTask_VisualizeTargeting::FinishSpawningActor(UObject* WorldContextObject, AGameplayAbilityTargetActor* SpawnedActor)
+{
+	if (SpawnedActor)
+	{
+		check(MyTargetActor == SpawnedActor);
+
+		const FTransform SpawnTransform = AbilitySystemComponent->GetOwner()->GetTransform();
+
+		SpawnedActor->FinishSpawning(SpawnTransform);
+		AbilitySystemComponent->SpawnedTargetActors.Push(SpawnedActor);
+		SpawnedActor->StartTargeting(Ability.Get());
+	}
+}
+
+// ---------------------------------------------------------------------------------------
+
+void UAbilityTask_VisualizeTargeting::OnDestroy(bool AbilityEnded)
+{
+	if (MyTargetActor.IsValid() && !MyTargetActor->HasAnyFlags(RF_ClassDefaultObject))
+	{
+		MyTargetActor->Destroy();
+	}
+
+	Super::OnDestroy(AbilityEnded);
+}
+
+
+// --------------------------------------------------------------------------------------
+
+void UAbilityTask_VisualizeTargeting::OnTimeElapsed()
+{
+	TimeElapsed.Broadcast();
+	EndTask();
+}
diff --git a/Engine/Source/Runtime/GameplayAbilities/Private/Abilities/Tasks/AbilityTask_Wait.cpp b/Engine/Source/Runtime/GameplayAbilities/Private/Abilities/Tasks/AbilityTask_Wait.cpp
deleted file mode 100644
diff --git a/Engine/Source/Runtime/GameplayAbilities/Public/Abilities/GameplayAbility.h b/Engine/Source/Runtime/GameplayAbilities/Public/Abilities/GameplayAbility.h
--- a/Engine/Source/Runtime/GameplayAbilities/Public/Abilities/GameplayAbility.h
+++ b/Engine/Source/Runtime/GameplayAbilities/Public/Abilities/GameplayAbility.h
@@ -268,6 +268,10 @@ protected:
 	/** Destroys intanced-per-execution abilities. Instance-per-actor abilities should 'reset'. Non instance abilities - what can we do? */
 	virtual void CancelAbility(const FGameplayAbilityActorInfo* ActorInfo, const FGameplayAbilityActivationInfo ActivationInfo) { }
 
+	/** Destroys intanced-per-execution abilities. Instance-per-actor abilities should 'reset'. Non instance abilities - what can we do? */
+	UFUNCTION(BlueprintCallable, Category = Ability)
+	void CancelTaskByInstanceName(FName InstanceName);
+
 	// -------------------------------------
 	//	EndAbility
 	// -------------------------------------
@@ -387,8 +391,6 @@ protected:
 
 	void MontageBranchPoint_AbilityDecisionStart(const FGameplayAbilityActorInfo* ActorInfo) const;
 	
-	
-
 private:
 
 	/** 
diff --git a/Engine/Source/Runtime/GameplayAbilities/Public/Abilities/Tasks/AbilityTask.h b/Engine/Source/Runtime/GameplayAbilities/Public/Abilities/Tasks/AbilityTask.h
--- a/Engine/Source/Runtime/GameplayAbilities/Public/Abilities/Tasks/AbilityTask.h
+++ b/Engine/Source/Runtime/GameplayAbilities/Public/Abilities/Tasks/AbilityTask.h
@@ -88,13 +88,14 @@ class GAMEPLAYABILITIES_API UAbilityTask : public UObject
 
 	/** Helper function for instantiating and initializing a new task */
 	template <class T>
-	static T*	NewTask(UObject* WorldContextObject)
+	static T*	NewTask(UObject* WorldContextObject, FName InstanceName = FName())
 	{
 		check(WorldContextObject);
 
 		T* MyObj = NewObject<T>();
 		UGameplayAbility* ThisAbility = CastChecked<UGameplayAbility>(WorldContextObject);
 		MyObj->InitTask(ThisAbility);
+		MyObj->InstanceName = InstanceName;
 		return MyObj;
 	}
 
@@ -104,6 +105,12 @@ class GAMEPLAYABILITIES_API UAbilityTask : public UObject
 	/** Called explicitly to end the task (usually by the task itself). Calls OnDestroy. */
 	void EndTask();
 
+public:
+
+	/** This name allows us to find the task later so that we can end it. */
+	UPROPERTY()
+	FName InstanceName;
+
 protected:	
 
 	/** End and CleanUp the task - may be called by the task itself or by the owning ability if the ability is ending. Do NOT call directly! Call EndTask() or AbilityEnded() */
@@ -111,4 +118,20 @@ protected:
 
 	/** If true, this task will receive TickTask calls from AbilitySystemComponent */
 	bool bTickingTask;
-};
\ No newline at end of file
+};
+
+//For searching through lists of ability instances
+struct FAbilityInstanceNamePredicate
+{
+	FAbilityInstanceNamePredicate(FName DesiredInstanceName)
+	{
+		InstanceName = DesiredInstanceName;
+	}
+
+	bool operator()(const TWeakObjectPtr<UAbilityTask> A) const
+	{
+		return (A.IsValid() && !A.Get()->InstanceName.IsNone() && A.Get()->InstanceName.IsValid() && (A.Get()->InstanceName == InstanceName));
+	}
+
+	FName InstanceName;
+};
diff --git a/Engine/Source/Runtime/GameplayAbilities/Public/Abilities/Tasks/AbilityTask_VisualizeTargeting.h b/Engine/Source/Runtime/GameplayAbilities/Public/Abilities/Tasks/AbilityTask_VisualizeTargeting.h
new file mode 100644
--- /dev/null
+++ b/Engine/Source/Runtime/GameplayAbilities/Public/Abilities/Tasks/AbilityTask_VisualizeTargeting.h
@@ -0,0 +1,49 @@
+// Copyright 1998-2014 Epic Games, Inc. All Rights Reserved.
+#pragma once
+
+#include "AbilityTask.h"
+#include "AbilityTask_VisualizeTargeting.generated.h"
+
+class AGameplayAbilityTargetActor;
+
+DECLARE_DYNAMIC_MULTICAST_DELEGATE(FVisualizeTargetingDelegate);
+
+UCLASS(MinimalAPI)
+class UAbilityTask_VisualizeTargeting: public UAbilityTask
+{
+	GENERATED_UCLASS_BODY()
+
+	UPROPERTY(BlueprintAssignable)
+	FVisualizeTargetingDelegate TimeElapsed;
+
+	void OnTimeElapsed();
+
+	/** Spawns Targeting actor. */
+	UFUNCTION(BlueprintCallable, meta=(HidePin = "WorldContextObject", DefaultToSelf = "WorldContextObject", BlueprintInternalUseOnly = "true", HideSpawnParms="Instigator"), Category="Ability|Tasks")
+	static UAbilityTask_VisualizeTargeting* VisualizeTargeting(UObject* WorldContextObject, TSubclassOf<AGameplayAbilityTargetActor> Class, FName TaskInstanceName, float Duration = -1.0f);
+
+	UFUNCTION(BlueprintCallable, meta = (HidePin = "WorldContextObject", DefaultToSelf = "WorldContextObject", BlueprintInternalUseOnly = "true"), Category = "Abilities")
+	bool BeginSpawningActor(UObject* WorldContextObject, TSubclassOf<AGameplayAbilityTargetActor> Class, AGameplayAbilityTargetActor*& SpawnedActor);
+
+	UFUNCTION(BlueprintCallable, meta = (HidePin = "WorldContextObject", DefaultToSelf = "WorldContextObject", BlueprintInternalUseOnly = "true"), Category = "Abilities")
+	void FinishSpawningActor(UObject* WorldContextObject, AGameplayAbilityTargetActor* SpawnedActor);
+
+protected:
+
+	virtual void OnDestroy(bool AbilityEnded) override;
+
+	TSubclassOf<AGameplayAbilityTargetActor> TargetClass;
+
+	/** The TargetActor that we spawned, or the class CDO if this is a static targeting task */
+	TWeakObjectPtr<AGameplayAbilityTargetActor>	MyTargetActor;
+};
+
+
+/**
+*	Requirements for using Begin/Finish SpawningActor functionality:
+*		-Have a parameters named 'Class' in your Proxy factor function (E.g., WaitTargetdata)
+*		-Have a function named BeginSpawningActor w/ the same Class parameter
+*			-This function should spawn the actor with SpawnActorDeferred and return true/false if it spawned something.
+*		-Have a function named FinishSpawningActor w/ an AActor* of the class you spawned
+*			-This function *must* call ExecuteConstruction + PostActorConstruction
+*/
\ No newline at end of file
diff --git a/Engine/Source/Runtime/GameplayAbilities/Public/Abilities/Tasks/AbilityTask_Wait.h b/Engine/Source/Runtime/GameplayAbilities/Public/Abilities/Tasks/AbilityTask_Wait.h
deleted file mode 100644
