Fix crash when global curve tables were being GC'd in non editor builds

diff --git a/Engine/Source/Runtime/GameplayAbilities/Public/AbilitySystemGlobals.h b/Engine/Source/Runtime/GameplayAbilities/Public/AbilitySystemGlobals.h
--- a/Engine/Source/Runtime/GameplayAbilities/Public/AbilitySystemGlobals.h
+++ b/Engine/Source/Runtime/GameplayAbilities/Public/AbilitySystemGlobals.h
@@ -66,10 +66,13 @@ protected:
 
 private:
 
+	UPROPERTY()
 	UCurveTable* GlobalCurveTable;
 
+	UPROPERTY()
 	UCurveTable* GlobalAttributeDefaultsTable;
 
+	UPROPERTY()
 	UDataTable* GlobalAttributeMetaDataTable;
 
 	TSharedPtr<FAttributeSetInitter> GlobalAttributeSetInitter;
