Support for using GameplayCues outside of GameplayEffects

diff --git a/Engine/Source/Runtime/GameplayAbilities/Private/Abilities/GameplayAbility.cpp b/Engine/Source/Runtime/GameplayAbilities/Private/Abilities/GameplayAbility.cpp
--- a/Engine/Source/Runtime/GameplayAbilities/Private/Abilities/GameplayAbility.cpp
+++ b/Engine/Source/Runtime/GameplayAbilities/Private/Abilities/GameplayAbility.cpp
@@ -810,3 +810,37 @@ bool FGameplayAbilityTargetData_SingleTargetHit::NetSerialize(FArchive& Ar, clas
 	return true;
 }
 
+
+
+// -----------------------------------------------------------------
+
+
+/**
+ *	Helper methods for adding GaAmeplayCues without having to go through GameplayEffects.
+ *	For now, none of these will happen predictively. We can eventually build this out more to 
+ *	work with the PredictionKey system.
+ */
+
+void UGameplayAbility::ExecuteGameplayCue(FGameplayTag GameplayCueTag)
+{
+	if (ensure(CurrentActorInfo) && CurrentActorInfo->IsNetAuthority())
+	{
+		CurrentActorInfo->AbilitySystemComponent->ExecuteGameplayCue(GameplayCueTag);
+	}
+}
+
+void UGameplayAbility::AddGameplayCue(FGameplayTag GameplayCueTag)
+{
+	if (ensure(CurrentActorInfo) && CurrentActorInfo->IsNetAuthority())
+	{
+		CurrentActorInfo->AbilitySystemComponent->AddGameplayCue(GameplayCueTag);
+	}
+}
+
+void UGameplayAbility::RemoveGameplayCue(FGameplayTag GameplayCueTag)
+{
+	if (ensure(CurrentActorInfo) && CurrentActorInfo->IsNetAuthority())
+	{
+		CurrentActorInfo->AbilitySystemComponent->RemoveGameplayCue(GameplayCueTag);
+	}
+}
\ No newline at end of file
diff --git a/Engine/Source/Runtime/GameplayAbilities/Private/AbilitySystemComponent.cpp b/Engine/Source/Runtime/GameplayAbilities/Private/AbilitySystemComponent.cpp
--- a/Engine/Source/Runtime/GameplayAbilities/Private/AbilitySystemComponent.cpp
+++ b/Engine/Source/Runtime/GameplayAbilities/Private/AbilitySystemComponent.cpp
@@ -39,6 +39,7 @@ UAbilitySystemComponent::UAbilitySystemComponent(const class FPostConstructIniti
 	PrimaryComponentTick.bCanEverTick = true;
 	
 	ActiveGameplayEffects.Owner = this;
+	ActiveGameplayCues.Owner = this;
 
 	SetIsReplicated(true);
 	ReplicatedPredictionKey = 0;
@@ -73,7 +74,7 @@ void UAbilitySystemComponent::K2_InitStats(TSubclassOf<class UAttributeSet> Attr
 	InitStats(Attributes, DataTable);
 }
 
-const UAttributeSet*	UAbilitySystemComponent::GetOrCreateAttributeSubobject(const TSubclassOf<UAttributeSet> AttributeClass)
+const UAttributeSet* UAbilitySystemComponent::GetOrCreateAttributeSubobject(const TSubclassOf<UAttributeSet> AttributeClass)
 {
 	AActor *OwningActor = GetOwner();
 	const UAttributeSet *MyAttributes  = NULL;
@@ -442,11 +443,11 @@ void UAbilitySystemComponent::InvokeGameplayCueEvent(const FGameplayEffectSpec &
 		ABILITY_LOG(Warning, TEXT("InvokeGameplayCueEvent: %s"), *Spec.ToSimpleString());
 	}
 
-
-	AActor *ActorOwner = AbilityActorInfo->Actor.Get();
-	IGameplayCueInterface * GameplayCueInterface = InterfaceCast<IGameplayCueInterface>(ActorOwner);
+	AActor* ActorOwner = AbilityActorInfo->Actor.Get();
+	IGameplayCueInterface* GameplayCueInterface = InterfaceCast<IGameplayCueInterface>(ActorOwner);
 	if (!GameplayCueInterface)
 	{
+		ABILITY_LOG(Warning, TEXT("InvokeGameplayCueEvent %s on Actor %s that is not IGameplayCueInterface"), *Spec.ToSimpleString(), ActorOwner ? *ActorOwner->GetName() : TEXT("NULL"));
 		return;
 	}
 
@@ -469,12 +470,81 @@ void UAbilitySystemComponent::InvokeGameplayCueEvent(const FGameplayEffectSpec &
 	}
 }
 
+void UAbilitySystemComponent::ExecuteGameplayCue(const FGameplayTag GameplayCueTag)
+{
+	if (IsOwnerActorAuthoritative())
+	{
+		NetMulticast_InvokeGameplayCueExecuted(GameplayCueTag);
+	}
+	else
+	{
+		InvokeGameplayCueEvent(GameplayCueTag, EGameplayCueEvent::Executed);
+	}
+}
+
+void UAbilitySystemComponent::AddGameplayCue(const FGameplayTag GameplayCueTag)
+{
+	if (IsOwnerActorAuthoritative())
+	{
+		ActiveGameplayCues.AddCue(GameplayCueTag);
+		NetMulticast_InvokeGameplayCueAdded(GameplayCueTag);
+	}
+	else
+	{
+		// Allow for predictive gameplaycue events? Needs more thought
+		InvokeGameplayCueEvent(GameplayCueTag, EGameplayCueEvent::OnActive);
+		InvokeGameplayCueEvent(GameplayCueTag, EGameplayCueEvent::WhileActive);
+	}
+}
+
+void UAbilitySystemComponent::RemoveGameplayCue(const FGameplayTag GameplayCueTag)
+{
+	InvokeGameplayCueEvent(GameplayCueTag, EGameplayCueEvent::Removed);
+	if (IsOwnerActorAuthoritative())
+	{
+		ActiveGameplayCues.RemoveCue(GameplayCueTag);
+	}
+}
+
+void UAbilitySystemComponent::InvokeGameplayCueEvent(const FGameplayTag GameplayCueTag, EGameplayCueEvent::Type EventType)
+{
+	AActor* ActorOwner = AbilityActorInfo->Actor.Get();
+		IGameplayCueInterface* GameplayCueInterface = InterfaceCast<IGameplayCueInterface>(ActorOwner);
+	if (!GameplayCueInterface)
+	{
+		return;
+	}
+
+	FGameplayCueParameters CueParameters;
+	CueParameters.InstigatorContext.AddInstigator(ActorOwner);	// Treat ourselves as the instigator for now. We may need to allow this to be passed in eventually.
+
+	const float NormalizedMagnitude = 1.f;
+
+	GameplayCueInterface->HandleGameplayCue(ActorOwner, GameplayCueTag, EventType, CueParameters);
+}
 
-void UAbilitySystemComponent::NetMulticast_InvokeGameplayCueExecuted_Implementation(const FGameplayEffectSpec Spec)
+void UAbilitySystemComponent::NetMulticast_InvokeGameplayCueExecuted_FromSpec_Implementation(const FGameplayEffectSpec Spec)
 {
 	InvokeGameplayCueEvent(Spec, EGameplayCueEvent::Executed);
 }
 
+void UAbilitySystemComponent::NetMulticast_InvokeGameplayCueExecuted_Implementation(const FGameplayTag GameplayCueTag)
+{
+	InvokeGameplayCueEvent(GameplayCueTag, EGameplayCueEvent::Executed);
+}
+
+void UAbilitySystemComponent::NetMulticast_InvokeGameplayCueAdded_Implementation(const FGameplayTag GameplayCueTag)
+{
+	InvokeGameplayCueEvent(GameplayCueTag, EGameplayCueEvent::OnActive);
+}
+
+bool UAbilitySystemComponent::IsGameplayCueActive(const FGameplayTag GameplayCueTag) const
+{
+	return (ActiveGameplayEffects.HasMatchingGameplayTag(GameplayCueTag) || ActiveGameplayCues.HasMatchingGameplayTag(GameplayCueTag));
+}
+
+// ----------------------------------------------------------------------------------------
+
 void UAbilitySystemComponent::AddDependancyToAttribute(FGameplayAttribute Attribute, const TWeakPtr<FAggregator> InDependant)
 {
 	ActiveGameplayEffects.AddDependancyToAttribute(Attribute, InDependant);
@@ -525,6 +595,7 @@ void UAbilitySystemComponent::GetLifetimeReplicatedProps(TArray< FLifetimeProper
 
 	DOREPLIFETIME(UAbilitySystemComponent, SpawnedAttributes);
 	DOREPLIFETIME(UAbilitySystemComponent, ActiveGameplayEffects);
+	DOREPLIFETIME(UAbilitySystemComponent, ActiveGameplayCues);
 
 	DOREPLIFETIME(UAbilitySystemComponent, ReplicatedInstancedAbilities);
 	DOREPLIFETIME(UAbilitySystemComponent, ActivatableAbilities);
@@ -618,15 +689,6 @@ uint32 UAbilitySystemComponent::GetNextPredictionKey()
 	return ++LocalPredictionKey;
 }
 
-void UAbilitySystemComponent::DispatchBlueprintCustomHandler(AActor* Actor, UFunction* Func, EGameplayCueEvent::Type EventType, FGameplayCueParameters Parameters)
-{
-	AbilitySystemComponent_eventBlueprintCustomHandler_Parms Parms;
-	Parms.EventType = EventType;
-	Parms.Parameters = Parameters;
-
-	Actor->ProcessEvent(Func, &Parms);
-}
-
 // ---------------------------------------------------------------------------------------
 
 void UAbilitySystemComponent::PrintAllGameplayEffects() const
diff --git a/Engine/Source/Runtime/GameplayAbilities/Private/GameplayCueInterface.cpp b/Engine/Source/Runtime/GameplayAbilities/Private/GameplayCueInterface.cpp
--- a/Engine/Source/Runtime/GameplayAbilities/Private/GameplayCueInterface.cpp
+++ b/Engine/Source/Runtime/GameplayAbilities/Private/GameplayCueInterface.cpp
@@ -38,4 +38,56 @@ void IGameplayCueInterface::HandleGameplayCue(AActor *Self, FGameplayTag Gamepla
 	}
 
 	// Fall back to a possible generic handler if available
+}
+
+
+
+void FActiveGameplayCue::PreReplicatedRemove(const struct FActiveGameplayCueContainer &InArray)
+{
+	InArray.Owner->InvokeGameplayCueEvent(GameplayCueTag, EGameplayCueEvent::Removed);
+}
+
+void FActiveGameplayCue::PostReplicatedAdd(const struct FActiveGameplayCueContainer &InArray)
+{
+	InArray.Owner->InvokeGameplayCueEvent(GameplayCueTag, EGameplayCueEvent::WhileActive);
+}
+
+void FActiveGameplayCueContainer::AddCue(const FGameplayTag& Tag)
+{
+	UWorld* World = Owner->GetWorld();
+
+	FActiveGameplayCue	NewCue;
+	NewCue.GameplayCueTag =Tag;
+	MarkItemDirty(NewCue);
+
+	GameplayCues.Add(NewCue);
+}
+
+void FActiveGameplayCueContainer::RemoveCue(const FGameplayTag& Tag)
+{
+	for (int32 idx=0; idx < GameplayCues.Num(); ++idx)
+	{
+		FActiveGameplayCue& Cue = GameplayCues[idx];
+
+		if (Cue.GameplayCueTag == Tag)
+		{
+			GameplayCues.RemoveAt(idx);
+			MarkArrayDirty();
+			return;
+		}
+	}
+}
+
+bool FActiveGameplayCueContainer::HasMatchingGameplayTag(FGameplayTag TagToCheck) const
+{
+	for (const FActiveGameplayCue& Cue : GameplayCues)
+	{
+		/** We expect that if TagToCheck=(a.b) and we have a (a.b.c) tag, then we match  */
+		if (Cue.GameplayCueTag.Matches(EGameplayTagMatchType::IncludeParentTags, TagToCheck, EGameplayTagMatchType::Explicit))
+		{
+			return true;
+		}
+	}
+
+	return false;
 }
\ No newline at end of file
diff --git a/Engine/Source/Runtime/GameplayAbilities/Private/GameplayEffect.cpp b/Engine/Source/Runtime/GameplayAbilities/Private/GameplayEffect.cpp
--- a/Engine/Source/Runtime/GameplayAbilities/Private/GameplayEffect.cpp
+++ b/Engine/Source/Runtime/GameplayAbilities/Private/GameplayEffect.cpp
@@ -1378,7 +1378,7 @@ void FActiveGameplayEffectsContainer::ExecuteActiveEffectsFrom(const FGameplayEf
 		// TODO: check replication policy. Right now we will replicate every execute via a multicast RPC
 
 		ABILITY_LOG(Log, TEXT("Invoking Execute GameplayCue for %s"), *Spec.ToSimpleString() );
-		Owner->NetMulticast_InvokeGameplayCueExecuted(Spec);
+		Owner->NetMulticast_InvokeGameplayCueExecuted_FromSpec(Spec);
 	}
 
 }
@@ -1633,30 +1633,49 @@ FActiveGameplayEffect & FActiveGameplayEffectsContainer::CreateNewActiveGameplay
 	return NewEffect;
 }
 
-void FActiveGameplayEffectsContainer::InternalOnActiveGameplayEffectAdded(const FActiveGameplayEffect& Effect)
+void FActiveGameplayEffectsContainer::UpdateTagMap(const FGameplayTag& BaseTag, int32 CountDelta, IGameplayTagsModule& GameplayTagsModule)
 {
-	// Update gameplaytag count and broadcast delegate if we just added this tag (count=0, prior to increment)
-	IGameplayTagsModule& GameplayTagsModule = IGameplayTagsModule::Get();
-	for (auto BaseTagIt = Effect.Spec.Def->OwnedTagsContainer.CreateConstIterator(); BaseTagIt; ++BaseTagIt)
+	FGameplayTagContainer TagAndParentsContainer = GameplayTagsModule.GetGameplayTagsManager().RequestGameplayTagParents(BaseTag);
+	for (auto TagIt = TagAndParentsContainer.CreateConstIterator(); TagIt; ++TagIt)
 	{
-		const FGameplayTag& BaseTag = *BaseTagIt;
-		FGameplayTagContainer TagAndParentsContainer = GameplayTagsModule.GetGameplayTagsManager().RequestGameplayTagParents(BaseTag);
-		for (auto TagIt = TagAndParentsContainer.CreateConstIterator(); TagIt; ++TagIt)
+		const FGameplayTag& Tag = *TagIt;
+		int32& Count = GameplayTagCountMap.FindOrAdd(Tag);
+
+		bool WasZero = Count == 0;
+		Count = FMath::Max(Count + CountDelta, 0);
+		
+		// If we went from 0->1 or 1->0
+		if (WasZero || Count == 0)
 		{
-			const FGameplayTag& Tag = *TagIt;
-			int32& Count = GameplayTagCountMap.FindOrAdd(Tag);
-			if (Count++ == 0)
+			FOnGameplayEffectTagCountChanged *Delegate = GameplayTagEventMap.Find(Tag);
+			if (Delegate)
 			{
-				FOnGameplayEffectTagCountChanged *Delegate = GameplayTagEventMap.Find(Tag);
-				if (Delegate)
-				{
-					Delegate->Broadcast(Tag, 1);
-				}
+				Delegate->Broadcast(Tag, Count);
 			}
 		}
 	}
 }
 
+void FActiveGameplayEffectsContainer::UpdateTagMap(const FGameplayTagContainer& Container, int32 CountDelta, IGameplayTagsModule& GameplayTagsModule)
+{
+	for (auto TagIt = Container.CreateConstIterator(); TagIt; ++TagIt)
+	{
+		const FGameplayTag& Tag = *TagIt;
+		UpdateTagMap(Tag, CountDelta, GameplayTagsModule);
+	}
+}
+
+void FActiveGameplayEffectsContainer::InternalOnActiveGameplayEffectAdded(const FActiveGameplayEffect& Effect)
+{
+	// Update gameplaytag count and broadcast delegate if we just added this tag (count=0, prior to increment)
+	IGameplayTagsModule& GameplayTagsModule = IGameplayTagsModule::Get();
+	UpdateTagMap(Effect.Spec.Def->OwnedTagsContainer, 1, GameplayTagsModule);
+	for (const FGameplayEffectCue& Cue : Effect.Spec.Def->GameplayCues)
+	{
+		UpdateTagMap(Cue.GameplayCueTags, 1, GameplayTagsModule);
+	}	
+}
+
 
 bool FActiveGameplayEffectsContainer::RemoveActiveGameplayEffect(FActiveGameplayEffectHandle Handle)
 {
@@ -1713,29 +1732,15 @@ bool FActiveGameplayEffectsContainer::InternalRemoveActiveGameplayEffect(int32 I
 /** This does cleanup that has to happen whether the effect is being removed locally or due to replication */
 void FActiveGameplayEffectsContainer::InternalOnActiveGameplayEffectRemoved(const FActiveGameplayEffect& Effect)
 {
-
 	Effect.OnRemovedDelegate.Broadcast();
 
 	// Update gameplaytag count and broadcast delegate if we are at 0
 	IGameplayTagsModule& GameplayTagsModule = IGameplayTagsModule::Get();
-	for (auto BaseTagIt = Effect.Spec.Def->OwnedTagsContainer.CreateConstIterator(); BaseTagIt; ++BaseTagIt)
+	UpdateTagMap(Effect.Spec.Def->OwnedTagsContainer, -1, GameplayTagsModule);
+
+	for (const FGameplayEffectCue& Cue : Effect.Spec.Def->GameplayCues)
 	{
-		const FGameplayTag& BaseTag = *BaseTagIt;
-		FGameplayTagContainer TagAndParentsContainer = GameplayTagsModule.GetGameplayTagsManager().RequestGameplayTagParents(BaseTag);
-		for (auto TagIt = TagAndParentsContainer.CreateConstIterator(); TagIt; ++TagIt)
-		{
-			const FGameplayTag& Tag = *TagIt;
-			int32& Count = GameplayTagCountMap[Tag];
-			if (--Count == 0)
-			{
-				FOnGameplayEffectTagCountChanged *Delegate = GameplayTagEventMap.Find(Tag);
-				if (Delegate)
-				{
-					Delegate->Broadcast(Tag, 0);
-				}
-			}
-			check(Count >= 0);
-		}
+		UpdateTagMap(Cue.GameplayCueTags, -1, GameplayTagsModule);
 	}
 }
 
diff --git a/Engine/Source/Runtime/GameplayAbilities/Public/Abilities/GameplayAbility.h b/Engine/Source/Runtime/GameplayAbilities/Public/Abilities/GameplayAbility.h
--- a/Engine/Source/Runtime/GameplayAbilities/Public/Abilities/GameplayAbility.h
+++ b/Engine/Source/Runtime/GameplayAbilities/Public/Abilities/GameplayAbility.h
@@ -290,6 +290,22 @@ protected:
 	virtual void ClientActivateAbility(const FGameplayAbilityActorInfo* ActorInfo, const FGameplayAbilityActivationInfo ActivationInfo);
 
 	// -------------------------------------
+	//	GameplayCue
+	//	Abilities can invoke GameplayCues without having to create GameplayEffects
+	// -------------------------------------
+	
+	UFUNCTION(BlueprintCallable, Category = Ability, meta=(GameplayTagFilter="GameplayCue"))
+	virtual void ExecuteGameplayCue(FGameplayTag GameplayCueTag);
+
+	UFUNCTION(BlueprintCallable, Category = Ability, meta=(GameplayTagFilter="GameplayCue"))
+	virtual void AddGameplayCue(FGameplayTag GameplayCueTag);
+
+	UFUNCTION(BlueprintCallable, Category = Ability, meta=(GameplayTagFilter="GameplayCue"))
+	virtual void RemoveGameplayCue(FGameplayTag GameplayCueTag);
+
+
+	// -------------------------------------
+
 
 	bool IsInstantiated() const
 	{
@@ -390,7 +406,6 @@ protected:
 	void MontageBranchPoint_AbilityDecisionStop(const FGameplayAbilityActorInfo* ActorInfo) const;
 
 	void MontageBranchPoint_AbilityDecisionStart(const FGameplayAbilityActorInfo* ActorInfo) const;
-	
 private:
 
 	/** 
diff --git a/Engine/Source/Runtime/GameplayAbilities/Public/AbilitySystemComponent.h b/Engine/Source/Runtime/GameplayAbilities/Public/AbilitySystemComponent.h
--- a/Engine/Source/Runtime/GameplayAbilities/Public/AbilitySystemComponent.h
+++ b/Engine/Source/Runtime/GameplayAbilities/Public/AbilitySystemComponent.h
@@ -2,6 +2,7 @@
 #pragma once
 
 #include "GameplayEffect.h"
+#include "GameplayCueInterface.h"
 #include "Abilities/GameplayAbilityTypes.h"
 #include "AbilitySystemComponent.generated.h"
 
@@ -320,12 +321,36 @@ class GAMEPLAYABILITIES_API UAbilitySystemComponent : public UActorComponent
 	//	GameplayCues
 	// 
 	// ----------------------------------------------------------------------------------------------------------------
+	 
 
+	// GameplayCues can come from GameplayEffectSpecs
 
-	UFUNCTION(BlueprintImplementableEvent, Category = GameplayCue, meta = (BlueprintInternalUseOnly = "true"))
-	void BlueprintCustomHandler(EGameplayCueEvent::Type EventType, FGameplayCueParameters Parameters);
+	UFUNCTION(NetMulticast, unreliable)
+	void NetMulticast_InvokeGameplayCueExecuted_FromSpec(const FGameplayEffectSpec Spec);
+
+	void InvokeGameplayCueEvent(const FGameplayEffectSpec &Spec, EGameplayCueEvent::Type EventType);
+
+	// GameplayCues can also come on their own. For now these will have no additional parameters (just a tag) though that it something we could
+	// support down the road if we wanted.
+
+	void ExecuteGameplayCue(const FGameplayTag GameplayCueTag);
+
+	void AddGameplayCue(const FGameplayTag GameplayCueTag);
+	
+	void RemoveGameplayCue(const FGameplayTag GameplayCueTag);
+
+	UFUNCTION(NetMulticast, unreliable)
+	void NetMulticast_InvokeGameplayCueExecuted(const FGameplayTag GameplayCueTag);
+
+	UFUNCTION(NetMulticast, unreliable)
+	void NetMulticast_InvokeGameplayCueAdded(const FGameplayTag GameplayCueTag);
+
+	void InvokeGameplayCueEvent(const FGameplayTag GameplayCueTag, EGameplayCueEvent::Type EventType);
+
+	/** Allows polling to see if a GameplayCue is active. We expect most GameplayCue hanlding to be event based, but some cases we may need to check if a GamepalyCue is active (Animation Blueprint for example) */
+	UFUNCTION(BlueprintCallable, Category="GameplayCue", meta=(GameplayTagFilter="GameplayCue"))
+	bool IsGameplayCueActive(const FGameplayTag GameplayCueTag) const;
 
-	static void DispatchBlueprintCustomHandler(AActor* Actor, UFunction* Func, EGameplayCueEvent::Type EventType, FGameplayCueParameters Parameters);
 
 	// ----------------------------------------------------------------------------------------------------------------
 
@@ -546,15 +571,15 @@ private:
 	UPROPERTY(ReplicatedUsing=OnRep_GameplayEffects)
 	FActiveGameplayEffectsContainer	ActiveGameplayEffects;
 
+	UPROPERTY(ReplicatedUsing=OnRep_GameplayEffects)
+	FActiveGameplayCueContainer	ActiveGameplayCues;
+
 	UFUNCTION()
 	void OnRep_GameplayEffects();
 
 	// ---------------------------------------------
-
-	UFUNCTION(NetMulticast, unreliable)
-	void NetMulticast_InvokeGameplayCueExecuted(const FGameplayEffectSpec Spec);
-
-	void InvokeGameplayCueEvent(const FGameplayEffectSpec &Spec, EGameplayCueEvent::Type EventType);	
+	
+	
 	
 	// ---------------------------------------------
 
diff --git a/Engine/Source/Runtime/GameplayAbilities/Public/GameplayCueInterface.h b/Engine/Source/Runtime/GameplayAbilities/Public/GameplayCueInterface.h
--- a/Engine/Source/Runtime/GameplayAbilities/Public/GameplayCueInterface.h
+++ b/Engine/Source/Runtime/GameplayAbilities/Public/GameplayCueInterface.h
@@ -27,4 +27,64 @@ class GAMEPLAYABILITIES_API IGameplayCueInterface
 	void BlueprintCustomHandler(EGameplayCueEvent::Type EventType, FGameplayCueParameters Parameters);
 
 	static void DispatchBlueprintCustomHandler(AActor* Actor, UFunction* Func, EGameplayCueEvent::Type EventType, FGameplayCueParameters Parameters);	
+};
+
+
+/**
+ *	This is meant to provide another way of using GameplayCues without having to go through GameplayEffects.
+ *	E.g., it is convenient if GameplayAbilities can issue replicated GameplayCues without having to create
+ *	a GameplayEffect.
+ *	
+ *	Essentially provides bare necessities to replicate GameplayCue Tags.
+ */
+
+
+struct FActiveGameplayCueContainer;
+
+USTRUCT(BlueprintType)
+struct FActiveGameplayCue : public FFastArraySerializerItem
+{
+	GENERATED_USTRUCT_BODY()
+
+	UPROPERTY()
+	FGameplayTag GameplayCueTag;
+
+	void PreReplicatedRemove(const struct FActiveGameplayCueContainer &InArray);
+	void PostReplicatedAdd(const struct FActiveGameplayCueContainer &InArray);
+	void PostReplicatedChange(const struct FActiveGameplayCueContainer &InArray) { }
+};
+
+USTRUCT(BlueprintType)
+struct FActiveGameplayCueContainer : public FFastArraySerializer
+{
+	GENERATED_USTRUCT_BODY()
+
+	UPROPERTY()
+	TArray< FActiveGameplayCue >	GameplayCues;
+
+	UPROPERTY()
+	class UAbilitySystemComponent*	Owner;
+
+	void AddCue(const FGameplayTag& Tag);
+	void RemoveCue(const FGameplayTag& Tag);
+
+	bool HasMatchingGameplayTag(FGameplayTag TagToCheck) const;
+
+	bool NetDeltaSerialize(FNetDeltaSerializeInfo & DeltaParms)
+	{
+		return FastArrayDeltaSerialize<FActiveGameplayCue>(GameplayCues, DeltaParms, *this);
+	}
+
+private:
+
+	int32 GetGameStateTime(const UWorld* World) const;
+};
+
+template<>
+struct TStructOpsTypeTraits< FActiveGameplayCueContainer > : public TStructOpsTypeTraitsBase
+{
+	enum
+	{
+		WithNetDeltaSerializer = true,
+	};
 };
\ No newline at end of file
diff --git a/Engine/Source/Runtime/GameplayAbilities/Public/GameplayEffect.h b/Engine/Source/Runtime/GameplayAbilities/Public/GameplayEffect.h
--- a/Engine/Source/Runtime/GameplayAbilities/Public/GameplayEffect.h
+++ b/Engine/Source/Runtime/GameplayAbilities/Public/GameplayEffect.h
@@ -155,7 +155,7 @@ struct FGameplayEffectCue
 	UPROPERTY(EditDefaultsOnly, Category = GameplayCue)
 	float	MaxLevel;
 
-	UPROPERTY(EditDefaultsOnly, Category = GameplayCue)
+	UPROPERTY(EditDefaultsOnly, Category = GameplayCue, meta = (Categories="GameplayCue"))
 	FGameplayTagContainer GameplayCueTags;
 
 	float NormalizeLevel(float InLevel)
@@ -1282,6 +1282,9 @@ private:
 	/** Called both in server side creation and replication creation/deletion */
 	void InternalOnActiveGameplayEffectAdded(const FActiveGameplayEffect& Effect);
 	void InternalOnActiveGameplayEffectRemoved(const FActiveGameplayEffect& Effect);
+
+	void UpdateTagMap(const FGameplayTagContainer& Container, int32 CountDelta, class IGameplayTagsModule& GameplayTagsModule);
+	void UpdateTagMap(const FGameplayTag& Tag, int32 CountDelta, class IGameplayTagsModule& GameplayTagsModule);
 };
 
 template<>
