Make AbilitySystemComponent implement IGameplayTagDataAsset interface

diff --git a/Engine/Source/Runtime/GameplayAbilities/Public/AbilitySystemComponent.h b/Engine/Source/Runtime/GameplayAbilities/Public/AbilitySystemComponent.h
--- a/Engine/Source/Runtime/GameplayAbilities/Public/AbilitySystemComponent.h
+++ b/Engine/Source/Runtime/GameplayAbilities/Public/AbilitySystemComponent.h
@@ -55,7 +55,7 @@ struct GAMEPLAYABILITIES_API FAttributeDefaults
  */
 
 UCLASS(ClassGroup=AbilitySystem, hidecategories=(Object,LOD,Lighting,Transform,Sockets,TextureStreaming), editinlinenew, meta=(BlueprintSpawnableComponent))
-class GAMEPLAYABILITIES_API UAbilitySystemComponent : public UActorComponent
+class GAMEPLAYABILITIES_API UAbilitySystemComponent : public UActorComponent, public IGameplayTagAssetInterface
 {
 	GENERATED_UCLASS_BODY()
 
@@ -257,13 +257,13 @@ class GAMEPLAYABILITIES_API UAbilitySystemComponent : public UActorComponent
 	// --------------------------------------------
 	// Tags
 	// --------------------------------------------
-	bool HasMatchingGameplayTag(FGameplayTag TagToCheck) const;
+	virtual bool HasMatchingGameplayTag(FGameplayTag TagToCheck) const override;
 
-	bool HasAllMatchingGameplayTags(const FGameplayTagContainer& TagContainer, bool bCountEmptyAsMatch = true) const;
+	virtual bool HasAllMatchingGameplayTags(const FGameplayTagContainer& TagContainer, bool bCountEmptyAsMatch = true) const override;
 
-	bool HasAnyMatchingGameplayTags(const FGameplayTagContainer& TagContainer, bool bCountEmptyAsMatch = true) const;
+	virtual bool HasAnyMatchingGameplayTags(const FGameplayTagContainer& TagContainer, bool bCountEmptyAsMatch = true) const override;
 
-	void GetOwnedGameplayTags(FGameplayTagContainer& TagContainer) const;
+	virtual void GetOwnedGameplayTags(FGameplayTagContainer& TagContainer) const override;
 
 	/** Allow events to be registered for specific gameplay tags being added or removed */
 	FOnGameplayEffectTagCountChanged& RegisterGameplayTagEvent(FGameplayTag Tag);
