non unity fix

diff --git a/Engine/Source/Runtime/GameplayAbilities/Private/GameplayCueInterface.cpp b/Engine/Source/Runtime/GameplayAbilities/Private/GameplayCueInterface.cpp
--- a/Engine/Source/Runtime/GameplayAbilities/Private/GameplayCueInterface.cpp
+++ b/Engine/Source/Runtime/GameplayAbilities/Private/GameplayCueInterface.cpp
@@ -1,6 +1,7 @@
 // Copyright 1998-2014 Epic Games, Inc. All Rights Reserved.
 
 #include "AbilitySystemPrivatePCH.h"
+#include "AbilitySystemComponent.h"
 #include "GameplayCueInterface.h"
 
 UGameplayCueInterface::UGameplayCueInterface(const FPostConstructInitializeProperties& PCIP)
