Properly initialize InstigatorAbilitySystemComponent

diff --git a/Engine/Source/Runtime/GameplayAbilities/Private/GameplayEffect.cpp b/Engine/Source/Runtime/GameplayAbilities/Private/GameplayEffect.cpp
--- a/Engine/Source/Runtime/GameplayAbilities/Private/GameplayEffect.cpp
+++ b/Engine/Source/Runtime/GameplayAbilities/Private/GameplayEffect.cpp
@@ -454,19 +454,10 @@ void FGameplayEffectInstigatorContext::AddInstigator(class AActor *InInstigator)
 	InstigatorAbilitySystemComponent = NULL;
 
 	// Cache off his AbilitySystemComponent.
-	if (Instigator)
+	IAbilitySystemInterface* AbilitySystmeInterface = InterfaceCast<IAbilitySystemInterface>(Instigator);
+	if (!AbilitySystmeInterface)
 	{
-		TArray<UObject*> ChildObjects;
-		GetObjectsWithOuter(Instigator, ChildObjects);
-
-		for (UObject * Obj : ChildObjects)
-		{
-			if (Obj && Obj->GetClass()->IsChildOf(UAbilitySystemComponent::StaticClass()))
-			{
-				InstigatorAbilitySystemComponent = Cast<UAbilitySystemComponent>(Obj);
-				break;
-			}
-		}
+		InstigatorAbilitySystemComponent = AbilitySystmeInterface->GetAbilitySystemComponent();
 	}
 }
 
