Make the AbilitySystemTestPawn implement the IAbilitySystemInterface.

Fixes some warnings when running gameplay ability automation tests.

diff --git a/Engine/Source/Runtime/GameplayAbilities/Private/AbilitySystemTestPawn.cpp b/Engine/Source/Runtime/GameplayAbilities/Private/AbilitySystemTestPawn.cpp
--- a/Engine/Source/Runtime/GameplayAbilities/Private/AbilitySystemTestPawn.cpp
+++ b/Engine/Source/Runtime/GameplayAbilities/Private/AbilitySystemTestPawn.cpp
@@ -30,4 +30,9 @@ void AAbilitySystemTestPawn::PostInitializeComponents()
 		AbilitySystemComponent->InitializeAbilities(DefaultAbilitySet);
 	}
 	*/
+}
+
+UAbilitySystemComponent* AAbilitySystemTestPawn::GetAbilitySystemComponent() const
+{
+	return FindComponentByClass<UAbilitySystemComponent>();
 }
\ No newline at end of file
diff --git a/Engine/Source/Runtime/GameplayAbilities/Private/GameplayEffectTests.cpp b/Engine/Source/Runtime/GameplayAbilities/Private/GameplayEffectTests.cpp
--- a/Engine/Source/Runtime/GameplayAbilities/Private/GameplayEffectTests.cpp
+++ b/Engine/Source/Runtime/GameplayAbilities/Private/GameplayEffectTests.cpp
@@ -4854,7 +4854,7 @@ bool GameplayEffectsTest_ImmunityMod(UWorld *World, FAutomationTestBase * Test)
 		float ExpectedValue = StartHealth;
 		float ActualValue = DestComponent->GetSet<UAbilitySystemTestAttributeSet>()->Health;
 
-		Test->TestTrue(SKILL_TEST_TEXT("Buff Instant Damage Applied"), (ActualValue == ExpectedValue));
+		Test->TestTrue(SKILL_TEST_TEXT("Buff Instant Damage Applied. %.2f == %.2f", ActualValue, ExpectedValue), (ActualValue == ExpectedValue));
 		ABILITY_LOG(Log, TEXT("Final Health: %.2f"), DestComponent->GetSet<UAbilitySystemTestAttributeSet>()->Health);
 	}
 
diff --git a/Engine/Source/Runtime/GameplayAbilities/Public/AbilitySystemTestPawn.h b/Engine/Source/Runtime/GameplayAbilities/Public/AbilitySystemTestPawn.h
--- a/Engine/Source/Runtime/GameplayAbilities/Public/AbilitySystemTestPawn.h
+++ b/Engine/Source/Runtime/GameplayAbilities/Public/AbilitySystemTestPawn.h
@@ -2,15 +2,18 @@
 
 #pragma once
 #include "GameplayCueInterface.h"
+#include "AbilitySystemInterface.h"
 #include "GameplayCueView.h"
 #include "AbilitySystemTestPawn.generated.h"
 
 UCLASS(Blueprintable, BlueprintType)
-class GAMEPLAYABILITIES_API AAbilitySystemTestPawn : public ADefaultPawn, public IGameplayCueInterface
+class GAMEPLAYABILITIES_API AAbilitySystemTestPawn : public ADefaultPawn, public IGameplayCueInterface, public IAbilitySystemInterface
 {
 	GENERATED_UCLASS_BODY()
 
 	virtual void PostInitializeComponents() override;
+
+	virtual UAbilitySystemComponent* GetAbilitySystemComponent() const override;
 	
 	/** DefaultPawn collision component */
 	UPROPERTY(Category = AbilitySystem, VisibleAnywhere, BlueprintReadOnly)
