Change predicted, instant gameplay effects so they have infinite duration on clients. This allows us to correctly handle prediction of instant effects.

diff --git a/Engine/Source/Runtime/GameplayAbilities/Private/AbilitySystemComponent.cpp b/Engine/Source/Runtime/GameplayAbilities/Private/AbilitySystemComponent.cpp
--- a/Engine/Source/Runtime/GameplayAbilities/Private/AbilitySystemComponent.cpp
+++ b/Engine/Source/Runtime/GameplayAbilities/Private/AbilitySystemComponent.cpp
@@ -304,7 +304,10 @@ FActiveGameplayEffectHandle UAbilitySystemComponent::ApplyGameplayEffectSpecToSe
 {
 	// Temp, only non instant, non periodic GEs can be predictive
 	// Effects with other effects may be a mix so go with non-predictive
-	check((BaseQualifier.CurrentPredictionKey() == 0) || (Spec.GetDuration() != UGameplayEffect::INSTANT_APPLICATION && Spec.GetPeriod() == UGameplayEffect::NO_PERIOD));
+	check((BaseQualifier.CurrentPredictionKey() == 0) || (Spec.GetPeriod() == UGameplayEffect::NO_PERIOD));
+
+	// Clients should treat predicted instant effects as if they have infinite duration. The effects will be cleaned up later.
+	bool bTreatAsInfiniteDuration = GetOwnerRole() != ROLE_Authority && BaseQualifier.CurrentPredictionKey() != 0 && Spec.GetDuration() == UGameplayEffect::INSTANT_APPLICATION;
 
 	// check if the effect being applied actually succeeds
 	float ChanceToApply = Spec.GetChanceToApplyToTarget();
@@ -319,10 +322,8 @@ FActiveGameplayEffectHandle UAbilitySystemComponent::ApplyGameplayEffectSpecToSe
 
 	FGameplayEffectSpec* OurCopyOfSpec = NULL;
 	TSharedPtr<FGameplayEffectSpec> StackSpec;
+	float Duration = bTreatAsInfiniteDuration ? UGameplayEffect::INFINITE_DURATION : Spec.GetDuration();
 	{
-	
-		float Duration = Spec.GetDuration();
-
 		if (Duration != UGameplayEffect::INSTANT_APPLICATION)
 		{
 			// recalculating stacking needs to come before creating the new effect
@@ -351,6 +352,22 @@ FActiveGameplayEffectHandle UAbilitySystemComponent::ApplyGameplayEffectSpecToSe
 		// In theory, we don't modify 2nd order modifiers after they are 'attached'
 		// Long complex chains can be created but we never say 'Modify a GE that is modding another GE'
 		OurCopyOfSpec->MakeUnique();
+
+		// if necessary add a modifier to OurCopyOfSpec to force it to have an infinite duration
+		if (bTreatAsInfiniteDuration)
+		{
+			FGameplayModifierInfo ModInfo;
+			ModInfo.ModifierOp = EGameplayModOp::Override;
+			ModInfo.Magnitude.SetValue(UGameplayEffect::INFINITE_DURATION);
+			ModInfo.EffectType = EGameplayModEffect::Duration;
+			ModInfo.ModifierType = EGameplayMod::ActiveGE;
+			TSharedPtr<FGameplayEffectLevelSpec> Level(new FGameplayEffectLevelSpec(0.f, OurCopyOfSpec->Def->LevelInfo, NULL));
+			FModifierSpec Mod(ModInfo, Level, NULL);
+			FModifierQualifier Qualifier;
+			Qualifier.Type(EGameplayMod::ActiveGE);
+
+			OurCopyOfSpec->ApplyModifier(Mod, Qualifier, true);
+		}
 	}
 
 	// Now that we have our own copy, apply our GEs that modify IncomingGEs
@@ -361,7 +378,7 @@ FActiveGameplayEffectHandle UAbilitySystemComponent::ApplyGameplayEffectSpecToSe
 	}	
 	
 	// Now that we have the final version of this effect, actually apply it if its going to be hanging around
-	if (Spec.GetDuration() != UGameplayEffect::INSTANT_APPLICATION )
+	if (Duration != UGameplayEffect::INSTANT_APPLICATION )
 	{
 		if (Spec.GetPeriod() == UGameplayEffect::NO_PERIOD)
 		{
@@ -384,7 +401,7 @@ FActiveGameplayEffectHandle UAbilitySystemComponent::ApplyGameplayEffectSpecToSe
 	// Execute the GE at least once (if instant, this will execute once and be done. If persistent, it was added to ActiveGameplayEffects above)
 	
 	// Execute if this is an instant application effect
-	if (Spec.GetDuration() == UGameplayEffect::INSTANT_APPLICATION)
+	if (Duration == UGameplayEffect::INSTANT_APPLICATION)
 	{
 		ExecuteGameplayEffect(*OurCopyOfSpec, FModifierQualifier(BaseQualifier).IgnoreHandle(MyHandle));
 	}
diff --git a/Engine/Source/Runtime/GameplayAbilities/Private/GameplayEffect.cpp b/Engine/Source/Runtime/GameplayAbilities/Private/GameplayEffect.cpp
--- a/Engine/Source/Runtime/GameplayAbilities/Private/GameplayEffect.cpp
+++ b/Engine/Source/Runtime/GameplayAbilities/Private/GameplayEffect.cpp
@@ -190,11 +190,10 @@ int32 FGameplayEffectSpec::ApplyModifiersFrom(const FGameplayEffectSpec &InSpec,
 	ABILITY_LOG_SCOPE(TEXT("FGameplayEffectSpec::ApplyModifiersFrom %s. InSpec: %s"), *this->ToSimpleString(), *InSpec.ToSimpleString());
 
 	int32 NumApplied = 0;
-	bool ShouldSnapshot = InSpec.ShouldApplyAsSnapshot(QualifierContext);
 
 	if (!InSpec.Def || !InSpec.Def->AreGameplayEffectTagRequirementsSatisfied(Def))
 	{
-		// InSpec doesn't match this gameplay effect but if InSpec provides immunity we also need to check the modifiers because they can be canceled independent of the gameplay effect
+		// InSpec doesn't match this GameplayEffectSpec but if InSpec provides immunity we also need to check the modifiers because they can be canceled independent of the GameplayEffectSpec
 		if ((int32)InSpec.Def->AppliesImmunityTo == (int32)QualifierContext.Type())
 		{
 			for (int ii = 0; ii < Modifiers.Num(); ++ii)
@@ -221,63 +220,7 @@ int32 FGameplayEffectSpec::ApplyModifiersFrom(const FGameplayEffectSpec &InSpec,
 
 	for (const FModifierSpec &InMod : InSpec.Modifiers)
 	{
-		if (!InMod.CanModifyInContext(QualifierContext))
-		{
-			continue;
-		}
-
-		switch (InMod.Info.EffectType)
-		{
-			case EGameplayModEffect::Magnitude:
-			{
-				for (FModifierSpec &MyMod : Modifiers)
-				{
-					if (InMod.CanModifyModifier(MyMod, QualifierContext))
-					{
-						InMod.ApplyModTo(MyMod, ShouldSnapshot);
-						NumApplied++;
-					}
-				}
-				break;
-			}
-
-			// Fixme: Duration mods aren't checking that attributes match - do we care?
-			// eg - "I mod duration of all GEs that modify Health" would need to check to see if this mod modifies health before doing the stuff below.
-			// Or can we get by with just tags?
-
-			case EGameplayModEffect::Duration:
-			{
-				// don't modify infinite duration unless we're overriding it
-				if (GetDuration() > 0.f || InMod.Info.ModifierOp == EGameplayModOp::Override)
-				{
-					Duration.Get()->ApplyMod(InMod.Info.ModifierOp, InMod.Aggregator, ShouldSnapshot);
-					NumApplied++;
-				}
-				
-				break;
-			}
-
-			case EGameplayModEffect::ChanceApplyTarget:
-			{
-				ChanceToApplyToTarget.Get()->ApplyMod(InMod.Info.ModifierOp, InMod.Aggregator, ShouldSnapshot);
-				NumApplied++;
-				break;
-			}
-
-			case EGameplayModEffect::ChanceExecuteEffect:
-			{
-				ChanceToExecuteOnGameplayEffect.Get()->ApplyMod(InMod.Info.ModifierOp, InMod.Aggregator, ShouldSnapshot);
-				NumApplied++;
-				break;
-			}
-
-			case EGameplayModEffect::LinkedGameplayEffect:
-			{
-				TargetEffectSpecs.Add(InMod.TargetEffectSpec.ToSharedRef());
-				NumApplied++;
-				break;
-			}
-		}
+		NumApplied += ApplyModifier(InMod, QualifierContext, InSpec.ShouldApplyAsSnapshot(QualifierContext));
 	}
 
 	return NumApplied;
@@ -303,6 +246,70 @@ int32 FGameplayEffectSpec::ExecuteModifiersFrom(const FGameplayEffectSpec &InSpe
 	return NumExecuted;
 }
 
+int32 FGameplayEffectSpec::ApplyModifier(const FModifierSpec &InMod, const FModifierQualifier &QualifierContext, bool bApplyAsSnapshot)
+{
+	int32 NumApplied = 0;
+	if (!InMod.CanModifyInContext(QualifierContext))
+	{
+		return 0;
+	}
+
+	switch (InMod.Info.EffectType)
+	{
+		case EGameplayModEffect::Magnitude:
+		{
+			for (FModifierSpec &MyMod : Modifiers)
+			{
+				if (InMod.CanModifyModifier(MyMod, QualifierContext))
+				{
+					InMod.ApplyModTo(MyMod, bApplyAsSnapshot);
+					NumApplied = 1;
+				}
+			}
+			break;
+		}
+
+		// Fixme: Duration mods aren't checking that attributes match - do we care?
+		// eg - "I mod duration of all GEs that modify Health" would need to check to see if this mod modifies health before doing the stuff below.
+		// Or can we get by with just tags?
+
+		case EGameplayModEffect::Duration:
+		{
+			// don't modify infinite duration unless we're overriding it
+			if (GetDuration() > 0.f || InMod.Info.ModifierOp == EGameplayModOp::Override)
+			{
+				Duration.Get()->ApplyMod(InMod.Info.ModifierOp, InMod.Aggregator, bApplyAsSnapshot);
+				NumApplied = 1;
+			}
+
+			break;
+		}
+
+		case EGameplayModEffect::ChanceApplyTarget:
+		{
+			ChanceToApplyToTarget.Get()->ApplyMod(InMod.Info.ModifierOp, InMod.Aggregator, bApplyAsSnapshot);
+			NumApplied = 1;
+			break;
+		}
+
+		case EGameplayModEffect::ChanceExecuteEffect:
+		{
+			ChanceToExecuteOnGameplayEffect.Get()->ApplyMod(InMod.Info.ModifierOp, InMod.Aggregator, bApplyAsSnapshot);
+			NumApplied = 1;
+			break;
+		}
+
+		case EGameplayModEffect::LinkedGameplayEffect:
+		{
+			TargetEffectSpecs.Add(InMod.TargetEffectSpec.ToSharedRef());
+			NumApplied = 1;
+			break;
+		}
+	}
+
+	return NumApplied;
+}
+
 float FGameplayEffectSpec::GetDuration() const
 {
 	return Duration.Get()->Evaluate().Magnitude;
@@ -387,7 +394,7 @@ EGameplayEffectStackingPolicy::Type FGameplayEffectSpec::GetStackingType() const
 
 bool FModifierSpec::CanModifyInContext(const FModifierQualifier &QualifierContext) const
 {
-	// Can only modify if are valid within this Qualifier Context
+	// Can only modify if this spec is valid within this Qualifier Context
 	// (E.g, if I am an OutgoingGE mod, I cannot modify during a IncomingGE context)
 	if (Info.ModifierType != QualifierContext.Type())
 	{
diff --git a/Engine/Source/Runtime/GameplayAbilities/Public/GameplayEffect.h b/Engine/Source/Runtime/GameplayAbilities/Public/GameplayEffect.h
--- a/Engine/Source/Runtime/GameplayAbilities/Public/GameplayEffect.h
+++ b/Engine/Source/Runtime/GameplayAbilities/Public/GameplayEffect.h
@@ -940,6 +940,7 @@ struct FModifierSpec
 
 	TSharedPtr<FGameplayEffectSpec> TargetEffectSpec;
 
+	// returns true if this GameplayEffectSpec can modify things in QualifierContext, returns false otherwise.
 	bool CanModifyInContext(const FModifierQualifier &QualifierContext) const;
 	// returns true if this GameplayEffect can modify Other, false otherwise
 	bool CanModifyModifier(FModifierSpec &Other, const FModifierQualifier &QualifierContext) const;
@@ -1031,12 +1032,17 @@ struct FGameplayEffectSpec
 
 	void InitModifiers(const FGlobalCurveDataOverride *CurveData, AActor *Owner, float Level);
 
-	// returns the number of modifiers applied to InSpec by the current GameplayEffect Spec
+	// returns the number of modifiers applied to the current GameplayEffectSpec by InSpec
 	// returns -1 if the current GameplayEffectSpec prevents InSpec from being applied
 	int32 ApplyModifiersFrom(const FGameplayEffectSpec &InSpec, const FModifierQualifier &QualifierContext);
 
 	int32 ExecuteModifiersFrom(const FGameplayEffectSpec &InSpec, const FModifierQualifier &QualifierContext);
 
+	// returns 1 if the modifier was applied, 0 otherwise
+	int32 ApplyModifier(const FModifierSpec &InMod, const FModifierQualifier &QualifierContext, bool bApplyAsSnapshot);
+
+	// returns true if modifiers applied to this GameplayEffectSpec should be a snapshot of the applied modifiers
+	// returns false if modifiers applied to this GameplayEffectSpec should link to the applied modifiers
 	bool ShouldApplyAsSnapshot(const FModifierQualifier &QualifierContext) const;
 
 	FString ToSimpleString() const
