#orion - Added IsActive() to UGameplayAbility() and made InstancedPerActor abilities track when they are active

diff --git a/Engine/Source/Runtime/GameplayAbilities/Private/Abilities/GameplayAbility.cpp b/Engine/Source/Runtime/GameplayAbilities/Private/Abilities/GameplayAbility.cpp
--- a/Engine/Source/Runtime/GameplayAbilities/Private/Abilities/GameplayAbility.cpp
+++ b/Engine/Source/Runtime/GameplayAbilities/Private/Abilities/GameplayAbility.cpp
@@ -141,6 +141,18 @@ void UGameplayAbility::PostNetInit()
 	}
 }
 
+bool UGameplayAbility::IsActive() const
+{
+	// Only Instanced-Per-Actor abilities persist between activations
+	if (GetInstancingPolicy() == EGameplayAbilityInstancingPolicy::InstancedPerActor)
+	{
+		return bIsActive;
+	}
+
+	// Non-instanced and Instanced-Per-Execution abilities are by definition active unless they are pending kill
+	return !IsPendingKill();
+}
+
 bool UGameplayAbility::IsSupportedForNetworking() const
 {
 	/**
@@ -339,8 +351,15 @@ bool UGameplayAbility::TryActivateAbility(const FGameplayAbilityActorInfo* Actor
 	return true;	
 }
 
+void UGameplayAbility::CancelAbility(const FGameplayAbilityActorInfo* ActorInfo, const FGameplayAbilityActivationInfo ActivationInfo)
+{
+	EndAbility(ActorInfo);
+}
+
 void UGameplayAbility::EndAbility(const FGameplayAbilityActorInfo* ActorInfo)
 {
+	bIsActive = false;
+
 	// Tell all our tasks that we are finished and they should cleanup
 	for (TWeakObjectPtr<UAbilityTask> Task : ActiveTasks)
 	{
@@ -383,6 +402,7 @@ void UGameplayAbility::PreActivate(const FGameplayAbilityActorInfo* ActorInfo, c
 
 	Comp->CancelAbilitiesWithTags(CancelAbilitiesWithTag, ActorInfo, ActivationInfo, this);
 
+	bIsActive = true;
 	Comp->NotifyAbilityActivated(this);
 
 	// Become the 'Targeting Ability'
diff --git a/Engine/Source/Runtime/GameplayAbilities/Public/Abilities/GameplayAbility.h b/Engine/Source/Runtime/GameplayAbilities/Public/Abilities/GameplayAbility.h
--- a/Engine/Source/Runtime/GameplayAbilities/Public/Abilities/GameplayAbility.h
+++ b/Engine/Source/Runtime/GameplayAbilities/Public/Abilities/GameplayAbility.h
@@ -169,6 +169,9 @@ public:
 
 	void PostNetInit();
 
+	/** Returns true if the ability is currently active */
+	bool IsActive() const;
+
 	/** This ability has these tags */
 	UPROPERTY(EditDefaultsOnly, Category = Tags)
 	FGameplayTagContainer AbilityTags;
@@ -266,7 +269,7 @@ protected:
 	// --------------------------------------
 
 	/** Destroys intanced-per-execution abilities. Instance-per-actor abilities should 'reset'. Non instance abilities - what can we do? */
-	virtual void CancelAbility(const FGameplayAbilityActorInfo* ActorInfo, const FGameplayAbilityActivationInfo ActivationInfo) { }
+	virtual void CancelAbility(const FGameplayAbilityActorInfo* ActorInfo, const FGameplayAbilityActivationInfo ActivationInfo);
 
 	/** Destroys intanced-per-execution abilities. Instance-per-actor abilities should 'reset'. Non instance abilities - what can we do? */
 	UFUNCTION(BlueprintCallable, Category = Ability)
@@ -416,4 +419,7 @@ private:
 	 *	 (E.g, child classes may want to cast to FMyGameAbilityActorInfo)
 	 */
 	mutable const FGameplayAbilityActorInfo* CurrentActorInfo;
+
+	/** True if the ability is currently active. For instance per owner abilities */
+	bool bIsActive;
 };
