copy ticking task list to avoid issues when tasks are destroyed during a tick

diff --git a/Engine/Source/Runtime/GameplayAbilities/Private/AbilitySystemComponent_Abilities.cpp b/Engine/Source/Runtime/GameplayAbilities/Private/AbilitySystemComponent_Abilities.cpp
--- a/Engine/Source/Runtime/GameplayAbilities/Private/AbilitySystemComponent_Abilities.cpp
+++ b/Engine/Source/Runtime/GameplayAbilities/Private/AbilitySystemComponent_Abilities.cpp
@@ -43,15 +43,39 @@ void UAbilitySystemComponent::InitializeComponent()
 
 void UAbilitySystemComponent::TickComponent(float DeltaTime, enum ELevelTick TickType, FActorComponentTickFunction *ThisTickFunction)
 {
+	SCOPE_CYCLE_COUNTER(STAT_TickAbilityTasks);
+
 	Super::TickComponent(DeltaTime, TickType, ThisTickFunction);
 
-	for (TWeakObjectPtr<UAbilityTask> TickingTask : TickingTasks)
+	// Because we have no control over what a task may do when it ticks, we must be careful.
+	// Ticking a task may kill the task right here. It could also potentially kill another task
+	// which was waiting on the original task to do something. Since when a tasks is killed, it removes
+	// itself from the TickingTask list, we will make a copy of the tasks we want to service before ticking any
+
+	int32 NumTickingTasks = TickingTasks.Num();
+	switch(NumTickingTasks)
 	{
-		if (TickingTask.IsValid())
-		{
-			TickingTask->TickTask(DeltaTime);
-		}
-	}
+		case 0:
+			break;
+		case 1:
+			if (TickingTasks[0].IsValid())
+			{
+				TickingTasks[0]->TickTask(DeltaTime);
+			}
+			break;
+		default:
+			{
+				TArray<TWeakObjectPtr<UAbilityTask> >	LocalTickingTasks = TickingTasks;
+				for (TWeakObjectPtr<UAbilityTask>& TickingTask : LocalTickingTasks)
+				{
+					if (TickingTask.IsValid())
+					{
+						TickingTask->TickTask(DeltaTime);
+					}
+				}
+			}
+		break;
+	};
 }
 
 void UAbilitySystemComponent::InitAbilityActorInfo(AActor* AvatarActor)
diff --git a/Engine/Source/Runtime/GameplayAbilities/Private/AbilitySystemStats.cpp b/Engine/Source/Runtime/GameplayAbilities/Private/AbilitySystemStats.cpp
--- a/Engine/Source/Runtime/GameplayAbilities/Private/AbilitySystemStats.cpp
+++ b/Engine/Source/Runtime/GameplayAbilities/Private/AbilitySystemStats.cpp
@@ -14,5 +14,6 @@ DEFINE_STAT(STAT_CreateNewActiveGameplayEffect);
 DEFINE_STAT(STAT_GetGameplayCueFunction);
 DEFINE_STAT(STAT_GetOutgoingSpec);
 DEFINE_STAT(STAT_InitAttributeSetDefaults);
+DEFINE_STAT(STAT_TickAbilityTasks);
 
 
diff --git a/Engine/Source/Runtime/GameplayAbilities/Public/AbilitySystemStats.h b/Engine/Source/Runtime/GameplayAbilities/Public/AbilitySystemStats.h
--- a/Engine/Source/Runtime/GameplayAbilities/Public/AbilitySystemStats.h
+++ b/Engine/Source/Runtime/GameplayAbilities/Public/AbilitySystemStats.h
@@ -17,4 +17,5 @@ DECLARE_CYCLE_STAT_EXTERN(TEXT("RemoveActiveGameplayEffect"), STAT_RemoveActiveG
 DECLARE_CYCLE_STAT_EXTERN(TEXT("CreateNewActiveGameplayEffect"), STAT_CreateNewActiveGameplayEffect, STATGROUP_AbilitySystem, );
 DECLARE_CYCLE_STAT_EXTERN(TEXT("GetGameplayCueFunction"), STAT_GetGameplayCueFunction, STATGROUP_AbilitySystem, );
 DECLARE_CYCLE_STAT_EXTERN(TEXT("GetOutgoingSpec"), STAT_GetOutgoingSpec, STATGROUP_AbilitySystem, );
-DECLARE_CYCLE_STAT_EXTERN(TEXT("InitAttributeSetDefaults"), STAT_InitAttributeSetDefaults, STATGROUP_AbilitySystem, );
\ No newline at end of file
+DECLARE_CYCLE_STAT_EXTERN(TEXT("InitAttributeSetDefaults"), STAT_InitAttributeSetDefaults, STATGROUP_AbilitySystem, );
+DECLARE_CYCLE_STAT_EXTERN(TEXT("TickAbilityTasks"), STAT_TickAbilityTasks, STATGROUP_AbilitySystem, );
\ No newline at end of file
