Fix non-unity for GameplayAbilities system

diff --git a/Engine/Source/Runtime/GameplayAbilities/Private/GameplayEffect.cpp b/Engine/Source/Runtime/GameplayAbilities/Private/GameplayEffect.cpp
--- a/Engine/Source/Runtime/GameplayAbilities/Private/GameplayEffect.cpp
+++ b/Engine/Source/Runtime/GameplayAbilities/Private/GameplayEffect.cpp
@@ -6,6 +6,7 @@
 #include "GameplayEffectStackingExtension.h"
 #include "GameplayTagsModule.h"
 #include "AbilitySystemComponent.h"
+#include "AbilitySystemInterface.h"
 
 const float UGameplayEffect::INFINITE_DURATION = -1.f;
 const float UGameplayEffect::INSTANT_APPLICATION = 0.f;
