Fix case where noninstanced abilities could be marked pending kill

diff --git a/Engine/Source/Runtime/GameplayAbilities/Private/AbilitySystemComponent_Abilities.cpp b/Engine/Source/Runtime/GameplayAbilities/Private/AbilitySystemComponent_Abilities.cpp
--- a/Engine/Source/Runtime/GameplayAbilities/Private/AbilitySystemComponent_Abilities.cpp
+++ b/Engine/Source/Runtime/GameplayAbilities/Private/AbilitySystemComponent_Abilities.cpp
@@ -163,6 +163,8 @@ void UAbilitySystemComponent::NotifyAbilityEnded(UGameplayAbility* Ability)
 	/** If this is instanced per execution, mark pending kill and remove it from our instanced lists if we are the authority */
 	if (Ability->GetInstancingPolicy() == EGameplayAbilityInstancingPolicy::InstancedPerExecution)
 	{
+		check(Ability->HasAnyFlags(RF_ClassDefaultObject) == false);	// Should never be calling this on a CDO for an instanced ability!
+
 		if (Ability->GetReplicationPolicy() != EGameplayAbilityReplicationPolicy::ReplicateNone)
 		{
 			if (GetOwnerRole() == ROLE_Authority)
@@ -204,32 +206,44 @@ void UAbilitySystemComponent::CancelAbilitiesWithTags(const FGameplayTagContaine
 
 	struct local
 	{
-		static void CancelAbilitiesWithTags(const FGameplayTagContainer InTags, TArray<FGameplayAbilityInputIDPair> Abilities, const FGameplayAbilityActorInfo* InActorInfo, const FGameplayAbilityActivationInfo InActivationInfo, UGameplayAbility* InIgnore)
+		static void CancelAbilitiesWithTags(TArray<UGameplayAbility*>& AbilitiesToCancel, const FGameplayTagContainer InTags, TArray<FGameplayAbilityInputIDPair> Abilities, const FGameplayAbilityActorInfo* InActorInfo, const FGameplayAbilityActivationInfo InActivationInfo, UGameplayAbility* InIgnore)
 		{
 			for (int32 idx=0; idx < Abilities.Num(); ++idx)
 			{
 				UGameplayAbility *Ability = Abilities[idx].Ability;
 				if (Ability && (Ability != InIgnore) && Ability->AbilityTags.MatchesAny(InTags, false))
 				{
-					Ability->CancelAbility(InActorInfo, InActivationInfo);
-
-					/*
-					if (!Ability->HasAnyFlags(RF_ClassDefaultObject) && Ability->GetInstancingPolicy() == EGameplayAbilityInstancingPolicy::InstancedPerExecution)
+					if (!Ability->HasAnyFlags(RF_ClassDefaultObject) || Ability->GetInstancingPolicy() == EGameplayAbilityInstancingPolicy::NonInstanced)
 					{
-						Ability->EndAbility();
-						Abilities.RemoveAtSwap(idx);
-						idx--;
-					}
-					*/
+						AbilitiesToCancel.Add(Ability);
+					}					
+				}
+			}
+		}
+
+		static void CancelAbilitiesWithTags(TArray<UGameplayAbility*>& AbilitiesToCancel, const FGameplayTagContainer InTags, TArray<UGameplayAbility*> Abilities, const FGameplayAbilityActorInfo* InActorInfo, const FGameplayAbilityActivationInfo InActivationInfo, UGameplayAbility* InIgnore)
+		{
+			for (int32 idx=0; idx < Abilities.Num(); ++idx)
+			{
+				UGameplayAbility *Ability = Abilities[idx];
+				if (Ability && (Ability != InIgnore) && Ability->AbilityTags.MatchesAny(InTags, false))
+				{					
+					AbilitiesToCancel.Add(Ability);
 				}
 			}
 		}
 	};
 
-	//local::CancelAbilitiesWithTags(Tags, ReplicatedInstancedAbilities, ActorInfo, Ignore);
-	//local::CancelAbilitiesWithTags(Tags, NonReplicatedInstancedAbilities, ActorInfo, Ignore);
+	TArray<UGameplayAbility*>	AbilitiesToCancel;
+	
+	local::CancelAbilitiesWithTags(AbilitiesToCancel, Tags, ReplicatedInstancedAbilities, ActorInfo, ActivationInfo, Ignore);
+	local::CancelAbilitiesWithTags(AbilitiesToCancel, Tags, NonReplicatedInstancedAbilities, ActorInfo, ActivationInfo, Ignore);
+	local::CancelAbilitiesWithTags(AbilitiesToCancel, Tags, ActivatableAbilities, ActorInfo, ActivationInfo, Ignore);
 
-	local::CancelAbilitiesWithTags(Tags, ActivatableAbilities, ActorInfo, ActivationInfo, Ignore);
+	for (UGameplayAbility* Ability : AbilitiesToCancel)
+	{
+		Ability->CancelAbility(ActorInfo, ActivationInfo);
+	}
 }
 
 #if !(UE_BUILD_SHIPPING || UE_BUILD_TEST)
