fix typo

diff --git a/Engine/Source/Runtime/GameplayAbilities/Private/GameplayEffect.cpp b/Engine/Source/Runtime/GameplayAbilities/Private/GameplayEffect.cpp
--- a/Engine/Source/Runtime/GameplayAbilities/Private/GameplayEffect.cpp
+++ b/Engine/Source/Runtime/GameplayAbilities/Private/GameplayEffect.cpp
@@ -463,7 +463,7 @@ void FGameplayEffectInstigatorContext::AddInstigator(class AActor *InInstigator)
 
 	// Cache off his AbilitySystemComponent.
 	IAbilitySystemInterface* AbilitySystmeInterface = InterfaceCast<IAbilitySystemInterface>(Instigator);
-	if (!AbilitySystmeInterface)
+	if (AbilitySystmeInterface)
 	{
 		InstigatorAbilitySystemComponent = AbilitySystmeInterface->GetAbilitySystemComponent();
 	}
