Select Multiple Targets task created, targeting data changed and new targeting functionality added.

diff --git a/Engine/Source/Runtime/GameplayAbilities/Private/Abilities/GameplayAbility.cpp b/Engine/Source/Runtime/GameplayAbilities/Private/Abilities/GameplayAbility.cpp
--- a/Engine/Source/Runtime/GameplayAbilities/Private/Abilities/GameplayAbility.cpp
+++ b/Engine/Source/Runtime/GameplayAbilities/Private/Abilities/GameplayAbility.cpp
@@ -750,6 +750,14 @@ FGameplayAbilityTargetDataHandle FGameplayAbilityTargetingLocationInfo::MakeTarg
 	return FGameplayAbilityTargetDataHandle(ReturnData);
 }
 
+FGameplayAbilityTargetDataHandle FGameplayAbilityTargetingLocationInfo::MakeTargetDataHandleFromActors(TArray<AActor*> TargetActors) const
+{
+	/** Note: This is cleaned up by the FGameplayAbilityTargetDataHandle (via an internal TSharedPtr) */
+	FGameplayAbilityTargetData_ActorArray* ReturnData = new FGameplayAbilityTargetData_ActorArray();
+	ReturnData->TargetActorArray = TargetActors;
+	ReturnData->SourceLocation = *this;
+	return FGameplayAbilityTargetDataHandle(ReturnData);
+}
 
 bool FGameplayAbilityTargetDataHandle::NetSerialize(FArchive& Ar, class UPackageMap* Map, bool& bOutSuccess)
 {
@@ -807,6 +815,49 @@ bool FGameplayAbilityTargetDataHandle::NetSerialize(FArchive& Ar, class UPackage
 	return true;
 }
 
+bool FGameplayAbilityTargetingLocationInfo::NetSerialize(FArchive& Ar, class UPackageMap* Map, bool& bOutSuccess)
+{
+	Ar << LocationType;
+
+	switch (LocationType)
+	{
+	case EGameplayAbilityTargetingLocationType::ActorTransform:
+		Ar << SourceActor;
+		break;
+	case EGameplayAbilityTargetingLocationType::SocketTransform:
+		Ar << SourceComponent;
+		Ar << SourceSocketName;
+		break;
+	case EGameplayAbilityTargetingLocationType::LiteralTransform:
+		Ar << LiteralTransform;
+		break;
+	default:
+		check(false);		//This case should not happen
+		break;
+	}
+
+	bOutSuccess = true;
+	return true;
+}
+
+bool FGameplayAbilityTargetData_LocationInfo::NetSerialize(FArchive& Ar, class UPackageMap* Map, bool& bOutSuccess)
+{
+	SourceLocation.NetSerialize(Ar, Map, bOutSuccess);
+	TargetLocation.NetSerialize(Ar, Map, bOutSuccess);
+
+	bOutSuccess = true;
+	return true;
+}
+
+bool FGameplayAbilityTargetData_ActorArray::NetSerialize(FArchive& Ar, class UPackageMap* Map, bool& bOutSuccess)
+{
+	SourceLocation.NetSerialize(Ar, Map, bOutSuccess);
+	Ar << TargetActorArray;
+
+	bOutSuccess = true;
+	return true;
+}
+
 bool FGameplayAbilityTargetData_Mesh::NetSerialize(FArchive& Ar, class UPackageMap* Map, bool& bOutSuccess)
 {
 	//SourceActor can be used as a backup if the component isn't found.
diff --git a/Engine/Source/Runtime/GameplayAbilities/Private/Abilities/GameplayAbilityTargetActor_SingleLineTrace.cpp b/Engine/Source/Runtime/GameplayAbilities/Private/Abilities/GameplayAbilityTargetActor_SingleLineTrace.cpp
--- a/Engine/Source/Runtime/GameplayAbilities/Private/Abilities/GameplayAbilityTargetActor_SingleLineTrace.cpp
+++ b/Engine/Source/Runtime/GameplayAbilities/Private/Abilities/GameplayAbilityTargetActor_SingleLineTrace.cpp
@@ -67,12 +67,12 @@ FHitResult AGameplayAbilityTargetActor_SingleLineTrace::PerformTrace(AActor* InS
 	}
 
 	FVector TraceStart = InSourceActor->GetActorLocation();
-	FVector TraceEnd = TraceStart + (AimDirection * 3000.f);
+	FVector TraceEnd = TraceStart + (AimDirection * MaxRange);
 
 	//If we're using a socket, adjust the starting location and aim direction after the end position has been found. This way we can still aim with the camera, then fire accurately from the socket.
 	if (MyAbility)		//Server and launching client only
 	{
-		TraceStart = StartLocation.GetTargetingTransform(MyAbility->GetCurrentActorInfo()).GetLocation();
+		TraceStart = StartLocation.GetTargetingTransform().GetLocation();
 	}
 	AimDirection = (TraceEnd - TraceStart).SafeNormal();
 
diff --git a/Engine/Source/Runtime/GameplayAbilities/Private/AbilitySystemBlueprintLibrary.cpp b/Engine/Source/Runtime/GameplayAbilities/Private/AbilitySystemBlueprintLibrary.cpp
--- a/Engine/Source/Runtime/GameplayAbilities/Private/AbilitySystemBlueprintLibrary.cpp
+++ b/Engine/Source/Runtime/GameplayAbilities/Private/AbilitySystemBlueprintLibrary.cpp
@@ -28,10 +28,26 @@ void UAbilitySystemBlueprintLibrary::ApplyGameplayEffectToTargetData(FGameplayAb
 	}
 }
 
+FGameplayAbilityTargetDataHandle UAbilitySystemBlueprintLibrary::AbilityTargetDataFromLocations(FGameplayAbilityTargetingLocationInfo SourceLocation, FGameplayAbilityTargetingLocationInfo TargetLocation)
+{
+	// Construct TargetData
+	FGameplayAbilityTargetData_LocationInfo*	NewData = new FGameplayAbilityTargetData_LocationInfo();
+	NewData->SourceLocation = SourceLocation;
+	NewData->TargetLocation = TargetLocation;
+
+	// Give it a handle and return
+	FGameplayAbilityTargetDataHandle	Handle;
+	Handle.Data = TSharedPtr<FGameplayAbilityTargetData_LocationInfo>(NewData);
+	return Handle;
+}
+
 FGameplayAbilityTargetDataHandle UAbilitySystemBlueprintLibrary::AbilityTargetDataHandleFromAbilityTargetDataMesh(FGameplayAbilityTargetData_Mesh Data)
 {
+	// Construct TargetData
 	FGameplayAbilityTargetData_Mesh*	NewData = new FGameplayAbilityTargetData_Mesh(Data);
 	FGameplayAbilityTargetDataHandle	Handle;
+
+	// Give it a handle and return
 	Handle.Data = TSharedPtr<FGameplayAbilityTargetData_Mesh>(NewData);
 	return Handle;
 }
@@ -48,6 +64,16 @@ FGameplayAbilityTargetDataHandle UAbilitySystemBlueprintLibrary::AbilityTargetDa
 	return Handle;
 }
 
+TArray<AActor*> UAbilitySystemBlueprintLibrary::GetActorsFromTargetData(FGameplayAbilityTargetDataHandle TargetData)
+{
+	FGameplayAbilityTargetData* Data = TargetData.Data.Get();
+	if (Data)
+	{
+		return Data->GetActors();
+	}
+	return TArray<AActor*>();
+}
+
 bool UAbilitySystemBlueprintLibrary::TargetDataHasHitResult(FGameplayAbilityTargetDataHandle TargetData)
 {
 	FGameplayAbilityTargetData* Data = TargetData.Data.Get();
diff --git a/Engine/Source/Runtime/GameplayAbilities/Public/Abilities/GameplayAbilityTypes.h b/Engine/Source/Runtime/GameplayAbilities/Public/Abilities/GameplayAbilityTypes.h
--- a/Engine/Source/Runtime/GameplayAbilities/Public/Abilities/GameplayAbilityTypes.h
+++ b/Engine/Source/Runtime/GameplayAbilities/Public/Abilities/GameplayAbilityTypes.h
@@ -276,44 +276,52 @@ struct GAMEPLAYABILITIES_API FGameplayAbilityTargetingLocationInfo
 	GENERATED_USTRUCT_BODY()
 
 	FGameplayAbilityTargetingLocationInfo()
+	: LocationType(EGameplayAbilityTargetingLocationType::LiteralTransform)
 	{
 	};
 
+	void operator=(const FGameplayAbilityTargetingLocationInfo& Other)
+	{
+		LocationType = Other.LocationType;
+		LiteralTransform = Other.LiteralTransform;
+		SourceActor = Other.SourceActor;
+		SourceComponent = Other.SourceComponent;
+		SourceSocketName = Other.SourceSocketName;
+	}
+
+public:
 	//Not sure if we want to distinguish between where we actually aim from (should probably always be the user's camera or actor origin) and where we trace from.
-	FTransform GetTargetingTransform(const FGameplayAbilityActorInfo* ActorInfo) const
+	FTransform GetTargetingTransform() const		//This is split out so we can
 	{
 		//Return or calculate based on LocationType.
 		switch (LocationType)
 		{
 		case EGameplayAbilityTargetingLocationType::ActorTransform:
-			check(ActorInfo);
-			if (ActorInfo->Actor.IsValid())
+			if (SourceActor)
 			{
-				return ActorInfo->Actor.Get()->GetTransform();
+				return SourceActor->GetTransform();
 			}
-			return LiteralTransform;		//Fallback
+			break;
 		case EGameplayAbilityTargetingLocationType::SocketTransform:
-			check(ActorInfo);
-			check(ActorInfo->AnimInstance.IsValid());
-			if (ActorInfo->AnimInstance.Get()->GetOwningComponent())
-			{
-				return ActorInfo->AnimInstance.Get()->GetOwningComponent()->GetSocketTransform(SourceSocketName);		//Bad socket name will just return component transform anyway, so we're safe
-			}
-			if (ActorInfo->Actor.IsValid())
+			if (SourceComponent)
 			{
-				return ActorInfo->Actor->GetTransform();		//First fallback
+				return SourceComponent->GetSocketTransform(SourceSocketName);		//Bad socket name will just return component transform anyway, so we're safe
 			}
-			return LiteralTransform;		//Second fallback, but not really acceptable since it will almost surely be Identity
+			break;
 		case EGameplayAbilityTargetingLocationType::LiteralTransform:
 			return LiteralTransform;
 		default:
 			check(false);		//This case should not happen
-			return LiteralTransform;		//Fallback
+			break;
 		}
+		//Error
+		return FTransform::Identity;
 	}
 
 	FGameplayAbilityTargetDataHandle MakeTargetDataHandleFromHitResult(TWeakObjectPtr<UGameplayAbility> Ability, FHitResult HitResult) const;
 
+	FGameplayAbilityTargetDataHandle MakeTargetDataHandleFromActors(TArray<AActor*> TargetActors) const;
+
 	UPROPERTY(BlueprintReadWrite, meta = (ExposeOnSpawn = true), Category = Targeting)
 	TEnumAsByte<EGameplayAbilityTargetingLocationType::Type> LocationType;
 
@@ -321,15 +329,15 @@ struct GAMEPLAYABILITIES_API FGameplayAbilityTargetingLocationInfo
 	UPROPERTY(BlueprintReadWrite, meta = (ExposeOnSpawn = true), Category = Targeting)
 	FTransform LiteralTransform;
 
-	/** Actor who owns the named component. Actor's location is used as start point if component cannot be found. */
-	//UPROPERTY(BlueprintReadWrite, EditAnywhere, Category = Targeting)
-	//AActor* SourceActor;
+	/** A source actor is needed for Actor-based targeting, but not for Socket-based targeting. */
+	UPROPERTY(BlueprintReadWrite, meta = (ExposeOnSpawn = true), Category = Targeting)
+	AActor* SourceActor;
 
-	/** Local skeletal mesh component that holds the socket. */
-	//UPROPERTY(BlueprintReadWrite, EditAnywhere, Category = Targeting)
-	//USkeletalMeshComponent* SourceComponent;
+	/** Socket-based targeting requires a skeletal mesh component to check for the named socket. */
+	UPROPERTY(BlueprintReadWrite, meta = (ExposeOnSpawn = true), Category = Targeting)
+	USkeletalMeshComponent* SourceComponent;
 
-	/** If SourceActor and SourceComponent are valid, this is the name of the socket that will be used instead of the actor's location. */
+	/** If SourceComponent is valid, this is the name of the socket transform that will be used. If no Socket is provided, SourceComponent's transform will be used. */
 	UPROPERTY(BlueprintReadWrite, meta = (ExposeOnSpawn = true), Category = Targeting)
 	FName SourceSocketName;
 
@@ -340,7 +348,7 @@ struct GAMEPLAYABILITIES_API FGameplayAbilityTargetingLocationInfo
 		return TEXT("FGameplayAbilityTargetingLocationInfo");
 	}
 
-	//bool NetSerialize(FArchive& Ar, class UPackageMap* Map, bool& bOutSuccess);
+	bool NetSerialize(FArchive& Ar, class UPackageMap* Map, bool& bOutSuccess);
 
 	virtual UScriptStruct* GetScriptStruct()
 	{
@@ -348,6 +356,15 @@ struct GAMEPLAYABILITIES_API FGameplayAbilityTargetingLocationInfo
 	}
 };
 
+template<>
+struct TStructOpsTypeTraits<FGameplayAbilityTargetingLocationInfo> : public TStructOpsTypeTraitsBase
+{
+	enum
+	{
+		WithNetSerializer = true	// For now this is REQUIRED for FGameplayAbilityTargetDataHandle net serialization to work
+	};
+};
+
 
 /**
 *	FGameplayAbilityActivationInfo
@@ -495,6 +512,121 @@ struct GAMEPLAYABILITIES_API FGameplayAbilityTargetData
 	virtual FString ToString() const;
 };
 
+USTRUCT(BlueprintType)
+struct GAMEPLAYABILITIES_API FGameplayAbilityTargetData_LocationInfo : public FGameplayAbilityTargetData
+{
+	GENERATED_USTRUCT_BODY()
+
+	/** Generic location data for source */
+	UPROPERTY(BlueprintReadWrite, EditAnywhere, Category = Targeting)
+	FGameplayAbilityTargetingLocationInfo SourceLocation;
+
+	/** Generic location data for target */
+	UPROPERTY(BlueprintReadWrite, EditAnywhere, Category = Targeting)
+	FGameplayAbilityTargetingLocationInfo TargetLocation;
+
+	// -------------------------------------
+
+	virtual bool HasOrigin() const
+	{
+		return true;
+	}
+
+	virtual FTransform GetOrigin() const
+	{
+		return SourceLocation.GetTargetingTransform();
+	}
+
+	// -------------------------------------
+
+	virtual bool HasEndPoint() const
+	{
+		return true;
+	}
+
+	virtual FVector GetEndPoint() const
+	{
+		return TargetLocation.GetTargetingTransform().GetLocation();
+	}
+
+	// -------------------------------------
+
+	virtual UScriptStruct* GetScriptStruct()
+	{
+		return FGameplayAbilityTargetData_LocationInfo::StaticStruct();
+	}
+
+	virtual FString ToString() const override
+	{
+		return TEXT("FGameplayAbilityTargetData_LocationInfo");
+	}
+
+	bool NetSerialize(FArchive& Ar, class UPackageMap* Map, bool& bOutSuccess);
+};
+
+template<>
+struct TStructOpsTypeTraits<FGameplayAbilityTargetData_LocationInfo> : public TStructOpsTypeTraitsBase
+{
+	enum
+	{
+		WithNetSerializer = true	// For now this is REQUIRED for FGameplayAbilityTargetDataHandle net serialization to work
+	};
+};
+
+USTRUCT(BlueprintType)
+struct GAMEPLAYABILITIES_API FGameplayAbilityTargetData_ActorArray : public FGameplayAbilityTargetData
+{
+	GENERATED_USTRUCT_BODY()
+
+	/** We could be selecting this group of actors from any type of location, so use a generic location type */
+	UPROPERTY(BlueprintReadWrite, EditAnywhere, Category = Targeting)
+	FGameplayAbilityTargetingLocationInfo SourceLocation;
+
+	/** Rather than targeting a single point, this type of targeting selects multiple actors. */
+	UPROPERTY(BlueprintReadWrite, EditAnywhere, Category = Targeting)
+	TArray<AActor*> TargetActorArray;
+
+	virtual TArray<AActor*>	GetActors() const override
+	{
+		return TargetActorArray;
+	}
+
+	// -------------------------------------
+
+	virtual bool HasOrigin() const override
+	{
+		return true;
+	}
+
+	virtual FTransform GetOrigin() const override
+	{
+		return SourceLocation.GetTargetingTransform();		//No single aim-at target, since this is aimed at multiple characters.
+	}
+
+	// -------------------------------------
+
+	virtual UScriptStruct* GetScriptStruct() override
+	{
+		return FGameplayAbilityTargetData_ActorArray::StaticStruct();
+	}
+
+	virtual FString ToString() const override
+	{
+		return TEXT("FGameplayAbilityTargetData_ActorArray");
+	}
+
+	bool NetSerialize(FArchive& Ar, class UPackageMap* Map, bool& bOutSuccess);
+};
+
+template<>
+struct TStructOpsTypeTraits<FGameplayAbilityTargetData_ActorArray> : public TStructOpsTypeTraitsBase
+{
+	enum
+	{
+		WithNetSerializer = true	// For now this is REQUIRED for FGameplayAbilityTargetDataHandle net serialization to work
+	};
+};
+
 USTRUCT(BlueprintType)
 struct GAMEPLAYABILITIES_API FGameplayAbilityTargetData_Mesh : public FGameplayAbilityTargetData
 {
@@ -524,12 +656,12 @@ struct GAMEPLAYABILITIES_API FGameplayAbilityTargetData_Mesh : public FGameplayA
 
 	// -------------------------------------
 
-	virtual bool HasOrigin() const
+	virtual bool HasOrigin() const override
 	{
 		return (SourceActor != NULL);
 	}
 
-	virtual FTransform GetOrigin() const
+	virtual FTransform GetOrigin() const override
 	{
 		if (SourceActor)
 		{
@@ -546,26 +678,26 @@ struct GAMEPLAYABILITIES_API FGameplayAbilityTargetData_Mesh : public FGameplayA
 
 	// -------------------------------------
 
-	virtual bool HasEndPoint() const
+	virtual bool HasEndPoint() const override
 	{
 		return true;
 	}
 
-	virtual FVector GetEndPoint() const
+	virtual FVector GetEndPoint() const override
 	{
 		return TargetPoint;
 	}
 
 	// -------------------------------------
 
-	virtual FString ToString() const
+	virtual FString ToString() const override
 	{
 		return TEXT("FGameplayAbilityTargetData_Mesh");
 	}
 
 	bool NetSerialize(FArchive& Ar, class UPackageMap* Map, bool& bOutSuccess);
 
-	virtual UScriptStruct* GetScriptStruct()
+	virtual UScriptStruct* GetScriptStruct() override
 	{
 		return FGameplayAbilityTargetData_Mesh::StaticStruct();
 	}
@@ -595,7 +727,7 @@ struct GAMEPLAYABILITIES_API FGameplayAbilityTargetData_SingleTargetHit : public
 
 	// -------------------------------------
 
-	virtual TArray<AActor*>	GetActors() const
+	virtual TArray<AActor*>	GetActors() const override
 	{
 		TArray<AActor*>	Actors;
 		if (HitResult.Actor.IsValid())
@@ -607,12 +739,12 @@ struct GAMEPLAYABILITIES_API FGameplayAbilityTargetData_SingleTargetHit : public
 	
 	// -------------------------------------
 
-	virtual bool HasHitResult() const
+	virtual bool HasHitResult() const override
 	{
 		return true;
 	}
 
-	virtual const FHitResult* GetHitResult() const
+	virtual const FHitResult* GetHitResult() const override
 	{
 		return &HitResult;
 	}
@@ -624,7 +756,7 @@ struct GAMEPLAYABILITIES_API FGameplayAbilityTargetData_SingleTargetHit : public
 
 	bool NetSerialize(FArchive& Ar, class UPackageMap* Map, bool& bOutSuccess);
 
-	virtual UScriptStruct* GetScriptStruct()
+	virtual UScriptStruct* GetScriptStruct() override
 	{
 		return FGameplayAbilityTargetData_SingleTargetHit::StaticStruct();
 	}
diff --git a/Engine/Source/Runtime/GameplayAbilities/Public/AbilitySystemBlueprintLibrary.h b/Engine/Source/Runtime/GameplayAbilities/Public/AbilitySystemBlueprintLibrary.h
--- a/Engine/Source/Runtime/GameplayAbilities/Public/AbilitySystemBlueprintLibrary.h
+++ b/Engine/Source/Runtime/GameplayAbilities/Public/AbilitySystemBlueprintLibrary.h
@@ -28,12 +28,18 @@ class UAbilitySystemBlueprintLibrary : public UBlueprintFunctionLibrary
 	UFUNCTION(BlueprintCallable, Category = "Ability|TargetData")
 	static void ApplyGameplayEffectToTargetData(FGameplayAbilityTargetDataHandle Target, UGameplayEffect *GameplayEffect, const FGameplayAbilityActorInfo InstigatorInfo);
 
+	UFUNCTION(BlueprintPure, Category = "Ability|TargetData")
+	static FGameplayAbilityTargetDataHandle	AbilityTargetDataFromLocations(FGameplayAbilityTargetingLocationInfo SourceLocation, FGameplayAbilityTargetingLocationInfo TargetLocation);
+
 	UFUNCTION(BlueprintPure, Category = "Ability|TargetData")
 	static FGameplayAbilityTargetDataHandle	AbilityTargetDataHandleFromAbilityTargetDataMesh(FGameplayAbilityTargetData_Mesh Data);
 
 	UFUNCTION(BlueprintPure, Category = "Ability|TargetData")
 	static FGameplayAbilityTargetDataHandle	AbilityTargetDataFromHitResult(FHitResult HitResult);
 
+	UFUNCTION(BlueprintPure, Category = "Ability|TargetData")
+	static TArray<AActor*> GetActorsFromTargetData(FGameplayAbilityTargetDataHandle TargetData);
+
 	UFUNCTION(BlueprintPure, Category = "Ability|TargetData")
 	static bool TargetDataHasHitResult(FGameplayAbilityTargetDataHandle HitResult);
 
