CharacterJump abilities should be non instanced

diff --git a/Engine/Source/Runtime/GameplayAbilities/Private/Abilities/GameplayAbility_CharacterJump.cpp b/Engine/Source/Runtime/GameplayAbilities/Private/Abilities/GameplayAbility_CharacterJump.cpp
--- a/Engine/Source/Runtime/GameplayAbilities/Private/Abilities/GameplayAbility_CharacterJump.cpp
+++ b/Engine/Source/Runtime/GameplayAbilities/Private/Abilities/GameplayAbility_CharacterJump.cpp
@@ -13,6 +13,7 @@ UGameplayAbility_CharacterJump::UGameplayAbility_CharacterJump(const class FPost
 	: Super(PCIP)
 {
 	NetExecutionPolicy = EGameplayAbilityNetExecutionPolicy::Predictive;
+	InstancingPolicy = EGameplayAbilityInstancingPolicy::NonInstanced;
 }
 
 void UGameplayAbility_CharacterJump::ActivateAbility(const FGameplayAbilityActorInfo* ActorInfo, const FGameplayAbilityActivationInfo ActivationInfo)
