Updated comments for GameplayEffects.

diff --git a/Engine/Source/Runtime/GameplayAbilities/Private/GameplayEffect.cpp b/Engine/Source/Runtime/GameplayAbilities/Private/GameplayEffect.cpp
--- a/Engine/Source/Runtime/GameplayAbilities/Private/GameplayEffect.cpp
+++ b/Engine/Source/Runtime/GameplayAbilities/Private/GameplayEffect.cpp
@@ -276,7 +276,7 @@ int32 FGameplayEffectSpec::ApplyModifier(const FModifierSpec &InMod, const FModi
 
 		case EGameplayModEffect::Duration:
 		{
-			// don't modify infinite duration unless we're overriding it
+			// don't modify infinite duration or instant unless we're overriding it
 			if (GetDuration() > 0.f || InMod.Info.ModifierOp == EGameplayModOp::Override)
 			{
 				Duration.Get()->ApplyMod(InMod.Info.ModifierOp, InMod.Aggregator, bApplyAsSnapshot);
@@ -462,10 +462,10 @@ void FGameplayEffectInstigatorContext::AddInstigator(class AActor *InInstigator)
 	InstigatorAbilitySystemComponent = NULL;
 
 	// Cache off his AbilitySystemComponent.
-	IAbilitySystemInterface* AbilitySystmeInterface = InterfaceCast<IAbilitySystemInterface>(Instigator);
-	if (AbilitySystmeInterface)
+	IAbilitySystemInterface* AbilitySystemInterface = InterfaceCast<IAbilitySystemInterface>(Instigator);
+	if (AbilitySystemInterface)
 	{
-		InstigatorAbilitySystemComponent = AbilitySystmeInterface->GetAbilitySystemComponent();
+		InstigatorAbilitySystemComponent = AbilitySystemInterface->GetAbilitySystemComponent();
 	}
 }
 
diff --git a/Engine/Source/Runtime/GameplayAbilities/Public/GameplayEffect.h b/Engine/Source/Runtime/GameplayAbilities/Public/GameplayEffect.h
--- a/Engine/Source/Runtime/GameplayAbilities/Public/GameplayEffect.h
+++ b/Engine/Source/Runtime/GameplayAbilities/Public/GameplayEffect.h
@@ -46,11 +46,11 @@ struct GAMEPLAYABILITIES_API FGameplayEffectLevelDef
 {
 	GENERATED_USTRUCT_BODY()
 
-	/** When true, whatever creates of owns this will pass in a level. E.g, level is not intrinsic to this definition. */
+	/** When true, whatever creates or owns this will pass in a level. E.g, level is not intrinsic to this definition. */
 	UPROPERTY(EditDefaultsOnly, Category = GameplayEffectLevel)
 	bool InheritLevelFromOwner;
 
-	/** If set, the gameplay effects level will be tied to this attribute on the instigator */
+	/** If set, the gameplay effect's level will be tied to this attribute on the instigator */
 	UPROPERTY(EditDefaultsOnly, Category = GameplayEffectLevel, meta = (FilterMetaTag="HideFromLevelInfos"))
 	FGameplayAttribute	Attribute;
 
@@ -79,6 +79,7 @@ struct GAMEPLAYABILITIES_API FGameplayModifierInfo
 
 	}
 
+	/** How much this modifies what it is applied to */
 	UPROPERTY(EditDefaultsOnly, Category = GameplayModifier)
 	FScalableFloat Magnitude; // Not modified from defaults
 
@@ -90,7 +91,7 @@ struct GAMEPLAYABILITIES_API FGameplayModifierInfo
 	UPROPERTY(EditDefaultsOnly, Category=GameplayModifier)
 	FGameplayAttribute Attribute;
 
-	/** The numeric operation of this modifier: Override, Add, Multiply  */
+	/** The numeric operation of this modifier: Override, Add, Multiply, etc  */
 	UPROPERTY(EditDefaultsOnly, Category=GameplayModifier)
 	TEnumAsByte<EGameplayModOp::Type> ModifierOp;
 
@@ -102,15 +103,15 @@ struct GAMEPLAYABILITIES_API FGameplayModifierInfo
 	UPROPERTY(EditDefaultsOnly, Category=GameplayModifier)
 	UGameplayEffect* TargetEffect;
 
-	// The thing I modify requires these tags
+	/** The thing I modify requires these tags */
 	UPROPERTY(EditDefaultsOnly, Category=GameplayModifier)
 	FGameplayTagContainer RequiredTags;
 
-	// The thing I modify must not have any of these tags
+	/** The thing I modify must not have any of these tags */
 	UPROPERTY(EditDefaultsOnly, Category = GameplayModifier)
 	FGameplayTagContainer IgnoreTags;
 	
-	/** This modifiers tags. These tags are passed to any other modifiers that this modifies. */
+	/** This modifier's tags. These tags are passed to any other modifiers that this modifies. */
 	UPROPERTY(EditDefaultsOnly, Category=GameplayModifier)
 	FGameplayTagContainer OwnedTags;
 
@@ -149,12 +150,15 @@ struct FGameplayEffectCue
 		GameplayCueTags.AddTag(InTag);
 	}
 
+	/** The minimum level that this Cue supports */
 	UPROPERTY(EditDefaultsOnly, Category = GameplayCue)
 	float	MinLevel;
 
+	/** The maximum level that this Cue supports */
 	UPROPERTY(EditDefaultsOnly, Category = GameplayCue)
 	float	MaxLevel;
 
+	/** Tags passed to the gameplay cue handler when this cue is activated */
 	UPROPERTY(EditDefaultsOnly, Category = GameplayCue, meta = (Categories="GameplayCue"))
 	FGameplayTagContainer GameplayCueTags;
 
@@ -200,43 +204,48 @@ public:
 	bool ShowAllProperties;
 #endif
 
-	
+	/** Duration in seconds. 0.0 for instantaneous effects; -1.0 for infinite duration. */
 	UPROPERTY(EditDefaultsOnly, Category=GameplayEffect)
 	FScalableFloat	Duration;
 
+	/** Period in seconds. 0.0 for non-periodic effects */
 	UPROPERTY(EditDefaultsOnly, Category=GameplayEffect)
 	FScalableFloat	Period;
 	
+	/** Array of modifiers that will affect the target of this effect */
 	UPROPERTY(EditDefaultsOnly, BlueprintReadOnly, Category=GameplayEffect)
 	TArray<FGameplayModifierInfo> Modifiers;
 
+	/** Array of level definitions that will determine how this GameplayEffect scales */
 	UPROPERTY(EditDefaultsOnly, BlueprintReadWrite, Category=GameplayEffect)
 	FGameplayEffectLevelDef	LevelInfo;
 		
 	// "I can only be applied to targets that have these tags"
 	// "I can only exist on CE buckets on targets that have these tags":
 	
-	/** Container of gameplay tags that have to be present on the target for the effect to be applied */
+	/** Container of gameplay tags that have to be present on the target actor for the effect to be applied */
 	UPROPERTY(EditDefaultsOnly, BlueprintReadOnly, Category = Application)
 	FGameplayTagContainer ApplicationRequiredTargetTags;
 
 	// "I can only be applied if my instigator has these tags"
 
-	/** Container of gameplay tags that have to be present on the instigator for the effect to be applied */
+	/** Container of gameplay tags that have to be present on the instigator actor for the effect to be applied */
 	UPROPERTY(EditDefaultsOnly, BlueprintReadOnly, Category = Application)
 	FGameplayTagContainer ApplicationRequiredInstigatorTags;
 
+	/** Probability that this gameplay effect will be applied to the target actor (0.0 for never, 1.0 for always) */
 	UPROPERTY(EditDefaultsOnly, Category=Application, meta=(GameplayAttribute="True"))
 	FScalableFloat	ChanceToApplyToTarget;
 
+	/** Probability that this gameplay effect will execute on another GE after it has been successfully applied to the target actor (0.0 for never, 1.0 for always) */
 	UPROPERTY(EditDefaultsOnly, Category = Application, meta = (GameplayAttribute = "True"))
 	FScalableFloat	ChanceToExecuteOnGameplayEffect;
 
-	// other gameplay effects  that will be applied to the target of this effect
+	/** other gameplay effects that will be applied to the target of this effect if this effect applies */
 	UPROPERTY(EditDefaultsOnly, BlueprintReadOnly, Category = GameplayEffect)
 	TArray<UGameplayEffect*> TargetEffects;
 
-	// removes or blocks gameplay effects that it applies to if the tags match
+	/** removes active gameplay effects and stops gameplay effects from applying if the tags and qualification context match */
 	UPROPERTY(EditDefaultsOnly, BlueprintReadOnly, Category = GameplayEffect)
 	TEnumAsByte<EGameplayImmunity::Type> AppliesImmunityTo;
 
@@ -253,6 +262,7 @@ public:
 	 * 
 	 * @param InstigatorTags	Owned gameplay tags of the instigator applying the effect
 	 * @param TargetTags		Owned gameplay tags of the target about to be affected by the effect
+	 * @return					True if the instigator and target actor tags meet the requirements for this gameplay effect, false otherwise
 	 */
 	bool AreApplicationTagRequirementsSatisfied(const FGameplayTagContainer& InstigatorTags, const FGameplayTagContainer& TargetTags) const;
 
@@ -260,18 +270,24 @@ public:
 	// New Tagging functionality
 	// ------------------------------------------------
 
-	// "These are my tags"
+	/** "These are my tags" */
 	UPROPERTY(EditDefaultsOnly, BlueprintReadOnly, Category=Tags)
 	FGameplayTagContainer GameplayEffectTags;
 
-	// "In order to affect another GE, they must have ALL of these tags"
+	/** "In order to affect another GameplayEffect, they must have ALL of these tags" */
 	UPROPERTY(EditDefaultsOnly, BlueprintReadOnly, Category=Tags)
 	FGameplayTagContainer GameplayEffectRequiredTags;
 	
-	// "In order to affect another GE, they must NOT have ANY of these tags"
+	/** "In order to affect another GameplayEffect, they must NOT have ANY of these tags" */
 	UPROPERTY(EditDefaultsOnly, BlueprintReadOnly, Category=Tags)
 	FGameplayTagContainer GameplayEffectIgnoreTags;
 
+	/**
+	* Can this GameplayEffect modify a GameplayEffect that owns Tags
+	*
+	* @param Tags	Owned gameplay tags of the gameplay effect this effect is being applied to
+	* @return		True if the tags meet the requirements for this gameplay effect, false otherwise
+	*/
 	bool AreGameplayEffectTagRequirementsSatisfied(const FGameplayTagContainer& Tags) const
 	{
 		bool bHasRequired = Tags.MatchesAll(GameplayEffectRequiredTags, true);
@@ -280,7 +296,12 @@ public:
 		return bHasRequired && !bHasIgnored;
 	}
 
-	/** Can this GameplayEffect modify the input parameter, based on tags  */
+	/**
+	* Can this GameplayEffect modify the input parameter, based on tags
+	*
+	* @param GameplayEffectToBeModified	A GameplayEffect we are trying to apply this GameplayEffect to.
+	* @return							True if the tags owned by GameplayEffectToBeModified meet the requirements for this gameplay effect, false otherwise
+	*/
 	bool AreGameplayEffectTagRequirementsSatisfied(const UGameplayEffect *GameplayEffectToBeModified) const
 	{
 		return AreGameplayEffectTagRequirementsSatisfied(GameplayEffectToBeModified->GameplayEffectTags);
@@ -298,7 +319,7 @@ public:
 
 	void ValidateGameplayEffect();
 
-	/** Container of "owned" gameplay tags that are applied to actor with this effect applied to them */
+	/** "These tags are applied to the actor I am applied to" */
 	UPROPERTY(EditDefaultsOnly, BlueprintReadOnly, Category=Tags)
 	FGameplayTagContainer OwnedTagsContainer;
 
@@ -316,26 +337,31 @@ public:
 
 	// -----------------------------------------------
 	
+	/** Should copies of this GameplayEffect be a snapshot of the current state or update when it does (linked) */
 	UPROPERTY(EditDefaultsOnly, BlueprintReadOnly, Category = Advanced)
 	TEnumAsByte<EGameplayEffectCopyPolicy::Type>	CopyPolicy;
 
 	// ----------------------------------------------
 
+	/** Cues to trigger non-simulated reactions in response to this GameplayEffect such as sounds, particle effects, etc */
 	UPROPERTY(EditDefaultsOnly, BlueprintReadOnly, Category = Display)
 	TArray<FGameplayEffectCue>	GameplayCues;
 
-	/** Description of this combat effect. */
+	/** Description of this gameplay effect. */
 	UPROPERTY(EditDefaultsOnly, BlueprintReadWrite, Category = Display)
 	FText Description;
 
 	// ----------------------------------------------
 
+	/** Specifies the rule used to stack this GameplayEffect with other GameplayEffects. */
 	UPROPERTY(EditDefaultsOnly, BlueprintReadOnly, Category = Stacking)
 	TEnumAsByte<EGameplayEffectStackingPolicy::Type>	StackingPolicy;
 
+	/** An identifier for the stack. Both names and stacking policy must match for GameplayEffects to stack with each other. */
 	UPROPERTY(EditDefaultsOnly, BlueprintReadOnly, Category = Stacking)
 	FName StackedAttribName;
 
+	/** Specifies a custom stacking rule if one is needed. */
 	UPROPERTY(EditDefaultsOnly, Category = Stacking)
 	TSubclassOf<class UGameplayEffectStackingExtension> StackingExtension;
 
@@ -575,7 +601,7 @@ struct TStructOpsTypeTraits< FAggregatorRef > : public TStructOpsTypeTraitsBase
 };
 
 /**
- * GameplayModifierEvaluatedData
+ * GameplayModifierData
  *	This is the data that FAggregator aggregates and turns into FGameplayModifierEvaluatedData.
  *  It is distinct from FGameplayModifierEvaluatedData in that FGameplayModifierData ia level has not been applied to this data.
  *  FGameplayModifierData::Magnitude is an FScalableFloat which describes a numeric value for a given level.
@@ -1006,12 +1032,18 @@ struct FGameplayEffectSpec
 	// other effects that need to be applied to the target if this effect is successful
 	TArray< TSharedRef< FGameplayEffectSpec > > TargetEffectSpecs;
 
+	// The duration in seconds of this effect
+	// instantaneous effects should have a duration of UGameplayEffect::INSTANT_APPLICATION
+	// effects that last forever should have a duration of UGameplayEffect::INFINITE_DURATION
 	UPROPERTY()
 	FAggregatorRef	Duration;
 
+	// The period in seconds of this effect.
+	// Nonperiodic effects should have a period of UGameplayEffect::NO_PERIOD
 	UPROPERTY()
 	FAggregatorRef	Period;
 
+	// The chance, in a 0.0-1.0 range, that this GameplayEffect will be applied to the target Attribute or GameplayEffect.
 	UPROPERTY()
 	FAggregatorRef	ChanceToApplyToTarget;
 
@@ -1106,6 +1138,7 @@ struct FActiveGameplayEffect : public FFastArraySerializerItem
 
 	FTimerHandle DurationHandle;
 
+	/** Prediction keys used to sync between client and server. If CurrPredictionKey != 0 then this effect does not use prediction. */
 	uint32 PrevPredictionKey;
 	uint32 CurrPredictionKey;
 
