fix clients having 2x ability sets. fixes andrewb's problem

diff --git a/Engine/Source/Runtime/GameplayAbilities/Private/AbilitySystemComponent_Abilities.cpp b/Engine/Source/Runtime/GameplayAbilities/Private/AbilitySystemComponent_Abilities.cpp
--- a/Engine/Source/Runtime/GameplayAbilities/Private/AbilitySystemComponent_Abilities.cpp
+++ b/Engine/Source/Runtime/GameplayAbilities/Private/AbilitySystemComponent_Abilities.cpp
@@ -107,6 +107,14 @@ UGameplayAbility* UAbilitySystemComponent::GiveAbility(UGameplayAbility* Ability
 	return Ability;
 }
 
+void UAbilitySystemComponent::ClearAllAbilities()
+{
+	check(IsOwnerActorAuthoritative());	// Should be called on authority
+
+	// Note we aren't marking any old abilities pending kill. This shouldn't matter since they will be garbage collected.
+	ActivatableAbilities.Empty(ActivatableAbilities.Num());
+}
+
 void UAbilitySystemComponent::OnGiveAbility(UGameplayAbility* Ability, int32 InputID)
 {
 	for (FAbilityTriggerData TriggerData : Ability->AbilityTriggers)
diff --git a/Engine/Source/Runtime/GameplayAbilities/Public/AbilitySystemComponent.h b/Engine/Source/Runtime/GameplayAbilities/Public/AbilitySystemComponent.h
--- a/Engine/Source/Runtime/GameplayAbilities/Public/AbilitySystemComponent.h
+++ b/Engine/Source/Runtime/GameplayAbilities/Public/AbilitySystemComponent.h
@@ -371,6 +371,9 @@ class GAMEPLAYABILITIES_API UAbilitySystemComponent : public UActorComponent, pu
 
 	UGameplayAbility* GiveAbility(UGameplayAbility* Ability, int32 InputID);
 
+	/** Wipes all 'given' abilities. */
+	void ClearAllAbilities();
+
 	/** Will be called from GiveAbility or from OnRep. Initializes events (triggers and inputs) with the given ability */
 	void OnGiveAbility(UGameplayAbility* Ability, int32 InputID);
 
