Minor .h file cleanup.

diff --git a/Engine/Source/Runtime/GameplayAbilities/Public/Abilities/GameplayAbilityTypes.h b/Engine/Source/Runtime/GameplayAbilities/Public/Abilities/GameplayAbilityTypes.h
--- a/Engine/Source/Runtime/GameplayAbilities/Public/Abilities/GameplayAbilityTypes.h
+++ b/Engine/Source/Runtime/GameplayAbilities/Public/Abilities/GameplayAbilityTypes.h
@@ -290,8 +290,7 @@ struct GAMEPLAYABILITIES_API FGameplayAbilityTargetingLocationInfo
 	}
 
 public:
-	//Not sure if we want to distinguish between where we actually aim from (should probably always be the user's camera or actor origin) and where we trace from.
-	FTransform GetTargetingTransform() const		//This is split out so we can
+	FTransform GetTargetingTransform() const
 	{
 		//Return or calculate based on LocationType.
 		switch (LocationType)
@@ -322,9 +321,6 @@ public:
 
 	FGameplayAbilityTargetDataHandle MakeTargetDataHandleFromActors(TArray<AActor*> TargetActors) const;
 
-	UPROPERTY(BlueprintReadWrite, meta = (ExposeOnSpawn = true), Category = Targeting)
-	TEnumAsByte<EGameplayAbilityTargetingLocationType::Type> LocationType;
-
 	/** A literal world transform can be used, if one has been calculated outside of the actor using the ability. */
 	UPROPERTY(BlueprintReadWrite, meta = (ExposeOnSpawn = true), Category = Targeting)
 	FTransform LiteralTransform;
@@ -341,6 +337,10 @@ public:
 	UPROPERTY(BlueprintReadWrite, meta = (ExposeOnSpawn = true), Category = Targeting)
 	FName SourceSocketName;
 
+	/** Type of location used - will determine what data is transmitted over the network and what fields are used when calculating position. */
+	UPROPERTY(BlueprintReadWrite, meta = (ExposeOnSpawn = true), Category = Targeting)
+	TEnumAsByte<EGameplayAbilityTargetingLocationType::Type> LocationType;
+
 	// -------------------------------------
 
 	virtual FString ToString() const
