Added convenience method for getting skeletal mesh component directly.

diff --git a/Engine/Source/Runtime/GameplayAbilities/Private/Abilities/GameplayAbility.cpp b/Engine/Source/Runtime/GameplayAbilities/Private/Abilities/GameplayAbility.cpp
--- a/Engine/Source/Runtime/GameplayAbilities/Private/Abilities/GameplayAbility.cpp
+++ b/Engine/Source/Runtime/GameplayAbilities/Private/Abilities/GameplayAbility.cpp
@@ -509,6 +509,17 @@ FGameplayAbilityActorInfo UGameplayAbility::GetActorInfo()
 	return *CurrentActorInfo;
 }
 
+/** Convenience method for abilities to get skeletal mesh component - useful for aiming abilities */
+USkeletalMeshComponent* UGameplayAbility::GetOwningComponentFromActorInfo()
+{
+	check(CurrentActorInfo);
+	if (CurrentActorInfo->AnimInstance.IsValid())
+	{
+		return CurrentActorInfo->AnimInstance.Get()->GetOwningComponent();
+	}
+	return NULL;
+}
+
 /** Convenience method for abilities to get outgoing gameplay effect specs (for example, to pass on to projectiles to apply to whoever they hit) */
 FGameplayEffectSpecHandle UGameplayAbility::GetOutgoingSpec(UGameplayEffect* GameplayEffect) const
 {
diff --git a/Engine/Source/Runtime/GameplayAbilities/Public/Abilities/GameplayAbility.h b/Engine/Source/Runtime/GameplayAbilities/Public/Abilities/GameplayAbility.h
--- a/Engine/Source/Runtime/GameplayAbilities/Public/Abilities/GameplayAbility.h
+++ b/Engine/Source/Runtime/GameplayAbilities/Public/Abilities/GameplayAbility.h
@@ -335,6 +335,9 @@ protected:
 	UFUNCTION(BlueprintPure, Category=Ability)
 	FGameplayAbilityActorInfo GetActorInfo();
 
+	UFUNCTION(BlueprintPure, Category = Ability)
+	USkeletalMeshComponent* GetOwningComponentFromActorInfo();
+
 	UFUNCTION(BlueprintCallable, Category=Ability)
 	FGameplayEffectSpecHandle GetOutgoingSpec(UGameplayEffect* GameplayEffect) const;
 
