Updates to targeting:
-Make 'ConfirmationType' a property of the Targeting Task rather than targeting actor.
-Allow TargetActors to be 'server side' and only require a confirmation from client (or no confirm at all)
-Added _Radius target type (wip for now, targets all pawns)

diff --git a/Engine/Source/Runtime/GameplayAbilities/Private/Abilities/GameplayAbility.cpp b/Engine/Source/Runtime/GameplayAbilities/Private/Abilities/GameplayAbility.cpp
--- a/Engine/Source/Runtime/GameplayAbilities/Private/Abilities/GameplayAbility.cpp
+++ b/Engine/Source/Runtime/GameplayAbilities/Private/Abilities/GameplayAbility.cpp
@@ -716,11 +716,15 @@ void FGameplayAbilityTargetData::ApplyGameplayEffect(UGameplayEffect* GameplayEf
 		SpecToApply.InstigatorContext.AddHitResult(*GetHitResult());
 	}
 
+	if (HasOrigin())
+	{
+		SpecToApply.InstigatorContext.AddOrigin(GetOrigin().GetLocation());
+	}
+
 	TArray<AActor*> Actors = GetActors();
 	for (AActor * TargetActor : Actors)
 	{
 		check(TargetActor);
-		
 		UAbilitySystemComponent * TargetComponent = UAbilitySystemBlueprintLibrary::GetAbilitySystemComponent(TargetActor);
 		if (TargetComponent)
 		{
@@ -738,7 +742,6 @@ FString FGameplayAbilityTargetData::ToString() const
 
 FGameplayAbilityTargetDataHandle FGameplayAbilityTargetingLocationInfo::MakeTargetDataHandleFromHitResult(TWeakObjectPtr<UGameplayAbility> Ability, FHitResult HitResult) const
 {
-	/** Note: These are cleaned up by the FGameplayAbilityTargetDataHandle (via an internal TSharedPtr) */
 	if (LocationType == EGameplayAbilityTargetingLocationType::Type::SocketTransform)
 	{
 		const FGameplayAbilityActorInfo* ActorInfo = Ability.IsValid() ? Ability.Get()->GetCurrentActorInfo() : NULL;
@@ -756,6 +759,8 @@ FGameplayAbilityTargetDataHandle FGameplayAbilityTargetingLocationInfo::MakeTarg
 			return FGameplayAbilityTargetDataHandle(ReturnData);
 		}
 	}
+
+	/** Note: These are cleaned up by the FGameplayAbilityTargetDataHandle (via an internal TSharedPtr) */
 	FGameplayAbilityTargetData_SingleTargetHit* ReturnData = new FGameplayAbilityTargetData_SingleTargetHit();
 	ReturnData->HitResult = HitResult;
 	return FGameplayAbilityTargetDataHandle(ReturnData);
@@ -820,7 +825,7 @@ bool FGameplayAbilityTargetDataHandle::NetSerialize(FArchive& Ar, class UPackage
 	}
 	
 
-	ABILITY_LOG(Warning, TEXT("FGameplayAbilityTargetDataHandle Serialized: %s"), ScriptStruct ? *ScriptStruct->GetName() : TEXT("NULL") );
+	//ABILITY_LOG(Warning, TEXT("FGameplayAbilityTargetDataHandle Serialized: %s"), ScriptStruct ? *ScriptStruct->GetName() : TEXT("NULL") );
 
 	bOutSuccess = true;
 	return true;
@@ -925,4 +930,4 @@ void UGameplayAbility::RemoveGameplayCue(FGameplayTag GameplayCueTag)
 	{
 		CurrentActorInfo->AbilitySystemComponent->RemoveGameplayCue(GameplayCueTag);
 	}
-}
\ No newline at end of file
+}
diff --git a/Engine/Source/Runtime/GameplayAbilities/Private/Abilities/GameplayAbilityTargetActor.cpp b/Engine/Source/Runtime/GameplayAbilities/Private/Abilities/GameplayAbilityTargetActor.cpp
--- a/Engine/Source/Runtime/GameplayAbilities/Private/Abilities/GameplayAbilityTargetActor.cpp
+++ b/Engine/Source/Runtime/GameplayAbilities/Private/Abilities/GameplayAbilityTargetActor.cpp
@@ -14,6 +14,7 @@ AGameplayAbilityTargetActor::AGameplayAbilityTargetActor(const class FPostConstr
 	: Super(PCIP)
 {
 	StaticTargetFunction = false;
+	ShouldProduceTargetDataOnServer = false;
 }
 
 void AGameplayAbilityTargetActor::GetLifetimeReplicatedProps(TArray< FLifetimeProperty > & OutLifetimeProps) const
@@ -24,6 +25,7 @@ void AGameplayAbilityTargetActor::GetLifetimeReplicatedProps(TArray< FLifetimePr
 
 void AGameplayAbilityTargetActor::StartTargeting(UGameplayAbility* Ability)
 {
+	OwningAbility = Ability;
 }
 
 void AGameplayAbilityTargetActor::ConfirmTargeting()
@@ -69,5 +71,18 @@ bool AGameplayAbilityTargetActor::OnReplicatedTargetDataReceived(FGameplayAbilit
 
 bool AGameplayAbilityTargetActor::ShouldProduceTargetData() const
 {
-	return (MasterPC && MasterPC->IsLocalController());
+	// return true if we are locally owned, or (we are the server and this is a gameplaytarget ability that can produce target data server side)
+	return (MasterPC && (MasterPC->IsLocalController() || ShouldProduceTargetDataOnServer));
 }
+
+void AGameplayAbilityTargetActor::BindToConfirmCancelInputs()
+{
+	check(OwningAbility);
+
+	UAbilitySystemComponent* ASC = OwningAbility->GetCurrentActorInfo()->AbilitySystemComponent.Get();
+	if (ASC)
+	{
+		ASC->ConfirmCallbacks.AddDynamic(this, &AGameplayAbilityTargetActor::ConfirmTargeting);
+		ASC->CancelCallbacks.AddDynamic(this, &AGameplayAbilityTargetActor::CancelTargeting);
+	}
+}
\ No newline at end of file
diff --git a/Engine/Source/Runtime/GameplayAbilities/Private/Abilities/GameplayAbilityTargetActor_Radius.cpp b/Engine/Source/Runtime/GameplayAbilities/Private/Abilities/GameplayAbilityTargetActor_Radius.cpp
new file mode 100644
--- /dev/null
+++ b/Engine/Source/Runtime/GameplayAbilities/Private/Abilities/GameplayAbilityTargetActor_Radius.cpp
@@ -0,0 +1,95 @@
+// Copyright 1998-2014 Epic Games, Inc. All Rights Reserved.
+
+#include "AbilitySystemPrivatePCH.h"
+#include "AbilitySystemComponent.h"
+#include "GameplayAbilityTargetActor_Radius.h"
+#include "Engine/World.h"
+#include "Runtime/Engine/Public/Net/UnrealNetwork.h"
+
+// --------------------------------------------------------------------------------------------------------------------------------------------------------
+//
+//	AGameplayAbilityTargetActor_Radius
+//
+// --------------------------------------------------------------------------------------------------------------------------------------------------------
+
+AGameplayAbilityTargetActor_Radius::AGameplayAbilityTargetActor_Radius(const class FPostConstructInitializeProperties& PCIP)
+	: Super(PCIP)
+{
+	PrimaryActorTick.bCanEverTick = true;
+	PrimaryActorTick.TickGroup = TG_PrePhysics;
+	StaticTargetFunction = false;
+	ShouldProduceTargetDataOnServer = true;
+}
+
+void AGameplayAbilityTargetActor_Radius::GetLifetimeReplicatedProps(TArray< FLifetimeProperty > & OutLifetimeProps) const
+{
+	Super::GetLifetimeReplicatedProps(OutLifetimeProps);
+
+	DOREPLIFETIME(AGameplayAbilityTargetActor_Radius, SourceActor);
+}
+
+void AGameplayAbilityTargetActor_Radius::StartTargeting(UGameplayAbility* InAbility)
+{
+	Super::StartTargeting(InAbility);
+	SourceActor = InAbility->GetCurrentActorInfo()->Actor.Get();
+}
+
+void AGameplayAbilityTargetActor_Radius::ConfirmTargeting()
+{
+	check(ShouldProduceTargetData());
+
+	if (SourceActor)
+	{
+		FVector Origin = StartLocation.GetTargetingTransform().GetLocation();
+
+		FGameplayAbilityTargetDataHandle Handle = MakeTargetData(PerformOverlap(Origin), Origin);
+		TargetDataReadyDelegate.Broadcast(Handle);
+	}
+
+	Destroy();
+}
+
+FGameplayAbilityTargetDataHandle AGameplayAbilityTargetActor_Radius::MakeTargetData(const TArray<AActor*> Actors, const FVector& Origin) const
+{
+	if (OwningAbility)
+	{
+		/** Note: This will be cleaned up by the FGameplayAbilityTargetDataHandle (via an internal TSharedPtr) */
+		FGameplayAbilityTargetData_Radius* ReturnData = new FGameplayAbilityTargetData_Radius(Actors, Origin);
+		return FGameplayAbilityTargetDataHandle(ReturnData);
+	}
+
+	return FGameplayAbilityTargetDataHandle();
+}
+
+TArray<AActor*>	AGameplayAbilityTargetActor_Radius::PerformOverlap(const FVector& Origin)
+{
+	static FName RadiusTargetingOverlap = FName(TEXT("RadiusTargetingOverlap"));
+	bool bTraceComplex = false;
+	
+	TArray<AActor*> ActorsToIgnore;
+	ActorsToIgnore.Add(SourceActor);
+
+	FCollisionQueryParams Params(RadiusTargetingOverlap, bTraceComplex);
+	Params.bReturnPhysicalMaterial = true;
+	Params.bTraceAsyncScene = true;
+	Params.AddIgnoredActors(ActorsToIgnore);
+
+	TArray<FOverlapResult> Overlaps;
+
+	SourceActor->GetWorld()->OverlapMulti(Overlaps, Origin, FQuat::Identity, ECC_Pawn, FCollisionShape::MakeSphere(Radius), Params);
+
+
+	TArray<AActor*>	HitActors;
+
+	for (int32 i=0; i < Overlaps.Num(); ++i)
+	{
+		// Temp, this needs to be way more robust
+		APawn* PawnActor = Cast<APawn>(Overlaps[i].GetActor());
+		if (PawnActor)
+		{
+			HitActors.Add(PawnActor);
+		}
+	}
+
+	return HitActors;
+}
\ No newline at end of file
diff --git a/Engine/Source/Runtime/GameplayAbilities/Private/Abilities/GameplayAbilityTargetActor_SingleLineTrace.cpp b/Engine/Source/Runtime/GameplayAbilities/Private/Abilities/GameplayAbilityTargetActor_SingleLineTrace.cpp
--- a/Engine/Source/Runtime/GameplayAbilities/Private/Abilities/GameplayAbilityTargetActor_SingleLineTrace.cpp
+++ b/Engine/Source/Runtime/GameplayAbilities/Private/Abilities/GameplayAbilityTargetActor_SingleLineTrace.cpp
@@ -19,13 +19,12 @@ AGameplayAbilityTargetActor_SingleLineTrace::AGameplayAbilityTargetActor_SingleL
 	PrimaryActorTick.TickGroup = TG_PrePhysics;
 	StaticTargetFunction = false;
 	bDebug = false;
-	bBindToConfirmCancelInputs = true;
 }
 
 void AGameplayAbilityTargetActor_SingleLineTrace::GetLifetimeReplicatedProps(TArray< FLifetimeProperty > & OutLifetimeProps) const
 {
 	Super::GetLifetimeReplicatedProps(OutLifetimeProps);
-	//DOREPLIFETIME(AGameplayAbilityTargetActor, StartLocation);
+	
 	DOREPLIFETIME(AGameplayAbilityTargetActor_SingleLineTrace, bDebug);
 	DOREPLIFETIME(AGameplayAbilityTargetActor_SingleLineTrace, SourceActor);
 }
@@ -55,10 +54,10 @@ FHitResult AGameplayAbilityTargetActor_SingleLineTrace::PerformTrace(AActor* InS
 	Params.AddIgnoredActors(ActorsToIgnore);
 
 	FVector AimDirection = InSourceActor->GetActorForwardVector();		//Default
-	UGameplayAbility* MyAbility = Ability.Get();
-	if (MyAbility)		//Server and launching client only
+	
+	if (OwningAbility)		//Server and launching client only
 	{
-		APlayerController* AimingPC = MyAbility->GetCurrentActorInfo()->PlayerController.Get();
+		APlayerController* AimingPC = OwningAbility->GetCurrentActorInfo()->PlayerController.Get();
 		check(AimingPC);
 		FVector CamLoc;
 		FRotator CamRot;
@@ -70,7 +69,7 @@ FHitResult AGameplayAbilityTargetActor_SingleLineTrace::PerformTrace(AActor* InS
 	FVector TraceEnd = TraceStart + (AimDirection * MaxRange);
 
 	//If we're using a socket, adjust the starting location and aim direction after the end position has been found. This way we can still aim with the camera, then fire accurately from the socket.
-	if (MyAbility)		//Server and launching client only
+	if (OwningAbility)		//Server and launching client only
 	{
 		TraceStart = StartLocation.GetTargetingTransform().GetLocation();
 	}
@@ -95,29 +94,13 @@ FHitResult AGameplayAbilityTargetActor_SingleLineTrace::PerformTrace(AActor* InS
 FGameplayAbilityTargetDataHandle AGameplayAbilityTargetActor_SingleLineTrace::StaticGetTargetData(UWorld * World, const FGameplayAbilityActorInfo* ActorInfo, FGameplayAbilityActivationInfo ActivationInfo) const
 {
 	AActor* StaticSourceActor = ActorInfo->Actor.Get();
-	check(StaticSourceActor);
-	UAnimInstance* AnimInstance = ActorInfo->AnimInstance.Get();
-	USkeletalMeshComponent* StaticSourceComponent = AnimInstance ? AnimInstance->GetOwningComponent() : NULL;
-
 	return MakeTargetData(PerformTrace(StaticSourceActor));
 }
 
 void AGameplayAbilityTargetActor_SingleLineTrace::StartTargeting(UGameplayAbility* InAbility)
 {
-	Ability = InAbility;
+	Super::StartTargeting(InAbility);
 	SourceActor = InAbility->GetCurrentActorInfo()->Actor.Get();
-	UAnimInstance* AnimInstance = InAbility->GetCurrentActorInfo()->AnimInstance.Get();
-
-	// We can bind directly to our ASC's confirm/cancel events, or wait to be told from an outside source to confirm or cancel
-	if (bBindToConfirmCancelInputs && (MasterPC && MasterPC->IsLocalController()))
-	{
-		UAbilitySystemComponent* ASC = Ability->GetCurrentActorInfo()->AbilitySystemComponent.Get();
-		if (ASC)
-		{
-			ASC->ConfirmCallbacks.AddDynamic(this, &AGameplayAbilityTargetActor_SingleLineTrace::ConfirmTargeting);
-			ASC->CancelCallbacks.AddDynamic(this, &AGameplayAbilityTargetActor_SingleLineTrace::CancelTargeting);
-		}
-	}
 	
 	bDebug = true;
 
@@ -174,5 +157,5 @@ void AGameplayAbilityTargetActor_SingleLineTrace::ConfirmTargeting()
 FGameplayAbilityTargetDataHandle AGameplayAbilityTargetActor_SingleLineTrace::MakeTargetData(FHitResult HitResult) const
 {
 	/** Note: This will be cleaned up by the FGameplayAbilityTargetDataHandle (via an internal TSharedPtr) */
-	return StartLocation.MakeTargetDataHandleFromHitResult(Ability, HitResult);
-}
\ No newline at end of file
+	return StartLocation.MakeTargetDataHandleFromHitResult(OwningAbility, HitResult);
+}
diff --git a/Engine/Source/Runtime/GameplayAbilities/Private/Abilities/Tasks/AbilityTask_WaitTargetData.cpp b/Engine/Source/Runtime/GameplayAbilities/Private/Abilities/Tasks/AbilityTask_WaitTargetData.cpp
--- a/Engine/Source/Runtime/GameplayAbilities/Private/Abilities/Tasks/AbilityTask_WaitTargetData.cpp
+++ b/Engine/Source/Runtime/GameplayAbilities/Private/Abilities/Tasks/AbilityTask_WaitTargetData.cpp
@@ -10,10 +10,11 @@ UAbilityTask_WaitTargetData::UAbilityTask_WaitTargetData(const class FPostConstr
 	
 }
 
-UAbilityTask_WaitTargetData* UAbilityTask_WaitTargetData::WaitTargetData(UObject* WorldContextObject, TSubclassOf<AGameplayAbilityTargetActor> InTargetClass)
+UAbilityTask_WaitTargetData* UAbilityTask_WaitTargetData::WaitTargetData(UObject* WorldContextObject, TEnumAsByte<EGameplayTargetingConfirmation::Type> ConfirmationType, TSubclassOf<AGameplayAbilityTargetActor> InTargetClass)
 {
 	auto MyObj = NewTask<UAbilityTask_WaitTargetData>(WorldContextObject);
 	MyObj->TargetClass = InTargetClass;
+	MyObj->ConfirmationType = ConfirmationType;
 	return MyObj;	
 }
 
@@ -42,7 +43,7 @@ bool UAbilityTask_WaitTargetData::BeginSpawningActor(UObject* WorldContextObject
 
 		// Spawn the actor if this is a locally controlled ability (always) or if this is a replicating targeting mode.
 		// (E.g., server will spawn this target actor to replicate to all non owning clients)
-		if (Replicates || IsLocallyControlled)
+		if (Replicates || IsLocallyControlled || CDO->ShouldProduceTargetDataOnServer)
 		{
 			if (StaticFunc)
 			{
@@ -59,6 +60,7 @@ bool UAbilityTask_WaitTargetData::BeginSpawningActor(UObject* WorldContextObject
 					SpawnedActor = World->SpawnActorDeferred<AGameplayAbilityTargetActor>(Class, FVector::ZeroVector, FRotator::ZeroRotator, NULL, NULL, true);
 				}
 
+				// If we spawned the target actor, always register the callbacks for when the data is ready.
 				SpawnedActor->TargetDataReadyDelegate.AddUObject(this, &UAbilityTask_WaitTargetData::OnTargetDataReadyCallback);
 				SpawnedActor->CanceledDelegate.AddUObject(this, &UAbilityTask_WaitTargetData::OnTargetDataCancelledCallback);
 
@@ -70,22 +72,25 @@ bool UAbilityTask_WaitTargetData::BeginSpawningActor(UObject* WorldContextObject
 					TargetActor->MasterPC = MyAbility->GetCurrentActorInfo()->PlayerController.Get();
 				}
 			}
-		}
-		
+		}		
 
 		// If not locally controlled (server for remote client), see if TargetData was already sent
 		// else register callback for when it does get here.
 		if (!IsLocallyControlled)
 		{
-			if (AbilitySystemComponent->ReplicatedTargetData.IsValid())
-			{
-				ValidData.Broadcast(AbilitySystemComponent->ReplicatedTargetData);
-				EndTask();
-			}
-			else
+			// Register with the TargetData callbacks if we are expecting client to send them
+			if (!CDO->ShouldProduceTargetDataOnServer)
 			{
-				AbilitySystemComponent->ReplicatedTargetDataDelegate.AddUObject(this, &UAbilityTask_WaitTargetData::OnTargetDataReplicatedCallback);
-				AbilitySystemComponent->ReplicatedTargetDataCancelledDelegate.AddDynamic(this, &UAbilityTask_WaitTargetData::OnTargetDataReplicatedCancelledCallback);
+				if (AbilitySystemComponent->ReplicatedTargetData.IsValid())
+				{
+					ValidData.Broadcast(AbilitySystemComponent->ReplicatedTargetData);
+					EndTask();
+				}
+				else
+				{
+					AbilitySystemComponent->ReplicatedTargetDataDelegate.AddUObject(this, &UAbilityTask_WaitTargetData::OnTargetDataReplicatedCallback);
+					AbilitySystemComponent->ReplicatedTargetDataCancelledDelegate.AddDynamic(this, &UAbilityTask_WaitTargetData::OnTargetDataReplicatedCancelledCallback);
+				}
 			}
 		}
 	}
@@ -94,7 +99,7 @@ bool UAbilityTask_WaitTargetData::BeginSpawningActor(UObject* WorldContextObject
 
 void UAbilityTask_WaitTargetData::FinishSpawningActor(UObject* WorldContextObject, AGameplayAbilityTargetActor* SpawnedActor)
 {
-	if (SpawnedActor)
+	if (SpawnedActor && !SpawnedActor->IsPendingKill())
 	{
 		check(MyTargetActor == SpawnedActor);
 
@@ -107,6 +112,37 @@ void UAbilityTask_WaitTargetData::FinishSpawningActor(UObject* WorldContextObjec
 		AbilitySystemComponent->SpawnedTargetActors.Push(SpawnedActor);
 
 		SpawnedActor->StartTargeting(Ability.Get());
+
+		if (SpawnedActor->ShouldProduceTargetData())
+		{
+			// If instant confirm, then stop targeting immediately.
+			// Note this is kind of bad: we should be able to just call a static func on the CDO to do this. 
+			// But then we wouldn't get to set ExposeOnSpawnParameters.
+			if (ConfirmationType == EGameplayTargetingConfirmation::Instant)
+			{
+				SpawnedActor->ConfirmTargeting();
+			}
+			else if (ConfirmationType == EGameplayTargetingConfirmation::UserConfirmed)
+			{
+				// If not locally controlled, check if we already have a confirm/cancel repped to us
+				if (!Ability->GetCurrentActorInfo()->IsLocallyControlled())
+				{
+					if (AbilitySystemComponent->ReplicatedConfirmAbility)
+					{
+						SpawnedActor->ConfirmTargeting();
+						return;
+					}
+					else if (AbilitySystemComponent->ReplicatedCancelAbility)
+					{
+						SpawnedActor->CancelTargeting();
+						return;
+					}
+				}
+
+				// Bind to the Cancel/Confirm Delegates (called from local confirm or from repped confirm)
+				SpawnedActor->BindToConfirmCancelInputs();
+			}
+		}
 	}
 }
 
@@ -152,11 +188,8 @@ void UAbilityTask_WaitTargetData::OnTargetDataReplicatedCancelledCallback()
 void UAbilityTask_WaitTargetData::OnTargetDataReadyCallback(FGameplayAbilityTargetDataHandle Data)
 {
 	check(AbilitySystemComponent.IsValid());
-
-	// Send TargetData to the server IFF we are the client or this is a locally controlled player on the server 
-	// If server confirmed a client's target data locally, it doesn't need to go through the replicated target data code path
-	const FGameplayAbilityActorInfo* Info = Ability->GetCurrentActorInfo();
-	if (!Info->IsNetAuthority() || Info->IsLocallyControlled())
+	
+	if (ShouldReplicateDataToServer())
 	{
 		AbilitySystemComponent->ServerSetReplicatedTargetData(Data);
 	}
@@ -180,6 +213,7 @@ void UAbilityTask_WaitTargetData::OnTargetDataCancelledCallback(FGameplayAbility
 void UAbilityTask_WaitTargetData::OnDestroy(bool AbilityEnded)
 {
 	AbilitySystemComponent->ConsumeAbilityTargetData();
+	AbilitySystemComponent->ConsumeAbilityConfirmCancel();
 	AbilitySystemComponent->SetUserAbilityActivationInhibited(false);
 
 	AbilitySystemComponent->ReplicatedTargetDataDelegate.RemoveUObject(this, &UAbilityTask_WaitTargetData::OnTargetDataReplicatedCallback);
@@ -193,6 +227,18 @@ void UAbilityTask_WaitTargetData::OnDestroy(bool AbilityEnded)
 	Super::OnDestroy(AbilityEnded);
 }
 
+bool UAbilityTask_WaitTargetData::ShouldReplicateDataToServer() const
+{
+	// Send TargetData to the server IFF we are the client and this isn't a GameplayTargetActor that can produce data on the server	
+	const FGameplayAbilityActorInfo* Info = Ability->GetCurrentActorInfo();
+	if (!Info->IsNetAuthority() && !TargetClass->GetDefaultObject<AGameplayAbilityTargetActor>()->ShouldProduceTargetDataOnServer)
+	{
+		return true;
+	}
+
+	return false;
+}
+
 
 // --------------------------------------------------------------------------------------
 
diff --git a/Engine/Source/Runtime/GameplayAbilities/Private/GameplayAbilityTargetTypes.cpp b/Engine/Source/Runtime/GameplayAbilities/Private/GameplayAbilityTargetTypes.cpp
new file mode 100644
--- /dev/null
+++ b/Engine/Source/Runtime/GameplayAbilities/Private/GameplayAbilityTargetTypes.cpp
@@ -0,0 +1,13 @@
+// Copyright 1998-2014 Epic Games, Inc. All Rights Reserved.
+
+#include "AbilitySystemPrivatePCH.h"
+#include "GameplayAbilityTargetTypes.h"
+
+
+bool FGameplayAbilityTargetData_Radius::NetSerialize(FArchive& Ar, class UPackageMap* Map, bool& bOutSuccess)
+{
+	Ar << Actors;	// Fixme: will this go through the package map properly?
+	Ar << Origin;
+
+	return true;
+}
\ No newline at end of file
diff --git a/Engine/Source/Runtime/GameplayAbilities/Private/GameplayEffect.cpp b/Engine/Source/Runtime/GameplayAbilities/Private/GameplayEffect.cpp
--- a/Engine/Source/Runtime/GameplayAbilities/Private/GameplayEffect.cpp
+++ b/Engine/Source/Runtime/GameplayAbilities/Private/GameplayEffect.cpp
@@ -500,6 +500,9 @@ bool FGameplayEffectInstigatorContext::NetSerialize(FArchive& Ar, class UPackage
 		HitResult->Normal.NetSerialize(Ar, Map, bOutSuccess);
 	}
 
+	Ar << HasWorldOrigin;
+	Ar << WorldOrigin;
+
 	bOutSuccess = true;
 	return true;
 }
@@ -514,6 +517,12 @@ bool FGameplayEffectInstigatorContext::IsLocallyControlled() const
 	return false;
 }
 
+void FGameplayEffectInstigatorContext::AddOrigin(FVector InOrigin)
+{
+	HasWorldOrigin = true;
+	WorldOrigin = InOrigin;
+}
+
 
 // --------------------------------------------------------------------------------------------------------------------------------------------------------
 //
diff --git a/Engine/Source/Runtime/GameplayAbilities/Public/Abilities/GameplayAbilityTargetActor.h b/Engine/Source/Runtime/GameplayAbilities/Public/Abilities/GameplayAbilityTargetActor.h
--- a/Engine/Source/Runtime/GameplayAbilities/Public/Abilities/GameplayAbilityTargetActor.h
+++ b/Engine/Source/Runtime/GameplayAbilities/Public/Abilities/GameplayAbilityTargetActor.h
@@ -16,6 +16,10 @@ public:
 	UPROPERTY()
 	bool StaticTargetFunction;
 
+	/** The TargetData this class produces can be entirely generated on the server. We don't require the client to send us full or partial TargetData (possibly just a 'confirm') */
+	UPROPERTY()
+	bool ShouldProduceTargetDataOnServer;
+
 	/** Describes where the targeting action starts, usually the player character or a socket on the player character. */
 	//UPROPERTY(BlueprintReadOnly, meta=(ExposeOnSpawn=true), Category=Targeting)
 	UPROPERTY(BlueprintReadOnly, meta = (ExposeOnSpawn = true), Replicated, Category = Targeting)
@@ -34,6 +38,8 @@ public:
 	UFUNCTION()
 	virtual void CancelTargeting();
 
+	virtual void BindToConfirmCancelInputs();
+
 	virtual bool ShouldProduceTargetData() const;
 
 	/** Replicated target data was received from a client. Possibly sanitize/verify. return true if data is good and we should broadcast it as valid data. */
@@ -51,4 +57,8 @@ public:
 
 	UPROPERTY(BlueprintReadOnly, Category="Targeting")
 	APlayerController* MasterPC;
+
+	UPROPERTY()
+	UGameplayAbility* OwningAbility;
+
 };
\ No newline at end of file
diff --git a/Engine/Source/Runtime/GameplayAbilities/Public/Abilities/GameplayAbilityTargetActor_Radius.h b/Engine/Source/Runtime/GameplayAbilities/Public/Abilities/GameplayAbilityTargetActor_Radius.h
new file mode 100644
--- /dev/null
+++ b/Engine/Source/Runtime/GameplayAbilities/Public/Abilities/GameplayAbilityTargetActor_Radius.h
@@ -0,0 +1,32 @@
+// Copyright 1998-2014 Epic Games, Inc. All Rights Reserved.
+
+#pragma once
+
+#include "GameplayAbilityTargetActor.h"
+#include "GameplayAbilityWorldReticle.h"
+#include "GameplayAbilityTargetActor_Radius.generated.h"
+
+UCLASS(Blueprintable)
+class GAMEPLAYABILITIES_API AGameplayAbilityTargetActor_Radius : public AGameplayAbilityTargetActor
+{
+	GENERATED_UCLASS_BODY()
+
+public:
+
+	virtual void StartTargeting(UGameplayAbility* Ability);
+	
+	virtual void ConfirmTargeting() override;
+
+	UPROPERTY(BlueprintReadWrite, EditAnywhere, meta = (ExposeOnSpawn = true), Category = Radius)
+	float Radius;
+
+	UPROPERTY(BlueprintReadOnly, Replicated, Category = Targeting)
+	AActor* SourceActor;
+
+protected:
+
+	TArray<AActor*>	PerformOverlap(const FVector& Origin);
+
+	FGameplayAbilityTargetDataHandle MakeTargetData(const TArray<AActor*> Actors, const FVector& Origin) const;
+
+};
\ No newline at end of file
diff --git a/Engine/Source/Runtime/GameplayAbilities/Public/Abilities/GameplayAbilityTargetActor_SingleLineTrace.h b/Engine/Source/Runtime/GameplayAbilities/Public/Abilities/GameplayAbilityTargetActor_SingleLineTrace.h
--- a/Engine/Source/Runtime/GameplayAbilities/Public/Abilities/GameplayAbilityTargetActor_SingleLineTrace.h
+++ b/Engine/Source/Runtime/GameplayAbilities/Public/Abilities/GameplayAbilityTargetActor_SingleLineTrace.h
@@ -16,23 +16,18 @@ public:
 
 	virtual FGameplayAbilityTargetDataHandle StaticGetTargetData(UWorld * World, const FGameplayAbilityActorInfo* ActorInfo, FGameplayAbilityActivationInfo ActivationInfo) const override;
 
-	virtual void StartTargeting(UGameplayAbility* Ability);
+	virtual void StartTargeting(UGameplayAbility* Ability) override;
 	
 	virtual void ConfirmTargeting() override;	
 
 	virtual void Tick(float DeltaSeconds) override;
 	
-	TWeakObjectPtr<UGameplayAbility> Ability;
-	
 	UPROPERTY(BlueprintReadOnly, Replicated, Category = Targeting)
 	bool bDebug;
 
 	UPROPERTY(BlueprintReadOnly, Replicated, Category = Targeting)
 	AActor* SourceActor;
 
-	UPROPERTY(BlueprintReadOnly, meta=(ExposeOnSpawn=true), Category=Projectile)
-	bool bBindToConfirmCancelInputs;
-
 	UPROPERTY(BlueprintReadWrite, EditAnywhere, meta = (ExposeOnSpawn = true), Category = Projectile)
 	float MaxRange;
 
diff --git a/Engine/Source/Runtime/GameplayAbilities/Public/Abilities/GameplayAbilityTargetTypes.h b/Engine/Source/Runtime/GameplayAbilities/Public/Abilities/GameplayAbilityTargetTypes.h
new file mode 100644
--- /dev/null
+++ b/Engine/Source/Runtime/GameplayAbilities/Public/Abilities/GameplayAbilityTargetTypes.h
@@ -0,0 +1,64 @@
+// Copyright 1998-2014 Epic Games, Inc. All Rights Reserved.
+
+#pragma once
+
+#include "GameplayAbilityTargetTypes.generated.h"
+
+class UGameplayEffect;
+class UAnimInstance;
+class UAbilitySystemComponent;
+class UGameplayAbility;
+class AGameplayAbilityTargetActor;
+class UAbilityTask;
+class UAttributeSet;
+
+UENUM(BlueprintType)
+namespace EGameplayTargetingConfirmation
+{
+	enum Type
+	{
+		Instant,			// The targeting happens instantly without special logic or user input deciding when to 'fire'.
+		UserConfirmed,		// The targeting happens when the user confirms the targeting.
+		Custom,				// The GameplayTargeting Ability is responsible for deciding when the targeting data is ready. Not supported by all TargetingActors.
+	};
+}
+
+
+USTRUCT(BlueprintType)
+struct GAMEPLAYABILITIES_API FGameplayAbilityTargetData_Radius : public FGameplayAbilityTargetData
+{
+	GENERATED_USTRUCT_BODY()
+
+	FGameplayAbilityTargetData_Radius()
+	: Origin(0.f) { }
+
+	FGameplayAbilityTargetData_Radius(const TArray<AActor*>& InActors, const FVector& InOrigin)
+	: Actors(InActors), Origin(InOrigin) { }
+
+	virtual TArray<AActor*>	GetActors() const { return Actors; }
+
+	virtual bool HasOrigin() const { return true; }
+
+	virtual FTransform GetOrigin() const { return FTransform(Origin); }
+
+	bool NetSerialize(FArchive& Ar, class UPackageMap* Map, bool& bOutSuccess);
+
+	virtual UScriptStruct* GetScriptStruct()
+	{
+		return FGameplayAbilityTargetData_Radius::StaticStruct();
+	}
+
+private:
+
+	TArray<AActor*> Actors;
+	FVector Origin;
+};
+
+template<>
+struct TStructOpsTypeTraits<FGameplayAbilityTargetData_Radius> : public TStructOpsTypeTraitsBase
+{
+	enum
+	{
+		WithNetSerializer = true	// For now this is REQUIRED for FGameplayAbilityTargetDataHandle net serialization to work
+	};
+};
\ No newline at end of file
diff --git a/Engine/Source/Runtime/GameplayAbilities/Public/Abilities/Tasks/AbilityTask_WaitTargetData.h b/Engine/Source/Runtime/GameplayAbilities/Public/Abilities/Tasks/AbilityTask_WaitTargetData.h
--- a/Engine/Source/Runtime/GameplayAbilities/Public/Abilities/Tasks/AbilityTask_WaitTargetData.h
+++ b/Engine/Source/Runtime/GameplayAbilities/Public/Abilities/Tasks/AbilityTask_WaitTargetData.h
@@ -2,6 +2,7 @@
 #pragma once
 
 #include "AbilityTask.h"
+#include "GameplayAbilityTargetTypes.h"
 #include "AbilityTask_WaitTargetData.generated.h"
 
 class AGameplayAbilityTargetActor;
@@ -33,7 +34,7 @@ class UAbilityTask_WaitTargetData: public UAbilityTask
 
 	/** Spawns Targeting actor and waits for it to return valid data or to be cancelled. */
 	UFUNCTION(BlueprintCallable, meta=(HidePin = "WorldContextObject", DefaultToSelf = "WorldContextObject", BlueprintInternalUseOnly = "true", HideSpawnParms="Instigator"), Category="Ability|Tasks")
-	static UAbilityTask_WaitTargetData* WaitTargetData(UObject* WorldContextObject, TSubclassOf<AGameplayAbilityTargetActor> Class);
+	static UAbilityTask_WaitTargetData* WaitTargetData(UObject* WorldContextObject, TEnumAsByte<EGameplayTargetingConfirmation::Type> ConfirmationType, TSubclassOf<AGameplayAbilityTargetActor> Class);
 
 	UFUNCTION(BlueprintCallable, meta = (HidePin = "WorldContextObject", DefaultToSelf = "WorldContextObject", BlueprintInternalUseOnly = "true"), Category = "Abilities")
 	bool BeginSpawningActor(UObject* WorldContextObject, TSubclassOf<AGameplayAbilityTargetActor> Class, AGameplayAbilityTargetActor*& SpawnedActor);
@@ -45,10 +46,14 @@ protected:
 
 	virtual void OnDestroy(bool AbilityEnded) override;
 
+	bool ShouldReplicateDataToServer() const;
+
 	TSubclassOf<AGameplayAbilityTargetActor> TargetClass;
 
 	/** The TargetActor that we spawned, or the class CDO if this is a static targeting task */
 	TWeakObjectPtr<AGameplayAbilityTargetActor>	MyTargetActor;
+
+	TEnumAsByte<EGameplayTargetingConfirmation::Type> ConfirmationType;
 };
 
 
diff --git a/Engine/Source/Runtime/GameplayAbilities/Public/GameplayEffectTypes.h b/Engine/Source/Runtime/GameplayAbilities/Public/GameplayEffectTypes.h
--- a/Engine/Source/Runtime/GameplayAbilities/Public/GameplayEffectTypes.h
+++ b/Engine/Source/Runtime/GameplayAbilities/Public/GameplayEffectTypes.h
@@ -145,6 +145,7 @@ struct FGameplayEffectInstigatorContext
 	FGameplayEffectInstigatorContext()
 		: Instigator(NULL)
 		, InstigatorAbilitySystemComponent(NULL)
+		, HasWorldOrigin(false)
 	{
 	}
 
@@ -161,18 +162,20 @@ struct FGameplayEffectInstigatorContext
 
 	void AddHitResult(const FHitResult InHitResult);
 
+	void AddOrigin(FVector InOrigin);
+
 	FString ToString() const
 	{
 		return Instigator ? Instigator->GetName() : FString(TEXT("NONE"));
 	}
 
 	/** Should always return the original instigator that started the whole chain */
-	AActor * GetOriginalInstigator()
+	AActor* GetOriginalInstigator()
 	{
 		return Instigator;
 	}
 
-	UAbilitySystemComponent * GetOriginalInstigatorAbilitySystemComponent() const
+	UAbilitySystemComponent* GetOriginalInstigatorAbilitySystemComponent() const
 	{
 		return InstigatorAbilitySystemComponent;
 	}
@@ -184,11 +187,16 @@ struct FGameplayEffectInstigatorContext
 	AActor* Instigator;
 
 	UPROPERTY(NotReplicated)
-	UAbilitySystemComponent *InstigatorAbilitySystemComponent;
+	UAbilitySystemComponent* InstigatorAbilitySystemComponent;
 
 	/** Trace information - may be NULL in many cases */
 	TSharedPtr<FHitResult>	HitResult;
 
+	UPROPERTY()
+	FVector	WorldOrigin;
+
+	bool HasWorldOrigin;
+
 	bool NetSerialize(FArchive& Ar, class UPackageMap* Map, bool& bOutSuccess);
 };
 
