- game mode fixes
- ability work, set it up so we can cut the cord to a dead pawn
- added first pass on spectator/respawning

diff --git a/Engine/Source/Runtime/GameplayAbilities/Private/Abilities/GameplayAbility.cpp b/Engine/Source/Runtime/GameplayAbilities/Private/Abilities/GameplayAbility.cpp
--- a/Engine/Source/Runtime/GameplayAbilities/Private/Abilities/GameplayAbility.cpp
+++ b/Engine/Source/Runtime/GameplayAbilities/Private/Abilities/GameplayAbility.cpp
@@ -279,6 +279,12 @@ void UGameplayAbility::CommitExecute(const FGameplayAbilityActorInfo* ActorInfo,
 
 bool UGameplayAbility::TryActivateAbility(const FGameplayAbilityActorInfo* ActorInfo, uint32 PrevPredictionKey, uint32 CurrPredictionKey, UGameplayAbility ** OutInstancedAbility)
 {	
+	// make sure the ActorInfo and then Actor on that FGameplayAbilityActorInfo are valid, if not bail out.
+	if (ActorInfo == NULL || ActorInfo->Actor == NULL)
+	{
+		return false;
+	}
+
 	// This should only come from button presses/local instigation
 	ENetRole NetMode = ActorInfo->Actor->Role;
 	ensure(NetMode != ROLE_SimulatedProxy);
@@ -659,6 +665,14 @@ void FGameplayAbilityActorInfo::InitFromActor(AActor *InActor, UAbilitySystemCom
 	MovementComponent = InActor->FindComponentByClass<UMovementComponent>();
 }
 
+void FGameplayAbilityActorInfo::ClearActorInfo()
+{
+	Actor = NULL;
+	PlayerController = NULL;
+	AnimInstance = NULL;
+	MovementComponent = NULL;
+}
+
 bool FGameplayAbilityActorInfo::IsLocallyControlled() const
 {
 	if (PlayerController.IsValid())
diff --git a/Engine/Source/Runtime/GameplayAbilities/Private/AbilitySystemComponent.cpp b/Engine/Source/Runtime/GameplayAbilities/Private/AbilitySystemComponent.cpp
--- a/Engine/Source/Runtime/GameplayAbilities/Private/AbilitySystemComponent.cpp
+++ b/Engine/Source/Runtime/GameplayAbilities/Private/AbilitySystemComponent.cpp
@@ -617,6 +617,8 @@ void UAbilitySystemComponent::GetLifetimeReplicatedProps(TArray< FLifetimeProper
 	DOREPLIFETIME(UAbilitySystemComponent, ReplicatedInstancedAbilities);
 	DOREPLIFETIME(UAbilitySystemComponent, ActivatableAbilities);
 
+	DOREPLIFETIME(UAbilitySystemComponent, AbilityActor);
+
 	DOREPLIFETIME(UAbilitySystemComponent, ReplicatedPredictionKey);
 }
 
diff --git a/Engine/Source/Runtime/GameplayAbilities/Private/AbilitySystemComponent_Abilities.cpp b/Engine/Source/Runtime/GameplayAbilities/Private/AbilitySystemComponent_Abilities.cpp
--- a/Engine/Source/Runtime/GameplayAbilities/Private/AbilitySystemComponent_Abilities.cpp
+++ b/Engine/Source/Runtime/GameplayAbilities/Private/AbilitySystemComponent_Abilities.cpp
@@ -82,6 +82,31 @@ void UAbilitySystemComponent::InitAbilityActorInfo(AActor* AvatarActor)
 {
 	check(AbilityActorInfo.IsValid());
 	AbilityActorInfo->InitFromActor(AvatarActor, this);
+	AbilityActor = AvatarActor;
+}
+
+void UAbilitySystemComponent::ClearActorInfo()
+{
+	check(AbilityActorInfo.IsValid());
+	AbilityActorInfo->ClearActorInfo();
+	AbilityActor = NULL;
+}
+
+void UAbilitySystemComponent::OnRep_OwningActor()
+{
+	check(AbilityActorInfo.IsValid());
+
+	if (AbilityActor != AbilityActorInfo->Actor)
+	{
+		if (AbilityActor != NULL)
+		{
+			InitAbilityActorInfo(AbilityActor);
+		}
+		else
+		{
+			ClearActorInfo();
+		}
+	}
 }
 
 void UAbilitySystemComponent::RefreshAbilityActorInfo()
diff --git a/Engine/Source/Runtime/GameplayAbilities/Public/Abilities/GameplayAbilityTypes.h b/Engine/Source/Runtime/GameplayAbilities/Public/Abilities/GameplayAbilityTypes.h
--- a/Engine/Source/Runtime/GameplayAbilities/Public/Abilities/GameplayAbilityTypes.h
+++ b/Engine/Source/Runtime/GameplayAbilities/Public/Abilities/GameplayAbilityTypes.h
@@ -193,6 +193,7 @@ struct GAMEPLAYABILITIES_API FGameplayAbilityActorInfo
 	bool IsNetAuthority() const;
 
 	virtual void InitFromActor(AActor *Actor, UAbilitySystemComponent* InAbilitySystemComponent);
+	virtual void ClearActorInfo();
 };
 
 /**
diff --git a/Engine/Source/Runtime/GameplayAbilities/Public/AbilitySystemComponent.h b/Engine/Source/Runtime/GameplayAbilities/Public/AbilitySystemComponent.h
--- a/Engine/Source/Runtime/GameplayAbilities/Public/AbilitySystemComponent.h
+++ b/Engine/Source/Runtime/GameplayAbilities/Public/AbilitySystemComponent.h
@@ -493,6 +493,12 @@ class GAMEPLAYABILITIES_API UAbilitySystemComponent : public UActorComponent, pu
 
 	// -----------------------------------------------------------------------------
 
+	UPROPERTY(ReplicatedUsing=OnRep_OwningActor)
+	AActor* AbilityActor;
+	
+	UFUNCTION()
+	void OnRep_OwningActor();
+
 	/** Cached off data about the owning actor that abilities will need to frequently access (movement component, mesh component, anim instance, etc) */
 	TSharedPtr<FGameplayAbilityActorInfo>	AbilityActorInfo;
 
@@ -503,6 +509,11 @@ class GAMEPLAYABILITIES_API UAbilitySystemComponent : public UActorComponent, pu
 	 */
 	void InitAbilityActorInfo(AActor* AvatarActor);
 
+	/**
+	* This is called when the actor that is initialized to this system dies, this will clear that actor from this system and FGameplayAbilityActorInfo
+	*/
+	void ClearActorInfo();
+
 	/**
 	 *	This will refresh the Ability's ActorInfo structure based on the current ActorInfo. That is, AvatarActor will be the same but we will look for new
 	 *	AnimInstance, MovementComponent, PlayerController, etc.
diff --git a/Engine/Source/Runtime/GameplayAbilities/Public/AbilitySystemGlobals.h b/Engine/Source/Runtime/GameplayAbilities/Public/AbilitySystemGlobals.h
--- a/Engine/Source/Runtime/GameplayAbilities/Public/AbilitySystemGlobals.h
+++ b/Engine/Source/Runtime/GameplayAbilities/Public/AbilitySystemGlobals.h
@@ -38,6 +38,11 @@ class GAMEPLAYABILITIES_API UAbilitySystemGlobals : public UObject
 
 	UDataTable* GetGlobalAttributeMetaDataTable();
 
+	bool IsAbilitySystemGlobalsInitialized()
+	{
+		return GlobalAttributeSetInitter.IsValid();
+	}
+
 	void AutomationTestOnly_SetGlobalCurveTable(class UCurveTable *InTable)
 	{
 		GlobalCurveTable = InTable;
