Call ConfirmActivateSucceed on non instanced abilities

diff --git a/Engine/Source/Runtime/GameplayAbilities/Private/AbilitySystemComponent_Abilities.cpp b/Engine/Source/Runtime/GameplayAbilities/Private/AbilitySystemComponent_Abilities.cpp
--- a/Engine/Source/Runtime/GameplayAbilities/Private/AbilitySystemComponent_Abilities.cpp
+++ b/Engine/Source/Runtime/GameplayAbilities/Private/AbilitySystemComponent_Abilities.cpp
@@ -421,21 +421,28 @@ void UAbilitySystemComponent::ClientActivateAbilitySucceed_Implementation(UGamep
 
 	if (AbilityToActivate->NetExecutionPolicy == EGameplayAbilityNetExecutionPolicy::Predictive)
 	{
-		// Find the one we predictively spawned, tell them we are confirmed
-		bool found = false;
-		for (UGameplayAbility* LocalAbility : NonReplicatedInstancedAbilities)				// Fixme: this has to be updated once predictive abilities can replicate
+		if (AbilityToActivate->GetInstancingPolicy() == EGameplayAbilityInstancingPolicy::NonInstanced)
 		{
-			if (LocalAbility->GetCurrentActivationInfo().CurrPredictionKey == PredictionKey)
+			AbilityToActivate->ConfirmActivateSucceed();
+		}
+		else
+		{
+			// Find the one we predictively spawned, tell them we are confirmed
+			bool found = false;
+			for (UGameplayAbility* LocalAbility : NonReplicatedInstancedAbilities)				// Fixme: this has to be updated once predictive abilities can replicate
 			{
-				LocalAbility->ConfirmActivateSucceed();
-				found = true;
-				break;
+				if (LocalAbility->GetCurrentActivationInfo().CurrPredictionKey == PredictionKey)
+				{
+					LocalAbility->ConfirmActivateSucceed();
+					found = true;
+					break;
+				}
 			}
-		}
 
-		if (!found)
-		{
-			ABILITY_LOG(Warning, TEXT("Ability %s was confirmed by server but no longer exists on client (replication key: %d"), *AbilityToActivate->GetName(), PredictionKey);
+			if (!found)
+			{
+				ABILITY_LOG(Warning, TEXT("Ability %s was confirmed by server but no longer exists on client (replication key: %d"), *AbilityToActivate->GetName(), PredictionKey);
+			}
 		}
 	}
 	else
