non unity fix

diff --git a/Engine/Source/Runtime/GameplayAbilities/Private/Abilities/GameplayAbilityTargetActor.cpp b/Engine/Source/Runtime/GameplayAbilities/Private/Abilities/GameplayAbilityTargetActor.cpp
--- a/Engine/Source/Runtime/GameplayAbilities/Private/Abilities/GameplayAbilityTargetActor.cpp
+++ b/Engine/Source/Runtime/GameplayAbilities/Private/Abilities/GameplayAbilityTargetActor.cpp
@@ -2,6 +2,7 @@
 
 #include "AbilitySystemPrivatePCH.h"
 #include "GameplayAbilityTargetActor.h"
+#include "AbilitySystemComponent.h"
 #include "Net/UnrealNetwork.h"
 
 // --------------------------------------------------------------------------------------------------------------------------------------------------------
