Speculative fix for HTML5 build failure.

diff --git a/Engine/Source/Runtime/GameplayAbilities/Private/Abilities/GameplayAbilityTargetActor_Radius.cpp b/Engine/Source/Runtime/GameplayAbilities/Private/Abilities/GameplayAbilityTargetActor_Radius.cpp
--- a/Engine/Source/Runtime/GameplayAbilities/Private/Abilities/GameplayAbilityTargetActor_Radius.cpp
+++ b/Engine/Source/Runtime/GameplayAbilities/Private/Abilities/GameplayAbilityTargetActor_Radius.cpp
@@ -5,6 +5,7 @@
 #include "GameplayAbilityTargetActor_Radius.h"
 #include "Engine/World.h"
 #include "Runtime/Engine/Public/Net/UnrealNetwork.h"
+#include "GameplayAbilityTargetTypes.h
 
 // --------------------------------------------------------------------------------------------------------------------------------------------------------
 //
