GUBP byte-alignment fix.

diff --git a/Engine/Source/Runtime/GameplayAbilities/Private/AbilitySystemBlueprintLibrary.cpp b/Engine/Source/Runtime/GameplayAbilities/Private/AbilitySystemBlueprintLibrary.cpp
--- a/Engine/Source/Runtime/GameplayAbilities/Private/AbilitySystemBlueprintLibrary.cpp
+++ b/Engine/Source/Runtime/GameplayAbilities/Private/AbilitySystemBlueprintLibrary.cpp
@@ -28,7 +28,7 @@ void UAbilitySystemBlueprintLibrary::ApplyGameplayEffectToTargetData(FGameplayAb
 	}
 }
 
-FGameplayAbilityTargetDataHandle UAbilitySystemBlueprintLibrary::AbilityTargetDataFromLocations(FGameplayAbilityTargetingLocationInfo SourceLocation, FGameplayAbilityTargetingLocationInfo TargetLocation)
+FGameplayAbilityTargetDataHandle UAbilitySystemBlueprintLibrary::AbilityTargetDataFromLocations(const FGameplayAbilityTargetingLocationInfo& SourceLocation, const FGameplayAbilityTargetingLocationInfo& TargetLocation)
 {
 	// Construct TargetData
 	FGameplayAbilityTargetData_LocationInfo*	NewData = new FGameplayAbilityTargetData_LocationInfo();
diff --git a/Engine/Source/Runtime/GameplayAbilities/Public/AbilitySystemBlueprintLibrary.h b/Engine/Source/Runtime/GameplayAbilities/Public/AbilitySystemBlueprintLibrary.h
--- a/Engine/Source/Runtime/GameplayAbilities/Public/AbilitySystemBlueprintLibrary.h
+++ b/Engine/Source/Runtime/GameplayAbilities/Public/AbilitySystemBlueprintLibrary.h
@@ -29,7 +29,7 @@ class UAbilitySystemBlueprintLibrary : public UBlueprintFunctionLibrary
 	static void ApplyGameplayEffectToTargetData(FGameplayAbilityTargetDataHandle Target, UGameplayEffect *GameplayEffect, const FGameplayAbilityActorInfo InstigatorInfo);
 
 	UFUNCTION(BlueprintPure, Category = "Ability|TargetData")
-	static FGameplayAbilityTargetDataHandle	AbilityTargetDataFromLocations(FGameplayAbilityTargetingLocationInfo SourceLocation, FGameplayAbilityTargetingLocationInfo TargetLocation);
+	static FGameplayAbilityTargetDataHandle	AbilityTargetDataFromLocations(const FGameplayAbilityTargetingLocationInfo& SourceLocation, const FGameplayAbilityTargetingLocationInfo& TargetLocation);
 
 	UFUNCTION(BlueprintPure, Category = "Ability|TargetData")
 	static FGameplayAbilityTargetDataHandle	AbilityTargetDataHandleFromAbilityTargetDataMesh(FGameplayAbilityTargetData_Mesh Data);
