Added one-node functionality to create actor-based or socket-based targeting location info.

diff --git a/Engine/Source/Runtime/GameplayAbilities/Private/Abilities/GameplayAbility.cpp b/Engine/Source/Runtime/GameplayAbilities/Private/Abilities/GameplayAbility.cpp
--- a/Engine/Source/Runtime/GameplayAbilities/Private/Abilities/GameplayAbility.cpp
+++ b/Engine/Source/Runtime/GameplayAbilities/Private/Abilities/GameplayAbility.cpp
@@ -509,14 +509,14 @@ float UGameplayAbility::GetCooldownTimeRemaining(const FGameplayAbilityActorInfo
 	return 0.f;
 }
 
-FGameplayAbilityActorInfo UGameplayAbility::GetActorInfo()
+FGameplayAbilityActorInfo UGameplayAbility::GetActorInfo() const
 {
 	check(CurrentActorInfo);
 	return *CurrentActorInfo;
 }
 
 /** Convenience method for abilities to get skeletal mesh component - useful for aiming abilities */
-USkeletalMeshComponent* UGameplayAbility::GetOwningComponentFromActorInfo()
+USkeletalMeshComponent* UGameplayAbility::GetOwningComponentFromActorInfo() const
 {
 	check(CurrentActorInfo);
 	if (CurrentActorInfo->AnimInstance.IsValid())
@@ -586,6 +586,23 @@ void UGameplayAbility::MontageBranchPoint_AbilityDecisionStart(const FGameplayAb
 
 }
 
+FGameplayAbilityTargetingLocationInfo UGameplayAbility::MakeTargetLocationInfoFromOwnerActor() const
+{
+	FGameplayAbilityTargetingLocationInfo ReturnLocation;
+	ReturnLocation.LocationType = EGameplayAbilityTargetingLocationType::ActorTransform;
+	ReturnLocation.SourceActor = GetActorInfo().Actor.Get();
+	return ReturnLocation;
+}
+
+FGameplayAbilityTargetingLocationInfo UGameplayAbility::MakeTargetLocationInfoFromOwnerSkeletalMeshComponent(FName SocketName) const
+{
+	FGameplayAbilityTargetingLocationInfo ReturnLocation;
+	ReturnLocation.LocationType = EGameplayAbilityTargetingLocationType::ActorTransform;
+	ReturnLocation.SourceComponent = GetActorInfo().AnimInstance.IsValid() ? GetActorInfo().AnimInstance.Get()->GetOwningComponent() : NULL;
+	ReturnLocation.SourceSocketName = SocketName;
+	return ReturnLocation;
+}
+
 void UGameplayAbility::ClientActivateAbilitySucceed_Implementation(int32 PredictionKey)
 {
 	PostNetInit();
diff --git a/Engine/Source/Runtime/GameplayAbilities/Public/Abilities/GameplayAbility.h b/Engine/Source/Runtime/GameplayAbilities/Public/Abilities/GameplayAbility.h
--- a/Engine/Source/Runtime/GameplayAbilities/Public/Abilities/GameplayAbility.h
+++ b/Engine/Source/Runtime/GameplayAbilities/Public/Abilities/GameplayAbility.h
@@ -332,11 +332,11 @@ protected:
 		}
 	}
 
-	UFUNCTION(BlueprintPure, Category=Ability)
-	FGameplayAbilityActorInfo GetActorInfo();
+	UFUNCTION(BlueprintCallable, Category=Ability)
+	FGameplayAbilityActorInfo GetActorInfo() const;
 
-	UFUNCTION(BlueprintPure, Category = Ability)
-	USkeletalMeshComponent* GetOwningComponentFromActorInfo();
+	UFUNCTION(BlueprintCallable, Category = Ability)
+	USkeletalMeshComponent* GetOwningComponentFromActorInfo() const;
 
 	UFUNCTION(BlueprintCallable, Category=Ability)
 	FGameplayEffectSpecHandle GetOutgoingSpec(UGameplayEffect* GameplayEffect) const;
@@ -412,6 +412,12 @@ protected:
 	void MontageBranchPoint_AbilityDecisionStop(const FGameplayAbilityActorInfo* ActorInfo) const;
 
 	void MontageBranchPoint_AbilityDecisionStart(const FGameplayAbilityActorInfo* ActorInfo) const;
+
+	UFUNCTION(BlueprintCallable, Category=Ability)
+	FGameplayAbilityTargetingLocationInfo MakeTargetLocationInfoFromOwnerActor() const;
+
+	UFUNCTION(BlueprintCallable, Category = Ability)
+	FGameplayAbilityTargetingLocationInfo MakeTargetLocationInfoFromOwnerSkeletalMeshComponent(FName SocketName) const;
 private:
 
 	/** 
