Accidentally hit submit instead of save - temporarily rolling back.

diff --git a/Engine/Source/Runtime/GameplayAbilities/Public/Abilities/GameplayAbility.h b/Engine/Source/Runtime/GameplayAbilities/Public/Abilities/GameplayAbility.h
--- a/Engine/Source/Runtime/GameplayAbilities/Public/Abilities/GameplayAbility.h
+++ b/Engine/Source/Runtime/GameplayAbilities/Public/Abilities/GameplayAbility.h
@@ -332,10 +332,10 @@ protected:
 		}
 	}
 
-	UFUNCTION(BlueprintCallable, Category=Ability)
+	//UFUNCTION(BlueprintCallable, Category=Ability)
 	FGameplayAbilityActorInfo GetActorInfo() const;
 
-	UFUNCTION(BlueprintCallable, Category = Ability)
+	//UFUNCTION(BlueprintCallable, Category = Ability)
 	USkeletalMeshComponent* GetOwningComponentFromActorInfo() const;
 
 	UFUNCTION(BlueprintCallable, Category=Ability)
