Rolling back.

diff --git a/Engine/Source/Runtime/GameplayAbilities/Public/Abilities/GameplayAbility.h b/Engine/Source/Runtime/GameplayAbilities/Public/Abilities/GameplayAbility.h
--- a/Engine/Source/Runtime/GameplayAbilities/Public/Abilities/GameplayAbility.h
+++ b/Engine/Source/Runtime/GameplayAbilities/Public/Abilities/GameplayAbility.h
@@ -332,10 +332,10 @@ protected:
 		}
 	}
 
-	//UFUNCTION(BlueprintCallable, Category=Ability)
+	UFUNCTION(BlueprintCallable, Category=Ability)
 	FGameplayAbilityActorInfo GetActorInfo() const;
 
-	//UFUNCTION(BlueprintCallable, Category = Ability)
+	UFUNCTION(BlueprintCallable, Category = Ability)
 	USkeletalMeshComponent* GetOwningComponentFromActorInfo() const;
 
 	UFUNCTION(BlueprintCallable, Category=Ability)
@@ -413,10 +413,10 @@ protected:
 
 	void MontageBranchPoint_AbilityDecisionStart(const FGameplayAbilityActorInfo* ActorInfo) const;
 
-	UFUNCTION(BlueprintCallable, Category=Ability)
+	//UFUNCTION(BlueprintCallable, Category=Ability)
 	FGameplayAbilityTargetingLocationInfo MakeTargetLocationInfoFromOwnerActor() const;
 
-	UFUNCTION(BlueprintCallable, Category = Ability)
+	//UFUNCTION(BlueprintCallable, Category = Ability)
 	FGameplayAbilityTargetingLocationInfo MakeTargetLocationInfoFromOwnerSkeletalMeshComponent(FName SocketName) const;
 private:
 
