Fixed one-node functionality to create actor-based or socket-based targeting location info.

diff --git a/Engine/Source/Runtime/GameplayAbilities/Private/Abilities/GameplayAbility.cpp b/Engine/Source/Runtime/GameplayAbilities/Private/Abilities/GameplayAbility.cpp
--- a/Engine/Source/Runtime/GameplayAbilities/Private/Abilities/GameplayAbility.cpp
+++ b/Engine/Source/Runtime/GameplayAbilities/Private/Abilities/GameplayAbility.cpp
@@ -7,7 +7,7 @@
 #include "AbilitySystemComponent.h"
 #include "AbilitySystemBlueprintLibrary.h"
 #include "AbilityTask.h"
-
+#pragma optimize("",off)
 // --------------------------------------------------------------------------------------------------------------------------------------------------------
 //
 //	UGameplayAbility
@@ -597,7 +597,7 @@ FGameplayAbilityTargetingLocationInfo UGameplayAbility::MakeTargetLocationInfoFr
 FGameplayAbilityTargetingLocationInfo UGameplayAbility::MakeTargetLocationInfoFromOwnerSkeletalMeshComponent(FName SocketName) const
 {
 	FGameplayAbilityTargetingLocationInfo ReturnLocation;
-	ReturnLocation.LocationType = EGameplayAbilityTargetingLocationType::ActorTransform;
+	ReturnLocation.LocationType = EGameplayAbilityTargetingLocationType::SocketTransform;
 	ReturnLocation.SourceComponent = GetActorInfo().AnimInstance.IsValid() ? GetActorInfo().AnimInstance.Get()->GetOwningComponent() : NULL;
 	ReturnLocation.SourceSocketName = SocketName;
 	return ReturnLocation;
diff --git a/Engine/Source/Runtime/GameplayAbilities/Public/Abilities/GameplayAbility.h b/Engine/Source/Runtime/GameplayAbilities/Public/Abilities/GameplayAbility.h
--- a/Engine/Source/Runtime/GameplayAbilities/Public/Abilities/GameplayAbility.h
+++ b/Engine/Source/Runtime/GameplayAbilities/Public/Abilities/GameplayAbility.h
@@ -413,10 +413,10 @@ protected:
 
 	void MontageBranchPoint_AbilityDecisionStart(const FGameplayAbilityActorInfo* ActorInfo) const;
 
-	//UFUNCTION(BlueprintCallable, Category=Ability)
+	UFUNCTION(BlueprintCallable, Category=Ability)
 	FGameplayAbilityTargetingLocationInfo MakeTargetLocationInfoFromOwnerActor() const;
 
-	//UFUNCTION(BlueprintCallable, Category = Ability)
+	UFUNCTION(BlueprintCallable, Category = Ability)
 	FGameplayAbilityTargetingLocationInfo MakeTargetLocationInfoFromOwnerSkeletalMeshComponent(FName SocketName) const;
 private:
 
