Improvements to multiple target selection and single-line trace targeting schemes.

diff --git a/Engine/Source/Runtime/GameplayAbilities/Private/Abilities/GameplayAbilityTargetActor_SingleLineTrace.cpp b/Engine/Source/Runtime/GameplayAbilities/Private/Abilities/GameplayAbilityTargetActor_SingleLineTrace.cpp
--- a/Engine/Source/Runtime/GameplayAbilities/Private/Abilities/GameplayAbilityTargetActor_SingleLineTrace.cpp
+++ b/Engine/Source/Runtime/GameplayAbilities/Private/Abilities/GameplayAbilityTargetActor_SingleLineTrace.cpp
@@ -18,7 +18,9 @@ AGameplayAbilityTargetActor_SingleLineTrace::AGameplayAbilityTargetActor_SingleL
 	PrimaryActorTick.bCanEverTick = true;
 	PrimaryActorTick.TickGroup = TG_PrePhysics;
 	StaticTargetFunction = false;
+
 	bDebug = false;
+	MaxRange = 999999.0f;
 }
 
 void AGameplayAbilityTargetActor_SingleLineTrace::GetLifetimeReplicatedProps(TArray< FLifetimeProperty > & OutLifetimeProps) const
