Removing a leftover pragma.

diff --git a/Engine/Source/Runtime/GameplayAbilities/Private/Abilities/GameplayAbility.cpp b/Engine/Source/Runtime/GameplayAbilities/Private/Abilities/GameplayAbility.cpp
--- a/Engine/Source/Runtime/GameplayAbilities/Private/Abilities/GameplayAbility.cpp
+++ b/Engine/Source/Runtime/GameplayAbilities/Private/Abilities/GameplayAbility.cpp
@@ -7,7 +7,7 @@
 #include "AbilitySystemComponent.h"
 #include "AbilitySystemBlueprintLibrary.h"
 #include "AbilityTask.h"
-#pragma optimize("",off)
+
 // --------------------------------------------------------------------------------------------------------------------------------------------------------
 //
 //	UGameplayAbility
