First pass of ground trace working.

diff --git a/Engine/Source/Runtime/GameplayAbilities/Private/Abilities/GameplayAbilityTargetActor_GroundTrace.cpp b/Engine/Source/Runtime/GameplayAbilities/Private/Abilities/GameplayAbilityTargetActor_GroundTrace.cpp
new file mode 100644
--- /dev/null
+++ b/Engine/Source/Runtime/GameplayAbilities/Private/Abilities/GameplayAbilityTargetActor_GroundTrace.cpp
@@ -0,0 +1,82 @@
+// Copyright 1998-2014 Epic Games, Inc. All Rights Reserved.
+
+#include "AbilitySystemPrivatePCH.h"
+#include "AbilitySystemComponent.h"
+#include "GameplayAbilityTargetActor_GroundTrace.h"
+#include "Engine/World.h"
+#include "Runtime/Engine/Public/Net/UnrealNetwork.h"
+#pragma optimize("",off)
+// --------------------------------------------------------------------------------------------------------------------------------------------------------
+//
+//	AGameplayAbilityTargetActor_GroundTrace
+//
+// --------------------------------------------------------------------------------------------------------------------------------------------------------
+
+AGameplayAbilityTargetActor_GroundTrace::AGameplayAbilityTargetActor_GroundTrace(const class FPostConstructInitializeProperties& PCIP)
+	: Super(PCIP)
+{
+}
+
+FHitResult AGameplayAbilityTargetActor_GroundTrace::PerformTrace(AActor* InSourceActor) const
+{
+	static const FName LineTraceSingleName(TEXT("AGameplayAbilityTargetActor_GroundTrace"));
+	bool bTraceComplex = false;
+
+	FCollisionQueryParams Params(LineTraceSingleName, bTraceComplex);
+	Params.bReturnPhysicalMaterial = true;
+	Params.bTraceAsyncScene = true;
+	Params.AddIgnoredActor(InSourceActor);
+
+	FVector AimDirection = InSourceActor->GetActorForwardVector();		//Default
+	
+	if (OwningAbility)		//Server and launching client only
+	{
+		APlayerController* AimingPC = OwningAbility->GetCurrentActorInfo()->PlayerController.Get();
+		check(AimingPC);
+		FVector CamLoc;
+		FRotator CamRot;
+		AimingPC->GetPlayerViewPoint(CamLoc, CamRot);
+		AimDirection = CamRot.Vector();
+	}
+
+	FVector TraceStart = InSourceActor->GetActorLocation();
+	FVector TraceEnd = TraceStart + (AimDirection * MaxRange);
+
+	//If we're using a socket, adjust the starting location and aim direction after the end position has been found. This way we can still aim with the camera, then fire accurately from the socket.
+	if (OwningAbility)		//Server and launching client only
+	{
+		TraceStart = StartLocation.GetTargetingTransform().GetLocation();
+	}
+	AimDirection = (TraceEnd - TraceStart).SafeNormal();
+
+	// ------------------------------------------------------
+
+	FHitResult ReturnHitResult;
+	InSourceActor->GetWorld()->LineTraceSingle(ReturnHitResult, TraceStart, TraceEnd, ECC_WorldStatic, Params);
+	//Default to end of trace line if we don't hit anything.
+	if (!ReturnHitResult.bBlockingHit)
+	{
+		ReturnHitResult.Location = TraceEnd;
+	}
+
+	//Second trace, straight down. Consider using InSourceActor->GetWorld()->NavigationSystem->ProjectPointToNavigation() instead of tracing in the case of movement abilities (flag/bool).
+	TraceStart = ReturnHitResult.Location - AimDirection;		//Pull back very slightly to avoid scraping down walls
+	TraceEnd = TraceStart;
+	TraceEnd[2] -= 999999.0f;
+	if (ReturnHitResult.Actor.IsValid())
+	{
+		Params.AddIgnoredActor(ReturnHitResult.Actor.Get());
+	}
+	InSourceActor->GetWorld()->LineTraceSingle(ReturnHitResult, TraceStart, TraceEnd, ECC_WorldStatic, Params);
+	if (!ReturnHitResult.bBlockingHit)
+	{
+		ReturnHitResult.Location = TraceEnd;
+	}
+
+	if (AActor* LocalReticleActor = ReticleActor.Get())
+	{
+		LocalReticleActor->SetActorLocation(ReturnHitResult.Location);
+	}
+
+	return ReturnHitResult;
+}
diff --git a/Engine/Source/Runtime/GameplayAbilities/Private/Abilities/GameplayAbilityTargetActor_SingleLineTrace.cpp b/Engine/Source/Runtime/GameplayAbilities/Private/Abilities/GameplayAbilityTargetActor_SingleLineTrace.cpp
--- a/Engine/Source/Runtime/GameplayAbilities/Private/Abilities/GameplayAbilityTargetActor_SingleLineTrace.cpp
+++ b/Engine/Source/Runtime/GameplayAbilities/Private/Abilities/GameplayAbilityTargetActor_SingleLineTrace.cpp
@@ -5,7 +5,7 @@
 #include "GameplayAbilityTargetActor_SingleLineTrace.h"
 #include "Engine/World.h"
 #include "Runtime/Engine/Public/Net/UnrealNetwork.h"
-
+#pragma optimize("",off)
 // --------------------------------------------------------------------------------------------------------------------------------------------------------
 //
 //	AGameplayAbilityTargetActor_SingleLineTrace
diff --git a/Engine/Source/Runtime/GameplayAbilities/Public/Abilities/GameplayAbilityTargetActor_GroundTrace.h b/Engine/Source/Runtime/GameplayAbilities/Public/Abilities/GameplayAbilityTargetActor_GroundTrace.h
new file mode 100644
--- /dev/null
+++ b/Engine/Source/Runtime/GameplayAbilities/Public/Abilities/GameplayAbilityTargetActor_GroundTrace.h
@@ -0,0 +1,16 @@
+// Copyright 1998-2014 Epic Games, Inc. All Rights Reserved.
+
+#pragma once
+
+#include "GameplayAbilityTargetActor_SingleLineTrace.h"
+#include "GameplayAbilityWorldReticle.h"
+#include "GameplayAbilityTargetActor_GroundTrace.generated.h"
+
+UCLASS(Blueprintable)
+class GAMEPLAYABILITIES_API AGameplayAbilityTargetActor_GroundTrace : public AGameplayAbilityTargetActor_SingleLineTrace
+{
+	GENERATED_UCLASS_BODY()
+
+protected:
+	virtual FHitResult PerformTrace(AActor* InSourceActor) const override;
+};
\ No newline at end of file
