Pragma cleanup.

diff --git a/Engine/Source/Runtime/GameplayAbilities/Private/Abilities/GameplayAbilityTargetActor_GroundTrace.cpp b/Engine/Source/Runtime/GameplayAbilities/Private/Abilities/GameplayAbilityTargetActor_GroundTrace.cpp
--- a/Engine/Source/Runtime/GameplayAbilities/Private/Abilities/GameplayAbilityTargetActor_GroundTrace.cpp
+++ b/Engine/Source/Runtime/GameplayAbilities/Private/Abilities/GameplayAbilityTargetActor_GroundTrace.cpp
@@ -5,7 +5,7 @@
 #include "GameplayAbilityTargetActor_GroundTrace.h"
 #include "Engine/World.h"
 #include "Runtime/Engine/Public/Net/UnrealNetwork.h"
-#pragma optimize("",off)
+
 // --------------------------------------------------------------------------------------------------------------------------------------------------------
 //
 //	AGameplayAbilityTargetActor_GroundTrace
diff --git a/Engine/Source/Runtime/GameplayAbilities/Private/Abilities/GameplayAbilityTargetActor_SingleLineTrace.cpp b/Engine/Source/Runtime/GameplayAbilities/Private/Abilities/GameplayAbilityTargetActor_SingleLineTrace.cpp
--- a/Engine/Source/Runtime/GameplayAbilities/Private/Abilities/GameplayAbilityTargetActor_SingleLineTrace.cpp
+++ b/Engine/Source/Runtime/GameplayAbilities/Private/Abilities/GameplayAbilityTargetActor_SingleLineTrace.cpp
@@ -5,7 +5,7 @@
 #include "GameplayAbilityTargetActor_SingleLineTrace.h"
 #include "Engine/World.h"
 #include "Runtime/Engine/Public/Net/UnrealNetwork.h"
-#pragma optimize("",off)
+
 // --------------------------------------------------------------------------------------------------------------------------------------------------------
 //
 //	AGameplayAbilityTargetActor_SingleLineTrace
