Hiding some pins.

diff --git a/Engine/Source/Runtime/GameplayAbilities/Private/Abilities/GameplayAbility.cpp b/Engine/Source/Runtime/GameplayAbilities/Private/Abilities/GameplayAbility.cpp
--- a/Engine/Source/Runtime/GameplayAbilities/Private/Abilities/GameplayAbility.cpp
+++ b/Engine/Source/Runtime/GameplayAbilities/Private/Abilities/GameplayAbility.cpp
@@ -586,7 +586,7 @@ void UGameplayAbility::MontageBranchPoint_AbilityDecisionStart(const FGameplayAb
 
 }
 
-FGameplayAbilityTargetingLocationInfo UGameplayAbility::MakeTargetLocationInfoFromOwnerActor() const
+FGameplayAbilityTargetingLocationInfo UGameplayAbility::MakeTargetLocationInfoFromOwnerActor()
 {
 	FGameplayAbilityTargetingLocationInfo ReturnLocation;
 	ReturnLocation.LocationType = EGameplayAbilityTargetingLocationType::ActorTransform;
@@ -594,7 +594,7 @@ FGameplayAbilityTargetingLocationInfo UGameplayAbility::MakeTargetLocationInfoFr
 	return ReturnLocation;
 }
 
-FGameplayAbilityTargetingLocationInfo UGameplayAbility::MakeTargetLocationInfoFromOwnerSkeletalMeshComponent(FName SocketName) const
+FGameplayAbilityTargetingLocationInfo UGameplayAbility::MakeTargetLocationInfoFromOwnerSkeletalMeshComponent(FName SocketName)
 {
 	FGameplayAbilityTargetingLocationInfo ReturnLocation;
 	ReturnLocation.LocationType = EGameplayAbilityTargetingLocationType::SocketTransform;
diff --git a/Engine/Source/Runtime/GameplayAbilities/Public/Abilities/GameplayAbility.h b/Engine/Source/Runtime/GameplayAbilities/Public/Abilities/GameplayAbility.h
--- a/Engine/Source/Runtime/GameplayAbilities/Public/Abilities/GameplayAbility.h
+++ b/Engine/Source/Runtime/GameplayAbilities/Public/Abilities/GameplayAbility.h
@@ -413,11 +413,11 @@ protected:
 
 	void MontageBranchPoint_AbilityDecisionStart(const FGameplayAbilityActorInfo* ActorInfo) const;
 
-	UFUNCTION(BlueprintCallable, Category=Ability)
-	FGameplayAbilityTargetingLocationInfo MakeTargetLocationInfoFromOwnerActor() const;
+	UFUNCTION(BlueprintPure, Category = Ability, meta = (HidePin = "WorldContextObject", DefaultToSelf = "WorldContextObject"))
+	FGameplayAbilityTargetingLocationInfo MakeTargetLocationInfoFromOwnerActor();
 
-	UFUNCTION(BlueprintCallable, Category = Ability)
-	FGameplayAbilityTargetingLocationInfo MakeTargetLocationInfoFromOwnerSkeletalMeshComponent(FName SocketName) const;
+	UFUNCTION(BlueprintPure, Category = Ability, meta = (HidePin = "WorldContextObject", DefaultToSelf = "WorldContextObject"))
+	FGameplayAbilityTargetingLocationInfo MakeTargetLocationInfoFromOwnerSkeletalMeshComponent(FName SocketName);
 private:
 
 	/** 
