Abstract base targeting class, initial pass.

diff --git a/Engine/Source/Runtime/GameplayAbilities/Private/Abilities/GameplayAbilityTargetActor_SingleLineTrace.cpp b/Engine/Source/Runtime/GameplayAbilities/Private/Abilities/GameplayAbilityTargetActor_SingleLineTrace.cpp
--- a/Engine/Source/Runtime/GameplayAbilities/Private/Abilities/GameplayAbilityTargetActor_SingleLineTrace.cpp
+++ b/Engine/Source/Runtime/GameplayAbilities/Private/Abilities/GameplayAbilityTargetActor_SingleLineTrace.cpp
@@ -15,33 +15,8 @@
 AGameplayAbilityTargetActor_SingleLineTrace::AGameplayAbilityTargetActor_SingleLineTrace(const class FPostConstructInitializeProperties& PCIP)
 	: Super(PCIP)
 {
-	PrimaryActorTick.bCanEverTick = true;
-	PrimaryActorTick.TickGroup = TG_PrePhysics;
-	StaticTargetFunction = false;
-
-	bDebug = false;
-	MaxRange = 999999.0f;
-}
-
-void AGameplayAbilityTargetActor_SingleLineTrace::GetLifetimeReplicatedProps(TArray< FLifetimeProperty > & OutLifetimeProps) const
-{
-	Super::GetLifetimeReplicatedProps(OutLifetimeProps);
-	
-	DOREPLIFETIME(AGameplayAbilityTargetActor_SingleLineTrace, bDebug);
-	DOREPLIFETIME(AGameplayAbilityTargetActor_SingleLineTrace, SourceActor);
-}
-
-void AGameplayAbilityTargetActor_SingleLineTrace::EndPlay(const EEndPlayReason::Type EndPlayReason)
-{
-	if (ReticleActor.IsValid())
-	{
-		ReticleActor.Get()->Destroy();
-	}
-
-	Super::EndPlay(EndPlayReason);
 }
 
-
 FHitResult AGameplayAbilityTargetActor_SingleLineTrace::PerformTrace(AActor* InSourceActor) const
 {
 	static const FName LineTraceSingleName(TEXT("AGameplayAbilityTargetActor_SingleLineTrace"));
@@ -92,72 +67,3 @@ FHitResult AGameplayAbilityTargetActor_SingleLineTrace::PerformTrace(AActor* InS
 	}
 	return ReturnHitResult;
 }
-
-FGameplayAbilityTargetDataHandle AGameplayAbilityTargetActor_SingleLineTrace::StaticGetTargetData(UWorld * World, const FGameplayAbilityActorInfo* ActorInfo, FGameplayAbilityActivationInfo ActivationInfo) const
-{
-	AActor* StaticSourceActor = ActorInfo->Actor.Get();
-	return MakeTargetData(PerformTrace(StaticSourceActor));
-}
-
-void AGameplayAbilityTargetActor_SingleLineTrace::StartTargeting(UGameplayAbility* InAbility)
-{
-	Super::StartTargeting(InAbility);
-	SourceActor = InAbility->GetCurrentActorInfo()->Actor.Get();
-	
-	bDebug = true;
-
-	ReticleActor = GetWorld()->SpawnActor<AGameplayAbilityWorldReticle>(ReticleClass, GetActorLocation(), GetActorRotation());
-	if (AGameplayAbilityWorldReticle* CachedReticleActor = ReticleActor.Get())
-	{
-		CachedReticleActor->InitializeReticle(this);
-	}
-}
-
-void AGameplayAbilityTargetActor_SingleLineTrace::Tick(float DeltaSeconds)
-{
-	// very temp - do a mostly hardcoded trace from the source actor
-	if (SourceActor)
-	{
-		FHitResult HitResult = PerformTrace(SourceActor);
-		FVector EndPoint = HitResult.Component.IsValid() ? HitResult.ImpactPoint : HitResult.TraceEnd;
-
-		if (bDebug)
-		{
-			DrawDebugLine(GetWorld(), SourceActor->GetActorLocation(), EndPoint, FLinearColor::Green, false);
-			DrawDebugSphere(GetWorld(), EndPoint, 16, 10, FLinearColor::Green, false);
-		}
-
-		SetActorLocationAndRotation(EndPoint, SourceActor->GetActorRotation());
-
-		if (ShouldProduceTargetData() && bAutoFire)
-		{
-			if (TimeUntilAutoFire <= 0.0f)
-			{
-				ConfirmTargeting();
-				bAutoFire = false;
-			}
-			TimeUntilAutoFire -= DeltaSeconds;
-		}
-	}
-}
-
-void AGameplayAbilityTargetActor_SingleLineTrace::ConfirmTargeting()
-{
-	check(ShouldProduceTargetData());
-
-	if (SourceActor)
-	{
-		bDebug = false;
-		bAutoFire = false;
-		FGameplayAbilityTargetDataHandle Handle = MakeTargetData(PerformTrace(SourceActor));
-		TargetDataReadyDelegate.Broadcast(Handle);
-	}
-
-	Destroy();
-}
-
-FGameplayAbilityTargetDataHandle AGameplayAbilityTargetActor_SingleLineTrace::MakeTargetData(FHitResult HitResult) const
-{
-	/** Note: This will be cleaned up by the FGameplayAbilityTargetDataHandle (via an internal TSharedPtr) */
-	return StartLocation.MakeTargetDataHandleFromHitResult(OwningAbility, HitResult);
-}
diff --git a/Engine/Source/Runtime/GameplayAbilities/Private/Abilities/GameplayAbilityTargetActor_Trace.cpp b/Engine/Source/Runtime/GameplayAbilities/Private/Abilities/GameplayAbilityTargetActor_Trace.cpp
new file mode 100644
--- /dev/null
+++ b/Engine/Source/Runtime/GameplayAbilities/Private/Abilities/GameplayAbilityTargetActor_Trace.cpp
@@ -0,0 +1,111 @@
+// Copyright 1998-2014 Epic Games, Inc. All Rights Reserved.
+
+#include "AbilitySystemPrivatePCH.h"
+#include "AbilitySystemComponent.h"
+#include "GameplayAbilityTargetActor_Trace.h"
+#include "Engine/World.h"
+#include "Runtime/Engine/Public/Net/UnrealNetwork.h"
+
+// --------------------------------------------------------------------------------------------------------------------------------------------------------
+//
+//	AGameplayAbilityTargetActor_Trace
+//
+// --------------------------------------------------------------------------------------------------------------------------------------------------------
+
+AGameplayAbilityTargetActor_Trace::AGameplayAbilityTargetActor_Trace(const class FPostConstructInitializeProperties& PCIP)
+	: Super(PCIP)
+{
+	PrimaryActorTick.bCanEverTick = true;
+	PrimaryActorTick.TickGroup = TG_PrePhysics;
+	StaticTargetFunction = false;
+
+	bDebug = false;
+	MaxRange = 999999.0f;
+}
+
+void AGameplayAbilityTargetActor_Trace::GetLifetimeReplicatedProps(TArray< FLifetimeProperty > & OutLifetimeProps) const
+{
+	Super::GetLifetimeReplicatedProps(OutLifetimeProps);
+	
+	DOREPLIFETIME(AGameplayAbilityTargetActor_Trace, bDebug);
+	DOREPLIFETIME(AGameplayAbilityTargetActor_Trace, SourceActor);
+}
+
+void AGameplayAbilityTargetActor_Trace::EndPlay(const EEndPlayReason::Type EndPlayReason)
+{
+	if (ReticleActor.IsValid())
+	{
+		ReticleActor.Get()->Destroy();
+	}
+
+	Super::EndPlay(EndPlayReason);
+}
+
+FGameplayAbilityTargetDataHandle AGameplayAbilityTargetActor_Trace::StaticGetTargetData(UWorld * World, const FGameplayAbilityActorInfo* ActorInfo, FGameplayAbilityActivationInfo ActivationInfo) const
+{
+	AActor* StaticSourceActor = ActorInfo->Actor.Get();
+	return MakeTargetData(PerformTrace(StaticSourceActor));
+}
+
+void AGameplayAbilityTargetActor_Trace::StartTargeting(UGameplayAbility* InAbility)
+{
+	Super::StartTargeting(InAbility);
+	SourceActor = InAbility->GetCurrentActorInfo()->Actor.Get();
+	
+	bDebug = true;
+
+	ReticleActor = GetWorld()->SpawnActor<AGameplayAbilityWorldReticle>(ReticleClass, GetActorLocation(), GetActorRotation());
+	if (AGameplayAbilityWorldReticle* CachedReticleActor = ReticleActor.Get())
+	{
+		CachedReticleActor->InitializeReticle(this);
+	}
+}
+
+void AGameplayAbilityTargetActor_Trace::Tick(float DeltaSeconds)
+{
+	// very temp - do a mostly hardcoded trace from the source actor
+	if (SourceActor)
+	{
+		FHitResult HitResult = PerformTrace(SourceActor);
+		FVector EndPoint = HitResult.Component.IsValid() ? HitResult.ImpactPoint : HitResult.TraceEnd;
+
+		if (bDebug)
+		{
+			DrawDebugLine(GetWorld(), SourceActor->GetActorLocation(), EndPoint, FLinearColor::Green, false);
+			DrawDebugSphere(GetWorld(), EndPoint, 16, 10, FLinearColor::Green, false);
+		}
+
+		SetActorLocationAndRotation(EndPoint, SourceActor->GetActorRotation());
+
+		if (ShouldProduceTargetData() && bAutoFire)
+		{
+			if (TimeUntilAutoFire <= 0.0f)
+			{
+				ConfirmTargeting();
+				bAutoFire = false;
+			}
+			TimeUntilAutoFire -= DeltaSeconds;
+		}
+	}
+}
+
+void AGameplayAbilityTargetActor_Trace::ConfirmTargeting()
+{
+	check(ShouldProduceTargetData());
+
+	if (SourceActor)
+	{
+		bDebug = false;
+		bAutoFire = false;
+		FGameplayAbilityTargetDataHandle Handle = MakeTargetData(PerformTrace(SourceActor));
+		TargetDataReadyDelegate.Broadcast(Handle);
+	}
+
+	Destroy();
+}
+
+FGameplayAbilityTargetDataHandle AGameplayAbilityTargetActor_Trace::MakeTargetData(FHitResult HitResult) const
+{
+	/** Note: This will be cleaned up by the FGameplayAbilityTargetDataHandle (via an internal TSharedPtr) */
+	return StartLocation.MakeTargetDataHandleFromHitResult(OwningAbility, HitResult);
+}
diff --git a/Engine/Source/Runtime/GameplayAbilities/Public/Abilities/GameplayAbilityTargetActor_GroundTrace.h b/Engine/Source/Runtime/GameplayAbilities/Public/Abilities/GameplayAbilityTargetActor_GroundTrace.h
--- a/Engine/Source/Runtime/GameplayAbilities/Public/Abilities/GameplayAbilityTargetActor_GroundTrace.h
+++ b/Engine/Source/Runtime/GameplayAbilities/Public/Abilities/GameplayAbilityTargetActor_GroundTrace.h
@@ -2,12 +2,12 @@
 
 #pragma once
 
-#include "GameplayAbilityTargetActor_SingleLineTrace.h"
+#include "GameplayAbilityTargetActor_Trace.h"
 #include "GameplayAbilityWorldReticle.h"
 #include "GameplayAbilityTargetActor_GroundTrace.generated.h"
 
 UCLASS(Blueprintable)
-class GAMEPLAYABILITIES_API AGameplayAbilityTargetActor_GroundTrace : public AGameplayAbilityTargetActor_SingleLineTrace
+class GAMEPLAYABILITIES_API AGameplayAbilityTargetActor_GroundTrace : public AGameplayAbilityTargetActor_Trace
 {
 	GENERATED_UCLASS_BODY()
 
diff --git a/Engine/Source/Runtime/GameplayAbilities/Public/Abilities/GameplayAbilityTargetActor_SingleLineTrace.h b/Engine/Source/Runtime/GameplayAbilities/Public/Abilities/GameplayAbilityTargetActor_SingleLineTrace.h
--- a/Engine/Source/Runtime/GameplayAbilities/Public/Abilities/GameplayAbilityTargetActor_SingleLineTrace.h
+++ b/Engine/Source/Runtime/GameplayAbilities/Public/Abilities/GameplayAbilityTargetActor_SingleLineTrace.h
@@ -2,48 +2,15 @@
 
 #pragma once
 
-#include "GameplayAbilityTargetActor.h"
+#include "GameplayAbilityTargetActor_Trace.h"
 #include "GameplayAbilityWorldReticle.h"
 #include "GameplayAbilityTargetActor_SingleLineTrace.generated.h"
 
 UCLASS(Blueprintable)
-class GAMEPLAYABILITIES_API AGameplayAbilityTargetActor_SingleLineTrace : public AGameplayAbilityTargetActor
+class GAMEPLAYABILITIES_API AGameplayAbilityTargetActor_SingleLineTrace : public AGameplayAbilityTargetActor_Trace
 {
 	GENERATED_UCLASS_BODY()
 
-public:
-	virtual void EndPlay(const EEndPlayReason::Type EndPlayReason) override;
-
-	virtual FGameplayAbilityTargetDataHandle StaticGetTargetData(UWorld * World, const FGameplayAbilityActorInfo* ActorInfo, FGameplayAbilityActivationInfo ActivationInfo) const override;
-
-	virtual void StartTargeting(UGameplayAbility* Ability) override;
-	
-	virtual void ConfirmTargeting() override;	
-
-	virtual void Tick(float DeltaSeconds) override;
-	
-	UPROPERTY(BlueprintReadOnly, Replicated, Category = Targeting)
-	bool bDebug;
-
-	UPROPERTY(BlueprintReadOnly, Replicated, Category = Targeting)
-	AActor* SourceActor;
-
-	UPROPERTY(BlueprintReadWrite, EditAnywhere, meta = (ExposeOnSpawn = true), Category = Projectile)
-	float MaxRange;
-
-	UPROPERTY(BlueprintReadWrite, EditAnywhere, meta = (ExposeOnSpawn = true), Category = Projectile)
-	bool bAutoFire;				//Automatically confirm after the timer runs out. Should also have support for an auto-cancel feature.
-
-	UPROPERTY(BlueprintReadWrite, EditAnywhere, meta = (ExposeOnSpawn = true), Category = Projectile)
-	float TimeUntilAutoFire;
-
-	UPROPERTY(BlueprintReadWrite, EditAnywhere, meta = (ExposeOnSpawn = true), Category = Projectile)
-	TSubclassOf<AGameplayAbilityWorldReticle> ReticleClass;		//Using a special class for replication purposes
-
 protected:
-	virtual FHitResult PerformTrace(AActor* InSourceActor) const;
-
-	FGameplayAbilityTargetDataHandle MakeTargetData(FHitResult HitResult) const;
-
-	TWeakObjectPtr<AGameplayAbilityWorldReticle> ReticleActor;
+	virtual FHitResult PerformTrace(AActor* InSourceActor) const override;
 };
\ No newline at end of file
diff --git a/Engine/Source/Runtime/GameplayAbilities/Public/Abilities/GameplayAbilityTargetActor_Trace.h b/Engine/Source/Runtime/GameplayAbilities/Public/Abilities/GameplayAbilityTargetActor_Trace.h
new file mode 100644
--- /dev/null
+++ b/Engine/Source/Runtime/GameplayAbilities/Public/Abilities/GameplayAbilityTargetActor_Trace.h
@@ -0,0 +1,49 @@
+// Copyright 1998-2014 Epic Games, Inc. All Rights Reserved.
+
+#pragma once
+
+#include "GameplayAbilityTargetActor.h"
+#include "GameplayAbilityWorldReticle.h"
+#include "GameplayAbilityTargetActor_Trace.generated.h"
+
+UCLASS(Abstract, Blueprintable)
+class GAMEPLAYABILITIES_API AGameplayAbilityTargetActor_Trace : public AGameplayAbilityTargetActor
+{
+	GENERATED_UCLASS_BODY()
+
+public:
+	virtual void EndPlay(const EEndPlayReason::Type EndPlayReason) override;
+
+	virtual FGameplayAbilityTargetDataHandle StaticGetTargetData(UWorld * World, const FGameplayAbilityActorInfo* ActorInfo, FGameplayAbilityActivationInfo ActivationInfo) const override;
+
+	virtual void StartTargeting(UGameplayAbility* Ability) override;
+	
+	virtual void ConfirmTargeting() override;	
+
+	virtual void Tick(float DeltaSeconds) override;
+	
+	UPROPERTY(BlueprintReadOnly, Replicated, Category = Targeting)
+	bool bDebug;
+
+	UPROPERTY(BlueprintReadOnly, Replicated, Category = Targeting)
+	AActor* SourceActor;
+
+	UPROPERTY(BlueprintReadWrite, EditAnywhere, meta = (ExposeOnSpawn = true), Category = Projectile)
+	float MaxRange;
+
+	UPROPERTY(BlueprintReadWrite, EditAnywhere, meta = (ExposeOnSpawn = true), Category = Projectile)
+	bool bAutoFire;				//Automatically confirm after the timer runs out. Should also have support for an auto-cancel feature.
+
+	UPROPERTY(BlueprintReadWrite, EditAnywhere, meta = (ExposeOnSpawn = true), Category = Projectile)
+	float TimeUntilAutoFire;
+
+	UPROPERTY(BlueprintReadWrite, EditAnywhere, meta = (ExposeOnSpawn = true), Category = Projectile)
+	TSubclassOf<AGameplayAbilityWorldReticle> ReticleClass;		//Using a special class for replication purposes
+
+protected:
+	virtual FHitResult PerformTrace(AActor* InSourceActor) const PURE_VIRTUAL(AGameplayAbilityTargetActor_Trace, return FHitResult(););
+
+	FGameplayAbilityTargetDataHandle MakeTargetData(FHitResult HitResult) const;
+
+	TWeakObjectPtr<AGameplayAbilityWorldReticle> ReticleActor;
+};
\ No newline at end of file
