Moved debug property up to main targeting class, seems generally convenient to have it around for developing and testing new targeting classes.

diff --git a/Engine/Source/Runtime/GameplayAbilities/Private/Abilities/GameplayAbilityTargetActor.cpp b/Engine/Source/Runtime/GameplayAbilities/Private/Abilities/GameplayAbilityTargetActor.cpp
--- a/Engine/Source/Runtime/GameplayAbilities/Private/Abilities/GameplayAbilityTargetActor.cpp
+++ b/Engine/Source/Runtime/GameplayAbilities/Private/Abilities/GameplayAbilityTargetActor.cpp
@@ -16,12 +16,14 @@ AGameplayAbilityTargetActor::AGameplayAbilityTargetActor(const class FPostConstr
 {
 	StaticTargetFunction = false;
 	ShouldProduceTargetDataOnServer = false;
+	bDebug = false;
 }
 
 void AGameplayAbilityTargetActor::GetLifetimeReplicatedProps(TArray< FLifetimeProperty > & OutLifetimeProps) const
 {
 	Super::GetLifetimeReplicatedProps(OutLifetimeProps);
 	DOREPLIFETIME(AGameplayAbilityTargetActor, StartLocation);
+	DOREPLIFETIME(AGameplayAbilityTargetActor, bDebug);
 }
 
 void AGameplayAbilityTargetActor::StartTargeting(UGameplayAbility* Ability)
diff --git a/Engine/Source/Runtime/GameplayAbilities/Private/Abilities/GameplayAbilityTargetActor_Trace.cpp b/Engine/Source/Runtime/GameplayAbilities/Private/Abilities/GameplayAbilityTargetActor_Trace.cpp
--- a/Engine/Source/Runtime/GameplayAbilities/Private/Abilities/GameplayAbilityTargetActor_Trace.cpp
+++ b/Engine/Source/Runtime/GameplayAbilities/Private/Abilities/GameplayAbilityTargetActor_Trace.cpp
@@ -19,7 +19,6 @@ AGameplayAbilityTargetActor_Trace::AGameplayAbilityTargetActor_Trace(const class
 	PrimaryActorTick.TickGroup = TG_PrePhysics;
 	StaticTargetFunction = false;
 
-	bDebug = false;
 	MaxRange = 999999.0f;
 }
 
@@ -27,7 +26,6 @@ void AGameplayAbilityTargetActor_Trace::GetLifetimeReplicatedProps(TArray< FLife
 {
 	Super::GetLifetimeReplicatedProps(OutLifetimeProps);
 	
-	DOREPLIFETIME(AGameplayAbilityTargetActor_Trace, bDebug);
 	DOREPLIFETIME(AGameplayAbilityTargetActor_Trace, SourceActor);
 }
 
diff --git a/Engine/Source/Runtime/GameplayAbilities/Public/Abilities/GameplayAbilityTargetActor.h b/Engine/Source/Runtime/GameplayAbilities/Public/Abilities/GameplayAbilityTargetActor.h
--- a/Engine/Source/Runtime/GameplayAbilities/Public/Abilities/GameplayAbilityTargetActor.h
+++ b/Engine/Source/Runtime/GameplayAbilities/Public/Abilities/GameplayAbilityTargetActor.h
@@ -55,7 +55,10 @@ public:
 
 	virtual bool IsNetRelevantFor(class APlayerController* RealViewer, AActor* Viewer, const FVector& SrcLocation) override;
 
-	UPROPERTY(BlueprintReadOnly, Category="Targeting")
+	UPROPERTY(BlueprintReadOnly, Replicated, Category = Targeting)
+	bool bDebug;
+
+	UPROPERTY(BlueprintReadOnly, Category = "Targeting")
 	APlayerController* MasterPC;
 
 	UPROPERTY()
diff --git a/Engine/Source/Runtime/GameplayAbilities/Public/Abilities/GameplayAbilityTargetActor_Trace.h b/Engine/Source/Runtime/GameplayAbilities/Public/Abilities/GameplayAbilityTargetActor_Trace.h
--- a/Engine/Source/Runtime/GameplayAbilities/Public/Abilities/GameplayAbilityTargetActor_Trace.h
+++ b/Engine/Source/Runtime/GameplayAbilities/Public/Abilities/GameplayAbilityTargetActor_Trace.h
@@ -22,9 +22,6 @@ public:
 
 	virtual void Tick(float DeltaSeconds) override;
 	
-	UPROPERTY(BlueprintReadOnly, Replicated, Category = Targeting)
-	bool bDebug;
-
 	UPROPERTY(BlueprintReadOnly, Replicated, Category = Targeting)
 	AActor* SourceActor;
 
