Moved SourceActor up to the main targeting base class after realizing every single derived class has it, and it was the only replicated member left.

diff --git a/Engine/Source/Runtime/GameplayAbilities/Private/Abilities/GameplayAbilityTargetActor.cpp b/Engine/Source/Runtime/GameplayAbilities/Private/Abilities/GameplayAbilityTargetActor.cpp
--- a/Engine/Source/Runtime/GameplayAbilities/Private/Abilities/GameplayAbilityTargetActor.cpp
+++ b/Engine/Source/Runtime/GameplayAbilities/Private/Abilities/GameplayAbilityTargetActor.cpp
@@ -23,6 +23,7 @@ void AGameplayAbilityTargetActor::GetLifetimeReplicatedProps(TArray< FLifetimePr
 {
 	Super::GetLifetimeReplicatedProps(OutLifetimeProps);
 	DOREPLIFETIME(AGameplayAbilityTargetActor, StartLocation);
+	DOREPLIFETIME(AGameplayAbilityTargetActor, SourceActor);
 	DOREPLIFETIME(AGameplayAbilityTargetActor, bDebug);
 }
 
diff --git a/Engine/Source/Runtime/GameplayAbilities/Private/Abilities/GameplayAbilityTargetActor_Radius.cpp b/Engine/Source/Runtime/GameplayAbilities/Private/Abilities/GameplayAbilityTargetActor_Radius.cpp
--- a/Engine/Source/Runtime/GameplayAbilities/Private/Abilities/GameplayAbilityTargetActor_Radius.cpp
+++ b/Engine/Source/Runtime/GameplayAbilities/Private/Abilities/GameplayAbilityTargetActor_Radius.cpp
@@ -22,13 +22,6 @@ AGameplayAbilityTargetActor_Radius::AGameplayAbilityTargetActor_Radius(const cla
 	ShouldProduceTargetDataOnServer = true;
 }
 
-void AGameplayAbilityTargetActor_Radius::GetLifetimeReplicatedProps(TArray< FLifetimeProperty > & OutLifetimeProps) const
-{
-	Super::GetLifetimeReplicatedProps(OutLifetimeProps);
-
-	DOREPLIFETIME(AGameplayAbilityTargetActor_Radius, SourceActor);
-}
-
 void AGameplayAbilityTargetActor_Radius::StartTargeting(UGameplayAbility* InAbility)
 {
 	Super::StartTargeting(InAbility);
diff --git a/Engine/Source/Runtime/GameplayAbilities/Private/Abilities/GameplayAbilityTargetActor_Trace.cpp b/Engine/Source/Runtime/GameplayAbilities/Private/Abilities/GameplayAbilityTargetActor_Trace.cpp
--- a/Engine/Source/Runtime/GameplayAbilities/Private/Abilities/GameplayAbilityTargetActor_Trace.cpp
+++ b/Engine/Source/Runtime/GameplayAbilities/Private/Abilities/GameplayAbilityTargetActor_Trace.cpp
@@ -22,13 +22,6 @@ AGameplayAbilityTargetActor_Trace::AGameplayAbilityTargetActor_Trace(const class
 	MaxRange = 999999.0f;
 }
 
-void AGameplayAbilityTargetActor_Trace::GetLifetimeReplicatedProps(TArray< FLifetimeProperty > & OutLifetimeProps) const
-{
-	Super::GetLifetimeReplicatedProps(OutLifetimeProps);
-	
-	DOREPLIFETIME(AGameplayAbilityTargetActor_Trace, SourceActor);
-}
-
 void AGameplayAbilityTargetActor_Trace::EndPlay(const EEndPlayReason::Type EndPlayReason)
 {
 	if (ReticleActor.IsValid())
diff --git a/Engine/Source/Runtime/GameplayAbilities/Public/Abilities/GameplayAbilityTargetActor.h b/Engine/Source/Runtime/GameplayAbilities/Public/Abilities/GameplayAbilityTargetActor.h
--- a/Engine/Source/Runtime/GameplayAbilities/Public/Abilities/GameplayAbilityTargetActor.h
+++ b/Engine/Source/Runtime/GameplayAbilities/Public/Abilities/GameplayAbilityTargetActor.h
@@ -55,13 +55,15 @@ public:
 
 	virtual bool IsNetRelevantFor(class APlayerController* RealViewer, AActor* Viewer, const FVector& SrcLocation) override;
 
-	UPROPERTY(BlueprintReadOnly, Replicated, Category = Targeting)
-	bool bDebug;
-
 	UPROPERTY(BlueprintReadOnly, Category = "Targeting")
 	APlayerController* MasterPC;
 
 	UPROPERTY()
 	UGameplayAbility* OwningAbility;
 
+	UPROPERTY(BlueprintReadOnly, Replicated, Category = Targeting)
+	AActor* SourceActor;
+
+	UPROPERTY(BlueprintReadOnly, Replicated, Category = Targeting)
+	bool bDebug;
 };
\ No newline at end of file
diff --git a/Engine/Source/Runtime/GameplayAbilities/Public/Abilities/GameplayAbilityTargetActor_Radius.h b/Engine/Source/Runtime/GameplayAbilities/Public/Abilities/GameplayAbilityTargetActor_Radius.h
--- a/Engine/Source/Runtime/GameplayAbilities/Public/Abilities/GameplayAbilityTargetActor_Radius.h
+++ b/Engine/Source/Runtime/GameplayAbilities/Public/Abilities/GameplayAbilityTargetActor_Radius.h
@@ -20,9 +20,6 @@ public:
 	UPROPERTY(BlueprintReadWrite, EditAnywhere, meta = (ExposeOnSpawn = true), Category = Radius)
 	float Radius;
 
-	UPROPERTY(BlueprintReadOnly, Replicated, Category = Targeting)
-	AActor* SourceActor;
-
 protected:
 
 	TArray<AActor*>	PerformOverlap(const FVector& Origin);
diff --git a/Engine/Source/Runtime/GameplayAbilities/Public/Abilities/GameplayAbilityTargetActor_Trace.h b/Engine/Source/Runtime/GameplayAbilities/Public/Abilities/GameplayAbilityTargetActor_Trace.h
--- a/Engine/Source/Runtime/GameplayAbilities/Public/Abilities/GameplayAbilityTargetActor_Trace.h
+++ b/Engine/Source/Runtime/GameplayAbilities/Public/Abilities/GameplayAbilityTargetActor_Trace.h
@@ -22,9 +22,6 @@ public:
 
 	virtual void Tick(float DeltaSeconds) override;
 	
-	UPROPERTY(BlueprintReadOnly, Replicated, Category = Targeting)
-	AActor* SourceActor;
-
 	UPROPERTY(BlueprintReadWrite, EditAnywhere, meta = (ExposeOnSpawn = true), Category = Projectile)
 	float MaxRange;
 
