add AbilityTask_WaitVelocityChange

diff --git a/Engine/Source/Runtime/GameplayAbilities/Private/Abilities/Tasks/AbilityTask_WaitVelocityChange.cpp b/Engine/Source/Runtime/GameplayAbilities/Private/Abilities/Tasks/AbilityTask_WaitVelocityChange.cpp
new file mode 100644
--- /dev/null
+++ b/Engine/Source/Runtime/GameplayAbilities/Private/Abilities/Tasks/AbilityTask_WaitVelocityChange.cpp
@@ -0,0 +1,46 @@
+
+#include "AbilitySystemPrivatePCH.h"
+#include "AbilitySystemComponent.h"
+#include "AbilityTask_WaitVelocityChange.h"
+
+UAbilityTask_WaitVelocityChange::UAbilityTask_WaitVelocityChange(const class FPostConstructInitializeProperties& PCIP)
+	: Super(PCIP)
+{
+	bTickingTask = true;
+}
+
+void UAbilityTask_WaitVelocityChange::TickTask(float DeltaTime)
+{
+	if (CachedMovementComponent)
+	{
+		float dot = FVector::DotProduct(Direction, CachedMovementComponent->Velocity);
+
+		if (dot > MinimumMagnitude)
+		{
+			OnVelocityChage.Broadcast();
+			EndTask();
+		}
+	}
+	else
+	{
+		ABILITY_LOG(Warning, TEXT("UAbilityTask_WaitVelocityChange ticked without a valid movement component. ending."));
+		EndTask();
+	}
+}
+
+UAbilityTask_WaitVelocityChange* UAbilityTask_WaitVelocityChange::CreateWaitVelocityChange(UObject* WorldContextObject, FVector InDirection, float InMinimumMagnitude)
+{
+	auto MyObj = NewTask<UAbilityTask_WaitVelocityChange>(WorldContextObject);
+
+	MyObj->MinimumMagnitude = InMinimumMagnitude;
+	MyObj->Direction = InDirection.SafeNormal();
+	
+
+	return MyObj;
+}
+
+void UAbilityTask_WaitVelocityChange::Activate()
+{
+	const FGameplayAbilityActorInfo* ActorInfo = Ability->GetCurrentActorInfo();
+	CachedMovementComponent = ActorInfo->MovementComponent.Get();
+}
\ No newline at end of file
diff --git a/Engine/Source/Runtime/GameplayAbilities/Public/Abilities/Tasks/AbilityTask_WaitVelocityChange.h b/Engine/Source/Runtime/GameplayAbilities/Public/Abilities/Tasks/AbilityTask_WaitVelocityChange.h
new file mode 100644
--- /dev/null
+++ b/Engine/Source/Runtime/GameplayAbilities/Public/Abilities/Tasks/AbilityTask_WaitVelocityChange.h
@@ -0,0 +1,34 @@
+// Copyright 1998-2014 Epic Games, Inc. All Rights Reserved.
+#pragma once
+
+#include "AbilityTask.h"
+#include "AbilityTask_WaitVelocityChange.generated.h"
+
+
+UCLASS(MinimalAPI)
+class UAbilityTask_WaitVelocityChange: public UAbilityTask
+{
+	GENERATED_UCLASS_BODY()
+
+	DECLARE_DYNAMIC_MULTICAST_DELEGATE(FWaitVelocityChangeDelegate);
+
+	/** Delegate called when velocity requirements are met */
+	UPROPERTY(BlueprintAssignable)
+	FWaitVelocityChangeDelegate OnVelocityChage;
+
+	virtual void TickTask(float DeltaTime) override;
+
+	/** Wait for the actor's movement component velocity to be of minimum magnitude when projected along given direction */
+	UFUNCTION(BlueprintCallable, Category = "Ability|Tasks", meta = (FriendlyName="WaitVelocityChange",HidePin = "WorldContextObject", DefaultToSelf = "WorldContextObject", BlueprintInternalUseOnly = "TRUE"))
+	static UAbilityTask_WaitVelocityChange* CreateWaitVelocityChange(UObject* WorldContextObject, FVector Direction, float MinimumMagnitude);
+		
+	virtual void Activate() override;
+
+protected:
+
+	UPROPERTY()
+	UMovementComponent*	CachedMovementComponent;
+
+	float	MinimumMagnitude;
+	FVector Direction;
+};
\ No newline at end of file
