External confirm system implemented. Need to do a second pass for canceling via custom functions instead of just calling EndTask(). Fixed a race condition with task-canceling that came up in some rare cases. Marked reticle and targeting classes as notplaceable and put some comments near the top to indicate functionality.

diff --git a/Engine/Source/Runtime/GameplayAbilities/Private/Abilities/GameplayAbility.cpp b/Engine/Source/Runtime/GameplayAbilities/Private/Abilities/GameplayAbility.cpp
--- a/Engine/Source/Runtime/GameplayAbilities/Private/Abilities/GameplayAbility.cpp
+++ b/Engine/Source/Runtime/GameplayAbilities/Private/Abilities/GameplayAbility.cpp
@@ -410,10 +410,6 @@ void UGameplayAbility::PreActivate(const FGameplayAbilityActorInfo* ActorInfo, c
 
 	bIsActive = true;
 	Comp->NotifyAbilityActivated(this);
-
-	// Become the 'Targeting Ability'
-	// These may need to be more robust - does every ability that activates become the targeting ability, or only certain ones?
-	Comp->SetTargetAbility(this);
 }
 
 void UGameplayAbility::CallActivateAbility(const FGameplayAbilityActorInfo* ActorInfo, const FGameplayAbilityActivationInfo ActivationInfo)
@@ -626,18 +622,45 @@ void UGameplayAbility::TaskStarted(UAbilityTask* NewTask)
 	ActiveTasks.Add(NewTask);
 }
 
-void UGameplayAbility::CancelTaskByInstanceName(FName InstanceName)
+void UGameplayAbility::ConfirmTaskByInstanceName(FName InstanceName, bool bEndTask)
 {
 	TArray<TWeakObjectPtr<UAbilityTask>> NamedTasks = ActiveTasks.FilterByPredicate<FAbilityInstanceNamePredicate>(FAbilityInstanceNamePredicate(InstanceName));
 	for (int32 i = NamedTasks.Num() - 1; i >= 0; --i)
 	{
 		if (NamedTasks[i].IsValid())
 		{
-			NamedTasks[i].Get()->EndTask();
+			UAbilityTask* CurrentTask = NamedTasks[i].Get();
+			CurrentTask->ExternalConfirm(bEndTask);
 		}
 	}
 }
 
+void UGameplayAbility::CancelTasksByInstanceName()
+{
+	for (int32 j = 0; j < CancelTaskInstanceNames.Num(); ++j)
+	{
+		FName InstanceName = CancelTaskInstanceNames[j];
+		TArray<TWeakObjectPtr<UAbilityTask>> NamedTasks = ActiveTasks.FilterByPredicate<FAbilityInstanceNamePredicate>(FAbilityInstanceNamePredicate(InstanceName));
+		for (int32 i = NamedTasks.Num() - 1; i >= 0; --i)
+		{
+			if (NamedTasks[i].IsValid())
+			{
+				UAbilityTask* CurrentTask = NamedTasks[i].Get();
+				//TODO Change this to a proper cancel instead of just ending the task cold. Maybe make a separate function for calling EndTask so we don't lose this functionality.
+				CurrentTask->EndTask();
+			}
+		}
+	}
+	CancelTaskInstanceNames.Empty();
+}
+
+void UGameplayAbility::CancelTaskByInstanceName(FName InstanceName)
+{
+	//Avoid race condition by delaying for one frame
+	CancelTaskInstanceNames.AddUnique(InstanceName);
+	GetWorld()->GetTimerManager().SetTimerForNextTick(this, &UGameplayAbility::CancelTasksByInstanceName);
+}
+
 void UGameplayAbility::TaskEnded(UAbilityTask* Task)
 {
 	ActiveTasks.Remove(Task);
diff --git a/Engine/Source/Runtime/GameplayAbilities/Private/Abilities/GameplayAbilityTargetActor.cpp b/Engine/Source/Runtime/GameplayAbilities/Private/Abilities/GameplayAbilityTargetActor.cpp
--- a/Engine/Source/Runtime/GameplayAbilities/Private/Abilities/GameplayAbilityTargetActor.cpp
+++ b/Engine/Source/Runtime/GameplayAbilities/Private/Abilities/GameplayAbilityTargetActor.cpp
@@ -17,6 +17,7 @@ AGameplayAbilityTargetActor::AGameplayAbilityTargetActor(const class FPostConstr
 	StaticTargetFunction = false;
 	ShouldProduceTargetDataOnServer = false;
 	bDebug = false;
+	bDestroyOnConfirmation = true;
 }
 
 void AGameplayAbilityTargetActor::GetLifetimeReplicatedProps(TArray< FLifetimeProperty > & OutLifetimeProps) const
@@ -25,6 +26,7 @@ void AGameplayAbilityTargetActor::GetLifetimeReplicatedProps(TArray< FLifetimePr
 	DOREPLIFETIME(AGameplayAbilityTargetActor, StartLocation);
 	DOREPLIFETIME(AGameplayAbilityTargetActor, SourceActor);
 	DOREPLIFETIME(AGameplayAbilityTargetActor, bDebug);
+	DOREPLIFETIME(AGameplayAbilityTargetActor, bDestroyOnConfirmation);
 }
 
 void AGameplayAbilityTargetActor::StartTargeting(UGameplayAbility* Ability)
@@ -32,12 +34,19 @@ void AGameplayAbilityTargetActor::StartTargeting(UGameplayAbility* Ability)
 	OwningAbility = Ability;
 }
 
-void AGameplayAbilityTargetActor::ConfirmTargeting()
+void AGameplayAbilityTargetActor::ConfirmTargetingAndContinue()
 {
 	check(ShouldProduceTargetData());
-
 	TargetDataReadyDelegate.Broadcast(FGameplayAbilityTargetDataHandle());
-	Destroy();
+}
+
+void AGameplayAbilityTargetActor::ConfirmTargeting()
+{
+	ConfirmTargetingAndContinue();
+	if (bDestroyOnConfirmation)
+	{
+		Destroy();
+	}
 }
 
 /** Outside code is saying 'stop everything and just forget about it' */
diff --git a/Engine/Source/Runtime/GameplayAbilities/Private/Abilities/GameplayAbilityTargetActor_Radius.cpp b/Engine/Source/Runtime/GameplayAbilities/Private/Abilities/GameplayAbilityTargetActor_Radius.cpp
--- a/Engine/Source/Runtime/GameplayAbilities/Private/Abilities/GameplayAbilityTargetActor_Radius.cpp
+++ b/Engine/Source/Runtime/GameplayAbilities/Private/Abilities/GameplayAbilityTargetActor_Radius.cpp
@@ -28,19 +28,15 @@ void AGameplayAbilityTargetActor_Radius::StartTargeting(UGameplayAbility* InAbil
 	SourceActor = InAbility->GetCurrentActorInfo()->Actor.Get();
 }
 
-void AGameplayAbilityTargetActor_Radius::ConfirmTargeting()
+void AGameplayAbilityTargetActor_Radius::ConfirmTargetingAndContinue()
 {
 	check(ShouldProduceTargetData());
-
 	if (SourceActor)
 	{
 		FVector Origin = StartLocation.GetTargetingTransform().GetLocation();
-
 		FGameplayAbilityTargetDataHandle Handle = MakeTargetData(PerformOverlap(Origin), Origin);
 		TargetDataReadyDelegate.Broadcast(Handle);
 	}
-
-	Destroy();
 }
 
 FGameplayAbilityTargetDataHandle AGameplayAbilityTargetActor_Radius::MakeTargetData(const TArray<AActor*> Actors, const FVector& Origin) const
diff --git a/Engine/Source/Runtime/GameplayAbilities/Private/Abilities/GameplayAbilityTargetActor_Trace.cpp b/Engine/Source/Runtime/GameplayAbilities/Private/Abilities/GameplayAbilityTargetActor_Trace.cpp
--- a/Engine/Source/Runtime/GameplayAbilities/Private/Abilities/GameplayAbilityTargetActor_Trace.cpp
+++ b/Engine/Source/Runtime/GameplayAbilities/Private/Abilities/GameplayAbilityTargetActor_Trace.cpp
@@ -80,10 +80,9 @@ void AGameplayAbilityTargetActor_Trace::Tick(float DeltaSeconds)
 	}
 }
 
-void AGameplayAbilityTargetActor_Trace::ConfirmTargeting()
+void AGameplayAbilityTargetActor_Trace::ConfirmTargetingAndContinue()
 {
 	check(ShouldProduceTargetData());
-
 	if (SourceActor)
 	{
 		bDebug = false;
@@ -91,8 +90,6 @@ void AGameplayAbilityTargetActor_Trace::ConfirmTargeting()
 		FGameplayAbilityTargetDataHandle Handle = MakeTargetData(PerformTrace(SourceActor));
 		TargetDataReadyDelegate.Broadcast(Handle);
 	}
-
-	Destroy();
 }
 
 FGameplayAbilityTargetDataHandle AGameplayAbilityTargetActor_Trace::MakeTargetData(FHitResult HitResult) const
diff --git a/Engine/Source/Runtime/GameplayAbilities/Private/Abilities/Tasks/AbilityTask.cpp b/Engine/Source/Runtime/GameplayAbilities/Private/Abilities/Tasks/AbilityTask.cpp
--- a/Engine/Source/Runtime/GameplayAbilities/Private/Abilities/Tasks/AbilityTask.cpp
+++ b/Engine/Source/Runtime/GameplayAbilities/Private/Abilities/Tasks/AbilityTask.cpp
@@ -6,7 +6,7 @@
 UAbilityTask::UAbilityTask(const class FPostConstructInitializeProperties& PCIP)
 	: Super(PCIP)
 {
-
+	bTickingTask = false;
 }
 
 void UAbilityTask::Activate()
@@ -67,6 +67,14 @@ void UAbilityTask::EndTask()
 	}
 }
 
+void UAbilityTask::ExternalConfirm(bool bEndTask)
+{
+	if (bEndTask)
+	{
+		EndTask();
+	}
+}
+
 void UAbilityTask::OnDestroy(bool AbilityIsEnding)
 {
 	ensure(!IsPendingKill());
diff --git a/Engine/Source/Runtime/GameplayAbilities/Private/Abilities/Tasks/AbilityTask_MoveToLocation.cpp b/Engine/Source/Runtime/GameplayAbilities/Private/Abilities/Tasks/AbilityTask_MoveToLocation.cpp
--- a/Engine/Source/Runtime/GameplayAbilities/Private/Abilities/Tasks/AbilityTask_MoveToLocation.cpp
+++ b/Engine/Source/Runtime/GameplayAbilities/Private/Abilities/Tasks/AbilityTask_MoveToLocation.cpp
@@ -5,11 +5,31 @@
 UAbilityTask_MoveToLocation::UAbilityTask_MoveToLocation(const class FPostConstructInitializeProperties& PCIP)
 : Super(PCIP)
 {
+	bTickingTask = true;
+}
+
+UAbilityTask_MoveToLocation* UAbilityTask_MoveToLocation::MoveToLocation(class UObject* WorldContextObject, FVector Location, float Duration, UCurveFloat* OptionalInterpolationCurve)
+{
+	auto MyObj = NewTask<UAbilityTask_MoveToLocation>(WorldContextObject);
 
+	MyObj->StartLocation = MyObj->GetActor()->GetActorLocation();
+	MyObj->TargetLocation = Location;
+	MyObj->DurationOfMovement = FMath::Max(Duration, 0.001f);		//Avoid negative or divide-by-zero cases
+	MyObj->TimeMoveStarted = MyObj->GetWorld()->GetTimeSeconds();
+	MyObj->TimeMoveWillEnd = MyObj->TimeMoveStarted + MyObj->DurationOfMovement;
+	MyObj->LerpCurve = OptionalInterpolationCurve;
+
+	return MyObj;
 }
 
-void UAbilityTask_MoveToLocation::InterpolatePosition()
+void UAbilityTask_MoveToLocation::Activate()
 {
+}
+
+//TODO: This is still an awful way to do this and we should scrap this task or do it right.
+void UAbilityTask_MoveToLocation::TickTask(float DeltaTime)
+{
+	Super::TickTask(DeltaTime);
 	AActor* MyActor = GetActor();
 	if (MyActor)
 	{
@@ -32,30 +52,3 @@ void UAbilityTask_MoveToLocation::InterpolatePosition()
 		}
 	}
 }
-
-UAbilityTask_MoveToLocation* UAbilityTask_MoveToLocation::MoveToLocation(class UObject* WorldContextObject, FVector Location, float Duration, UCurveFloat* OptionalInterpolationCurve)
-{
-	auto MyObj = NewTask<UAbilityTask_MoveToLocation>(WorldContextObject);
-
-	MyObj->StartLocation = MyObj->GetActor()->GetActorLocation();
-	MyObj->TargetLocation = Location;
-	MyObj->DurationOfMovement = FMath::Max(Duration, 0.001f);		//Avoid negative or divide-by-zero cases
-	MyObj->TimeMoveStarted = MyObj->GetWorld()->GetTimeSeconds();
-	MyObj->TimeMoveWillEnd = MyObj->TimeMoveStarted + MyObj->DurationOfMovement;
-	MyObj->LerpCurve = OptionalInterpolationCurve;
-
-	return MyObj;
-}
-
-void UAbilityTask_MoveToLocation::OnDestroy(bool AbilityEnded)
-{
-	GetWorld()->GetTimerManager().ClearTimer(this, &UAbilityTask_MoveToLocation::InterpolatePosition);
-
-	Super::OnDestroy(AbilityEnded);
-}
-
-void UAbilityTask_MoveToLocation::Activate()
-{
-	// FIXME: This is bad - its calling InterpolatePOsition @ 1000hz!!!!
-	GetWorld()->GetTimerManager().SetTimer(this, &UAbilityTask_MoveToLocation::InterpolatePosition, 0.001f, true);
-}
\ No newline at end of file
diff --git a/Engine/Source/Runtime/GameplayAbilities/Private/Abilities/Tasks/AbilityTask_VisualizeTargeting.cpp b/Engine/Source/Runtime/GameplayAbilities/Private/Abilities/Tasks/AbilityTask_VisualizeTargeting.cpp
--- a/Engine/Source/Runtime/GameplayAbilities/Private/Abilities/Tasks/AbilityTask_VisualizeTargeting.cpp
+++ b/Engine/Source/Runtime/GameplayAbilities/Private/Abilities/Tasks/AbilityTask_VisualizeTargeting.cpp
@@ -12,7 +12,7 @@ UAbilityTask_VisualizeTargeting::UAbilityTask_VisualizeTargeting(const class FPo
 
 UAbilityTask_VisualizeTargeting* UAbilityTask_VisualizeTargeting::VisualizeTargeting(UObject* WorldContextObject, TSubclassOf<AGameplayAbilityTargetActor> InTargetClass, FName TaskInstanceName, float Duration)
 {
-	auto MyObj = NewTask<UAbilityTask_VisualizeTargeting>(WorldContextObject, TaskInstanceName);		//Register for abort list here, providing a given FName as a key
+	auto MyObj = NewTask<UAbilityTask_VisualizeTargeting>(WorldContextObject, TaskInstanceName);		//Register for task list here, providing a given FName as a key
 	MyObj->TargetClass = InTargetClass;
 	if (Duration > 0.0f)
 	{
@@ -89,6 +89,7 @@ void UAbilityTask_VisualizeTargeting::OnDestroy(bool AbilityEnded)
 	{
 		MyTargetActor->Destroy();
 	}
+	GetWorld()->GetTimerManager().ClearTimer(this, &UAbilityTask_VisualizeTargeting::OnTimeElapsed);
 
 	Super::OnDestroy(AbilityEnded);
 }
diff --git a/Engine/Source/Runtime/GameplayAbilities/Private/Abilities/Tasks/AbilityTask_WaitTargetData.cpp b/Engine/Source/Runtime/GameplayAbilities/Private/Abilities/Tasks/AbilityTask_WaitTargetData.cpp
--- a/Engine/Source/Runtime/GameplayAbilities/Private/Abilities/Tasks/AbilityTask_WaitTargetData.cpp
+++ b/Engine/Source/Runtime/GameplayAbilities/Private/Abilities/Tasks/AbilityTask_WaitTargetData.cpp
@@ -10,12 +10,13 @@ UAbilityTask_WaitTargetData::UAbilityTask_WaitTargetData(const class FPostConstr
 	
 }
 
-UAbilityTask_WaitTargetData* UAbilityTask_WaitTargetData::WaitTargetData(UObject* WorldContextObject, TEnumAsByte<EGameplayTargetingConfirmation::Type> ConfirmationType, TSubclassOf<AGameplayAbilityTargetActor> InTargetClass)
+UAbilityTask_WaitTargetData* UAbilityTask_WaitTargetData::WaitTargetData(UObject* WorldContextObject, FName TaskInstanceName, bool CanFireMultipleTimes, TEnumAsByte<EGameplayTargetingConfirmation::Type> ConfirmationType, TSubclassOf<AGameplayAbilityTargetActor> InTargetClass)
 {
-	auto MyObj = NewTask<UAbilityTask_WaitTargetData>(WorldContextObject);
+	auto MyObj = NewTask<UAbilityTask_WaitTargetData>(WorldContextObject, TaskInstanceName);		//Register for task list here, providing a given FName as a key
 	MyObj->TargetClass = InTargetClass;
 	MyObj->ConfirmationType = ConfirmationType;
-	return MyObj;	
+	MyObj->bMultiFire = CanFireMultipleTimes;
+	return MyObj;
 }
 
 // ---------------------------------------------------------------------------------------
@@ -81,10 +82,20 @@ bool UAbilityTask_WaitTargetData::BeginSpawningActor(UObject* WorldContextObject
 			// Register with the TargetData callbacks if we are expecting client to send them
 			if (!CDO->ShouldProduceTargetDataOnServer)
 			{
+				//Problem here - if there's targeting data just sitting around, fire events don't get hooked up because of this if-else, even if we don't end the task.
 				if (AbilitySystemComponent->ReplicatedTargetData.IsValid())
 				{
 					ValidData.Broadcast(AbilitySystemComponent->ReplicatedTargetData);
-					EndTask();
+					if (!bMultiFire)
+					{
+						EndTask();
+					}
+					else
+					{
+						//Since multifire is supported, we still need to hook up the callbacks
+						AbilitySystemComponent->ReplicatedTargetDataDelegate.AddUObject(this, &UAbilityTask_WaitTargetData::OnTargetDataReplicatedCallback);
+						AbilitySystemComponent->ReplicatedTargetDataCancelledDelegate.AddDynamic(this, &UAbilityTask_WaitTargetData::OnTargetDataReplicatedCancelledCallback);
+					}
 				}
 				else
 				{
@@ -171,10 +182,13 @@ void UAbilityTask_WaitTargetData::OnTargetDataReplicatedCallback(FGameplayAbilit
 		ValidData.Broadcast(Data);
 	}
 
-	EndTask();
+	if (!bMultiFire)
+	{
+		EndTask();
+	}
 }
 
-/** Client cancelled this Targeting Task (we are the server */
+/** Client cancelled this Targeting Task (we are the server) */
 void UAbilityTask_WaitTargetData::OnTargetDataReplicatedCancelledCallback()
 {
 	check(AbilitySystemComponent.IsValid());
@@ -195,8 +209,11 @@ void UAbilityTask_WaitTargetData::OnTargetDataReadyCallback(FGameplayAbilityTarg
 	}
 
 	ValidData.Broadcast(Data);
-	
-	EndTask();
+
+	if (!bMultiFire)
+	{
+		EndTask();
+	}
 }
 
 /** The TargetActor we spawned locally has called back with a cancel event (they still include the 'last/best' targetdata but the consumer of this may want to discard it) */
@@ -210,6 +227,21 @@ void UAbilityTask_WaitTargetData::OnTargetDataCancelledCallback(FGameplayAbility
 	EndTask();
 }
 
+/** Called when the ability is asked to confirm from an outside node. What this means depends on the individual task. By default, this does nothing other than ending if bEndTask is true. */
+void UAbilityTask_WaitTargetData::ExternalConfirm(bool bEndTask)
+{
+	check(AbilitySystemComponent.IsValid());
+	if (MyTargetActor.IsValid())
+	{
+		AGameplayAbilityTargetActor* CachedTargetActor = MyTargetActor.Get();
+		if (CachedTargetActor->ShouldProduceTargetData())
+		{
+			CachedTargetActor->ConfirmTargetingAndContinue();
+		}
+	}
+	Super::ExternalConfirm(bEndTask);
+}
+
 void UAbilityTask_WaitTargetData::OnDestroy(bool AbilityEnded)
 {
 	AbilitySystemComponent->ConsumeAbilityTargetData();
diff --git a/Engine/Source/Runtime/GameplayAbilities/Private/AbilitySystemComponent.cpp b/Engine/Source/Runtime/GameplayAbilities/Private/AbilitySystemComponent.cpp
--- a/Engine/Source/Runtime/GameplayAbilities/Private/AbilitySystemComponent.cpp
+++ b/Engine/Source/Runtime/GameplayAbilities/Private/AbilitySystemComponent.cpp
@@ -45,7 +45,6 @@ UAbilitySystemComponent::UAbilitySystemComponent(const class FPostConstructIniti
 	ReplicatedPredictionKey = 0;
 	LocalPredictionKey = 0;
 	UserAbilityActivationInhibited = false;
-	TargetingAbility = NULL;
 }
 
 UAbilitySystemComponent::~UAbilitySystemComponent()
diff --git a/Engine/Source/Runtime/GameplayAbilities/Private/AbilitySystemComponent_Abilities.cpp b/Engine/Source/Runtime/GameplayAbilities/Private/AbilitySystemComponent_Abilities.cpp
--- a/Engine/Source/Runtime/GameplayAbilities/Private/AbilitySystemComponent_Abilities.cpp
+++ b/Engine/Source/Runtime/GameplayAbilities/Private/AbilitySystemComponent_Abilities.cpp
@@ -709,11 +709,6 @@ bool UAbilitySystemComponent::ServerSetReplicatedTargetDataCancelled_Validate()
 
 // -------
 
-void UAbilitySystemComponent::SetTargetAbility(UGameplayAbility* NewTargetingAbility)
-{
-	TargetingAbility = NewTargetingAbility;
-}
-
 void UAbilitySystemComponent::ConsumeAbilityConfirmCancel()
 {
 	ReplicatedConfirmAbility = false;
diff --git a/Engine/Source/Runtime/GameplayAbilities/Public/Abilities/GameplayAbility.h b/Engine/Source/Runtime/GameplayAbilities/Public/Abilities/GameplayAbility.h
--- a/Engine/Source/Runtime/GameplayAbilities/Public/Abilities/GameplayAbility.h
+++ b/Engine/Source/Runtime/GameplayAbilities/Public/Abilities/GameplayAbility.h
@@ -273,6 +273,14 @@ protected:
 
 	/** Destroys intanced-per-execution abilities. Instance-per-actor abilities should 'reset'. Non instance abilities - what can we do? */
 	UFUNCTION(BlueprintCallable, Category = Ability)
+	void ConfirmTaskByInstanceName(FName InstanceName, bool bEndTask);
+
+	/** Internal function, cancels all the tasks we asked to cancel last frame (by instance name). */
+	void CancelTasksByInstanceName();
+	TArray<FName> CancelTaskInstanceNames;
+
+	/** Destroys instanced-per-execution abilities. Instance-per-actor abilities should 'reset'. Non instance abilities - what can we do? */
+	UFUNCTION(BlueprintCallable, Category = Ability)
 	void CancelTaskByInstanceName(FName InstanceName);
 
 	// -------------------------------------
diff --git a/Engine/Source/Runtime/GameplayAbilities/Public/Abilities/GameplayAbilityTargetActor.h b/Engine/Source/Runtime/GameplayAbilities/Public/Abilities/GameplayAbilityTargetActor.h
--- a/Engine/Source/Runtime/GameplayAbilities/Public/Abilities/GameplayAbilityTargetActor.h
+++ b/Engine/Source/Runtime/GameplayAbilities/Public/Abilities/GameplayAbilityTargetActor.h
@@ -4,7 +4,7 @@
 #include "GameplayAbility.h"
 #include "GameplayAbilityTargetActor.generated.h"
 
-
+/** TargetActors are spawned to assist with ability targeting. They are spawned by ability tasks and create/determine the outgoing targeting data passed from one task to another. */
 UCLASS(Blueprintable, abstract, notplaceable)
 class GAMEPLAYABILITIES_API AGameplayAbilityTargetActor : public AActor
 {
@@ -30,6 +30,9 @@ public:
 	/** Initialize and begin targeting logic  */
 	virtual void StartTargeting(UGameplayAbility* Ability);
 
+	/** Requesting targeting data, but not necessarily stopping/destroying the task. Useful for external target data requests. */
+	virtual void ConfirmTargetingAndContinue();
+
 	/** Outside code is saying 'stop and just give me what you have' */
 	UFUNCTION()
 	virtual void ConfirmTargeting();
@@ -61,6 +64,9 @@ public:
 	UPROPERTY()
 	UGameplayAbility* OwningAbility;
 
+	UPROPERTY(BlueprintReadOnly, Replicated, Category = Targeting)
+	bool bDestroyOnConfirmation;
+
 	UPROPERTY(BlueprintReadOnly, Replicated, Category = Targeting)
 	AActor* SourceActor;
 
diff --git a/Engine/Source/Runtime/GameplayAbilities/Public/Abilities/GameplayAbilityTargetActor_Radius.h b/Engine/Source/Runtime/GameplayAbilities/Public/Abilities/GameplayAbilityTargetActor_Radius.h
--- a/Engine/Source/Runtime/GameplayAbilities/Public/Abilities/GameplayAbilityTargetActor_Radius.h
+++ b/Engine/Source/Runtime/GameplayAbilities/Public/Abilities/GameplayAbilityTargetActor_Radius.h
@@ -6,7 +6,8 @@
 #include "GameplayAbilityWorldReticle.h"
 #include "GameplayAbilityTargetActor_Radius.generated.h"
 
-UCLASS(Blueprintable)
+/** Selects everything within a given radius of the source actor. */
+UCLASS(Blueprintable, notplaceable)
 class GAMEPLAYABILITIES_API AGameplayAbilityTargetActor_Radius : public AGameplayAbilityTargetActor
 {
 	GENERATED_UCLASS_BODY()
@@ -15,7 +16,7 @@ public:
 
 	virtual void StartTargeting(UGameplayAbility* Ability);
 	
-	virtual void ConfirmTargeting() override;
+	virtual void ConfirmTargetingAndContinue() override;
 
 	UPROPERTY(BlueprintReadWrite, EditAnywhere, meta = (ExposeOnSpawn = true), Category = Radius)
 	float Radius;
diff --git a/Engine/Source/Runtime/GameplayAbilities/Public/Abilities/GameplayAbilityTargetActor_Trace.h b/Engine/Source/Runtime/GameplayAbilities/Public/Abilities/GameplayAbilityTargetActor_Trace.h
--- a/Engine/Source/Runtime/GameplayAbilities/Public/Abilities/GameplayAbilityTargetActor_Trace.h
+++ b/Engine/Source/Runtime/GameplayAbilities/Public/Abilities/GameplayAbilityTargetActor_Trace.h
@@ -6,7 +6,8 @@
 #include "GameplayAbilityWorldReticle.h"
 #include "GameplayAbilityTargetActor_Trace.generated.h"
 
-UCLASS(Abstract, Blueprintable)
+/** Intermediate base class for all line-trace type targeting actors. */
+UCLASS(Abstract, Blueprintable, notplaceable)
 class GAMEPLAYABILITIES_API AGameplayAbilityTargetActor_Trace : public AGameplayAbilityTargetActor
 {
 	GENERATED_UCLASS_BODY()
@@ -17,8 +18,8 @@ public:
 	virtual FGameplayAbilityTargetDataHandle StaticGetTargetData(UWorld * World, const FGameplayAbilityActorInfo* ActorInfo, FGameplayAbilityActivationInfo ActivationInfo) const override;
 
 	virtual void StartTargeting(UGameplayAbility* Ability) override;
-	
-	virtual void ConfirmTargeting() override;	
+
+	virtual void ConfirmTargetingAndContinue() override;
 
 	virtual void Tick(float DeltaSeconds) override;
 	
diff --git a/Engine/Source/Runtime/GameplayAbilities/Public/Abilities/GameplayAbilityWorldReticle.h b/Engine/Source/Runtime/GameplayAbilities/Public/Abilities/GameplayAbilityWorldReticle.h
--- a/Engine/Source/Runtime/GameplayAbilities/Public/Abilities/GameplayAbilityWorldReticle.h
+++ b/Engine/Source/Runtime/GameplayAbilities/Public/Abilities/GameplayAbilityWorldReticle.h
@@ -5,6 +5,7 @@
 
 class AGameplayAbilityTargetActor;
 
+/** Reticles allow targeting to be visualized. Tasks can spawn these. Artists/designers can create BPs for these. */
 UCLASS(Blueprintable, notplaceable)
 class GAMEPLAYABILITIES_API AGameplayAbilityWorldReticle : public AActor
 {
diff --git a/Engine/Source/Runtime/GameplayAbilities/Public/Abilities/Tasks/AbilityTask.h b/Engine/Source/Runtime/GameplayAbilities/Public/Abilities/Tasks/AbilityTask.h
--- a/Engine/Source/Runtime/GameplayAbilities/Public/Abilities/Tasks/AbilityTask.h
+++ b/Engine/Source/Runtime/GameplayAbilities/Public/Abilities/Tasks/AbilityTask.h
@@ -105,6 +105,9 @@ class GAMEPLAYABILITIES_API UAbilityTask : public UObject
 	/** Called explicitly to end the task (usually by the task itself). Calls OnDestroy. */
 	void EndTask();
 
+	/** Called when the ability is asked to confirm from an outside node. What this means depends on the individual task. By default, this does nothing other than ending if bEndTask is true. */
+	virtual void ExternalConfirm(bool bEndTask);
+
 public:
 
 	/** This name allows us to find the task later so that we can end it. */
diff --git a/Engine/Source/Runtime/GameplayAbilities/Public/Abilities/Tasks/AbilityTask_MoveToLocation.h b/Engine/Source/Runtime/GameplayAbilities/Public/Abilities/Tasks/AbilityTask_MoveToLocation.h
--- a/Engine/Source/Runtime/GameplayAbilities/Public/Abilities/Tasks/AbilityTask_MoveToLocation.h
+++ b/Engine/Source/Runtime/GameplayAbilities/Public/Abilities/Tasks/AbilityTask_MoveToLocation.h
@@ -20,17 +20,16 @@ class UAbilityTask_MoveToLocation : public UAbilityTask
 	UPROPERTY(BlueprintAssignable)
 	FMoveToLocationDelegate		OnTargetLocationReached;
 
-	void InterpolatePosition();
-
 	/** Move to the specified location, using the curve (range 0 - 1) or linearly if no curve is specified */
 	UFUNCTION(BlueprintCallable, Category = "Ability|Tasks", meta = (HidePin = "WorldContextObject", DefaultToSelf = "WorldContextObject", BlueprintInternalUseOnly = "TRUE"))
 	static UAbilityTask_MoveToLocation* MoveToLocation(UObject* WorldContextObject, FVector Location, float Duration, UCurveFloat* OptionalInterpolationCurve);
 
 	virtual void Activate() override;
 
-protected:
+	/** Tick function for this task, if bTickingTask == true */
+	virtual void TickTask(float DeltaTime) override;
 
-	virtual void OnDestroy(bool AbilityEnded) override;
+protected:
 
 	FVector StartLocation;
 	FVector TargetLocation;
diff --git a/Engine/Source/Runtime/GameplayAbilities/Public/Abilities/Tasks/AbilityTask_WaitTargetData.h b/Engine/Source/Runtime/GameplayAbilities/Public/Abilities/Tasks/AbilityTask_WaitTargetData.h
--- a/Engine/Source/Runtime/GameplayAbilities/Public/Abilities/Tasks/AbilityTask_WaitTargetData.h
+++ b/Engine/Source/Runtime/GameplayAbilities/Public/Abilities/Tasks/AbilityTask_WaitTargetData.h
@@ -9,7 +9,8 @@ class AGameplayAbilityTargetActor;
 
 DECLARE_DYNAMIC_MULTICAST_DELEGATE_OneParam(FWaitTargetDataDelegate, FGameplayAbilityTargetDataHandle, Data);
 
-UCLASS(MinimalAPI)
+/** Wait for targeting actor (spawned from parameter) to provide data. Can be set not to end upon outputting data. Can be ended by task name. */
+UCLASS(MinimalAPI, notplaceable)
 class UAbilityTask_WaitTargetData: public UAbilityTask
 {
 	GENERATED_UCLASS_BODY()
@@ -20,6 +21,10 @@ class UAbilityTask_WaitTargetData: public UAbilityTask
 	UPROPERTY(BlueprintAssignable)
 	FWaitTargetDataDelegate	Cancelled;
 
+	/** If true, will not call EndTask when target data is created/received. */
+	UPROPERTY()
+	bool bMultiFire;
+
 	UFUNCTION()
 	void OnTargetDataReplicatedCallback(FGameplayAbilityTargetDataHandle Data);
 
@@ -34,7 +39,7 @@ class UAbilityTask_WaitTargetData: public UAbilityTask
 
 	/** Spawns Targeting actor and waits for it to return valid data or to be cancelled. */
 	UFUNCTION(BlueprintCallable, meta=(HidePin = "WorldContextObject", DefaultToSelf = "WorldContextObject", BlueprintInternalUseOnly = "true", HideSpawnParms="Instigator"), Category="Ability|Tasks")
-	static UAbilityTask_WaitTargetData* WaitTargetData(UObject* WorldContextObject, TEnumAsByte<EGameplayTargetingConfirmation::Type> ConfirmationType, TSubclassOf<AGameplayAbilityTargetActor> Class);
+	static UAbilityTask_WaitTargetData* WaitTargetData(UObject* WorldContextObject, FName TaskInstanceName, bool CanFireMultipleTimes, TEnumAsByte<EGameplayTargetingConfirmation::Type> ConfirmationType, TSubclassOf<AGameplayAbilityTargetActor> Class);
 
 	UFUNCTION(BlueprintCallable, meta = (HidePin = "WorldContextObject", DefaultToSelf = "WorldContextObject", BlueprintInternalUseOnly = "true"), Category = "Abilities")
 	bool BeginSpawningActor(UObject* WorldContextObject, TSubclassOf<AGameplayAbilityTargetActor> Class, AGameplayAbilityTargetActor*& SpawnedActor);
@@ -42,6 +47,9 @@ class UAbilityTask_WaitTargetData: public UAbilityTask
 	UFUNCTION(BlueprintCallable, meta = (HidePin = "WorldContextObject", DefaultToSelf = "WorldContextObject", BlueprintInternalUseOnly = "true"), Category = "Abilities")
 	void FinishSpawningActor(UObject* WorldContextObject, AGameplayAbilityTargetActor* SpawnedActor);
 
+	/** Called when the ability is asked to confirm from an outside node. What this means depends on the individual task. By default, this does nothing other than ending if bEndTask is true. */
+	virtual void ExternalConfirm(bool bEndTask) override;
+
 protected:
 
 	virtual void OnDestroy(bool AbilityEnded) override;
diff --git a/Engine/Source/Runtime/GameplayAbilities/Public/AbilitySystemComponent.h b/Engine/Source/Runtime/GameplayAbilities/Public/AbilitySystemComponent.h
--- a/Engine/Source/Runtime/GameplayAbilities/Public/AbilitySystemComponent.h
+++ b/Engine/Source/Runtime/GameplayAbilities/Public/AbilitySystemComponent.h
@@ -468,7 +468,7 @@ class GAMEPLAYABILITIES_API UAbilitySystemComponent : public UActorComponent, pu
 	UFUNCTION(BlueprintCallable, Category = "Abilities")
 	void TargetConfirm();
 
-	/** Any active targeting actors will be stopped and cancelled, not returning any targeting data */
+	/** Any active targeting actors will be stopped and canceled, not returning any targeting data */
 	UFUNCTION(BlueprintCallable, Category = "Abilities")
 	void TargetCancel();
 
@@ -536,16 +536,10 @@ class GAMEPLAYABILITIES_API UAbilitySystemComponent : public UActorComponent, pu
 	void ServerSetReplicatedTargetDataCancelled();
 	
 
-	void SetTargetAbility(UGameplayAbility* NewTargetingAbility);
-
 	void ConsumeAbilityConfirmCancel();
 
 	void ConsumeAbilityTargetData();
 
-	/** This ability has a 'lock' on replicated targeting data. Only 1 targeting ability can be active at once */
-	UPROPERTY()
-	UGameplayAbility* TargetingAbility;
-
 	bool ReplicatedConfirmAbility;
 	bool ReplicatedCancelAbility;
 
