Added external cancel capability. Added CustomMulti confirmation type. Including a code-only missile swarm asset to keep up-to-date.

diff --git a/Engine/Source/Runtime/GameplayAbilities/Private/Abilities/GameplayAbility.cpp b/Engine/Source/Runtime/GameplayAbilities/Private/Abilities/GameplayAbility.cpp
--- a/Engine/Source/Runtime/GameplayAbilities/Private/Abilities/GameplayAbility.cpp
+++ b/Engine/Source/Runtime/GameplayAbilities/Private/Abilities/GameplayAbility.cpp
@@ -635,8 +635,23 @@ void UGameplayAbility::ConfirmTaskByInstanceName(FName InstanceName, bool bEndTa
 	}
 }
 
-void UGameplayAbility::CancelTasksByInstanceName()
+void UGameplayAbility::EndOrCancelTasksByInstanceName()
 {
+	for (int32 j = 0; j < EndTaskInstanceNames.Num(); ++j)
+	{
+		FName InstanceName = EndTaskInstanceNames[j];
+		TArray<TWeakObjectPtr<UAbilityTask>> NamedTasks = ActiveTasks.FilterByPredicate<FAbilityInstanceNamePredicate>(FAbilityInstanceNamePredicate(InstanceName));
+		for (int32 i = NamedTasks.Num() - 1; i >= 0; --i)
+		{
+			if (NamedTasks[i].IsValid())
+			{
+				UAbilityTask* CurrentTask = NamedTasks[i].Get();
+				CurrentTask->EndTask();
+			}
+		}
+	}
+	EndTaskInstanceNames.Empty();
+
 	for (int32 j = 0; j < CancelTaskInstanceNames.Num(); ++j)
 	{
 		FName InstanceName = CancelTaskInstanceNames[j];
@@ -646,19 +661,25 @@ void UGameplayAbility::CancelTasksByInstanceName()
 			if (NamedTasks[i].IsValid())
 			{
 				UAbilityTask* CurrentTask = NamedTasks[i].Get();
-				//TODO Change this to a proper cancel instead of just ending the task cold. Maybe make a separate function for calling EndTask so we don't lose this functionality.
-				CurrentTask->EndTask();
+				CurrentTask->ExternalCancel();
 			}
 		}
 	}
 	CancelTaskInstanceNames.Empty();
 }
 
+void UGameplayAbility::EndTaskByInstanceName(FName InstanceName)
+{
+	//Avoid race condition by delaying for one frame
+	EndTaskInstanceNames.AddUnique(InstanceName);
+	GetWorld()->GetTimerManager().SetTimerForNextTick(this, &UGameplayAbility::EndOrCancelTasksByInstanceName);
+}
+
 void UGameplayAbility::CancelTaskByInstanceName(FName InstanceName)
 {
 	//Avoid race condition by delaying for one frame
 	CancelTaskInstanceNames.AddUnique(InstanceName);
-	GetWorld()->GetTimerManager().SetTimerForNextTick(this, &UGameplayAbility::CancelTasksByInstanceName);
+	GetWorld()->GetTimerManager().SetTimerForNextTick(this, &UGameplayAbility::EndOrCancelTasksByInstanceName);
 }
 
 void UGameplayAbility::TaskEnded(UAbilityTask* Task)
diff --git a/Engine/Source/Runtime/GameplayAbilities/Private/Abilities/Tasks/AbilityTask.cpp b/Engine/Source/Runtime/GameplayAbilities/Private/Abilities/Tasks/AbilityTask.cpp
--- a/Engine/Source/Runtime/GameplayAbilities/Private/Abilities/Tasks/AbilityTask.cpp
+++ b/Engine/Source/Runtime/GameplayAbilities/Private/Abilities/Tasks/AbilityTask.cpp
@@ -75,6 +75,11 @@ void UAbilityTask::ExternalConfirm(bool bEndTask)
 	}
 }
 
+void UAbilityTask::ExternalCancel()
+{
+	EndTask();
+}
+
 void UAbilityTask::OnDestroy(bool AbilityIsEnding)
 {
 	ensure(!IsPendingKill());
diff --git a/Engine/Source/Runtime/GameplayAbilities/Private/Abilities/Tasks/AbilityTask_WaitTargetData.cpp b/Engine/Source/Runtime/GameplayAbilities/Private/Abilities/Tasks/AbilityTask_WaitTargetData.cpp
--- a/Engine/Source/Runtime/GameplayAbilities/Private/Abilities/Tasks/AbilityTask_WaitTargetData.cpp
+++ b/Engine/Source/Runtime/GameplayAbilities/Private/Abilities/Tasks/AbilityTask_WaitTargetData.cpp
@@ -10,12 +10,11 @@ UAbilityTask_WaitTargetData::UAbilityTask_WaitTargetData(const class FPostConstr
 	
 }
 
-UAbilityTask_WaitTargetData* UAbilityTask_WaitTargetData::WaitTargetData(UObject* WorldContextObject, FName TaskInstanceName, bool CanFireMultipleTimes, TEnumAsByte<EGameplayTargetingConfirmation::Type> ConfirmationType, TSubclassOf<AGameplayAbilityTargetActor> InTargetClass)
+UAbilityTask_WaitTargetData* UAbilityTask_WaitTargetData::WaitTargetData(UObject* WorldContextObject, FName TaskInstanceName, TEnumAsByte<EGameplayTargetingConfirmation::Type> ConfirmationType, TSubclassOf<AGameplayAbilityTargetActor> InTargetClass)
 {
 	auto MyObj = NewTask<UAbilityTask_WaitTargetData>(WorldContextObject, TaskInstanceName);		//Register for task list here, providing a given FName as a key
 	MyObj->TargetClass = InTargetClass;
 	MyObj->ConfirmationType = ConfirmationType;
-	MyObj->bMultiFire = CanFireMultipleTimes;
 	return MyObj;
 }
 
@@ -86,16 +85,16 @@ bool UAbilityTask_WaitTargetData::BeginSpawningActor(UObject* WorldContextObject
 				if (AbilitySystemComponent->ReplicatedTargetData.IsValid())
 				{
 					ValidData.Broadcast(AbilitySystemComponent->ReplicatedTargetData);
-					if (!bMultiFire)
-					{
-						EndTask();
-					}
-					else
+					if (ConfirmationType == EGameplayTargetingConfirmation::CustomMulti)
 					{
 						//Since multifire is supported, we still need to hook up the callbacks
 						AbilitySystemComponent->ReplicatedTargetDataDelegate.AddUObject(this, &UAbilityTask_WaitTargetData::OnTargetDataReplicatedCallback);
 						AbilitySystemComponent->ReplicatedTargetDataCancelledDelegate.AddDynamic(this, &UAbilityTask_WaitTargetData::OnTargetDataReplicatedCancelledCallback);
 					}
+					else
+					{
+						EndTask();
+					}
 				}
 				else
 				{
@@ -182,13 +181,13 @@ void UAbilityTask_WaitTargetData::OnTargetDataReplicatedCallback(FGameplayAbilit
 		ValidData.Broadcast(Data);
 	}
 
-	if (!bMultiFire)
+	if (ConfirmationType != EGameplayTargetingConfirmation::CustomMulti)
 	{
 		EndTask();
 	}
 }
 
-/** Client cancelled this Targeting Task (we are the server) */
+/** Client canceled this Targeting Task (we are the server) */
 void UAbilityTask_WaitTargetData::OnTargetDataReplicatedCancelledCallback()
 {
 	check(AbilitySystemComponent.IsValid());
@@ -210,7 +209,7 @@ void UAbilityTask_WaitTargetData::OnTargetDataReadyCallback(FGameplayAbilityTarg
 
 	ValidData.Broadcast(Data);
 
-	if (!bMultiFire)
+	if (ConfirmationType != EGameplayTargetingConfirmation::CustomMulti)
 	{
 		EndTask();
 	}
@@ -242,6 +241,14 @@ void UAbilityTask_WaitTargetData::ExternalConfirm(bool bEndTask)
 	Super::ExternalConfirm(bEndTask);
 }
 
+/** Called when the ability is asked to confirm from an outside node. What this means depends on the individual task. By default, this does nothing other than ending if bEndTask is true. */
+void UAbilityTask_WaitTargetData::ExternalCancel()
+{
+	check(AbilitySystemComponent.IsValid());
+	Cancelled.Broadcast(FGameplayAbilityTargetDataHandle());
+	Super::ExternalCancel();
+}
+
 void UAbilityTask_WaitTargetData::OnDestroy(bool AbilityEnded)
 {
 	AbilitySystemComponent->ConsumeAbilityTargetData();
diff --git a/Engine/Source/Runtime/GameplayAbilities/Public/Abilities/GameplayAbility.h b/Engine/Source/Runtime/GameplayAbilities/Public/Abilities/GameplayAbility.h
--- a/Engine/Source/Runtime/GameplayAbilities/Public/Abilities/GameplayAbility.h
+++ b/Engine/Source/Runtime/GameplayAbilities/Public/Abilities/GameplayAbility.h
@@ -276,9 +276,14 @@ protected:
 	void ConfirmTaskByInstanceName(FName InstanceName, bool bEndTask);
 
 	/** Internal function, cancels all the tasks we asked to cancel last frame (by instance name). */
-	void CancelTasksByInstanceName();
+	void EndOrCancelTasksByInstanceName();
 	TArray<FName> CancelTaskInstanceNames;
 
+	/** Internal function, cancels all the tasks we asked to cancel last frame (by instance name). */
+	UFUNCTION(BlueprintCallable, Category = Ability)
+	void EndTaskByInstanceName(FName InstanceName);
+	TArray<FName> EndTaskInstanceNames;
+
 	/** Destroys instanced-per-execution abilities. Instance-per-actor abilities should 'reset'. Non instance abilities - what can we do? */
 	UFUNCTION(BlueprintCallable, Category = Ability)
 	void CancelTaskByInstanceName(FName InstanceName);
diff --git a/Engine/Source/Runtime/GameplayAbilities/Public/Abilities/GameplayAbilityTargetTypes.h b/Engine/Source/Runtime/GameplayAbilities/Public/Abilities/GameplayAbilityTargetTypes.h
--- a/Engine/Source/Runtime/GameplayAbilities/Public/Abilities/GameplayAbilityTargetTypes.h
+++ b/Engine/Source/Runtime/GameplayAbilities/Public/Abilities/GameplayAbilityTargetTypes.h
@@ -20,6 +20,7 @@ namespace EGameplayTargetingConfirmation
 		Instant,			// The targeting happens instantly without special logic or user input deciding when to 'fire'.
 		UserConfirmed,		// The targeting happens when the user confirms the targeting.
 		Custom,				// The GameplayTargeting Ability is responsible for deciding when the targeting data is ready. Not supported by all TargetingActors.
+		CustomMulti,		// The GameplayTargeting Ability is responsible for deciding when the targeting data is ready. Not supported by all TargetingActors. Should not destroy upon data production.
 	};
 }
 
diff --git a/Engine/Source/Runtime/GameplayAbilities/Public/Abilities/Tasks/AbilityTask.h b/Engine/Source/Runtime/GameplayAbilities/Public/Abilities/Tasks/AbilityTask.h
--- a/Engine/Source/Runtime/GameplayAbilities/Public/Abilities/Tasks/AbilityTask.h
+++ b/Engine/Source/Runtime/GameplayAbilities/Public/Abilities/Tasks/AbilityTask.h
@@ -108,6 +108,9 @@ class GAMEPLAYABILITIES_API UAbilityTask : public UObject
 	/** Called when the ability is asked to confirm from an outside node. What this means depends on the individual task. By default, this does nothing other than ending if bEndTask is true. */
 	virtual void ExternalConfirm(bool bEndTask);
 
+	/** Called when the ability is asked to cancel from an outside node. What this means depends on the individual task. By default, this does nothing other than ending the task. */
+	virtual void ExternalCancel();
+
 public:
 
 	/** This name allows us to find the task later so that we can end it. */
diff --git a/Engine/Source/Runtime/GameplayAbilities/Public/Abilities/Tasks/AbilityTask_WaitTargetData.h b/Engine/Source/Runtime/GameplayAbilities/Public/Abilities/Tasks/AbilityTask_WaitTargetData.h
--- a/Engine/Source/Runtime/GameplayAbilities/Public/Abilities/Tasks/AbilityTask_WaitTargetData.h
+++ b/Engine/Source/Runtime/GameplayAbilities/Public/Abilities/Tasks/AbilityTask_WaitTargetData.h
@@ -21,10 +21,6 @@ class UAbilityTask_WaitTargetData: public UAbilityTask
 	UPROPERTY(BlueprintAssignable)
 	FWaitTargetDataDelegate	Cancelled;
 
-	/** If true, will not call EndTask when target data is created/received. */
-	UPROPERTY()
-	bool bMultiFire;
-
 	UFUNCTION()
 	void OnTargetDataReplicatedCallback(FGameplayAbilityTargetDataHandle Data);
 
@@ -39,7 +35,7 @@ class UAbilityTask_WaitTargetData: public UAbilityTask
 
 	/** Spawns Targeting actor and waits for it to return valid data or to be cancelled. */
 	UFUNCTION(BlueprintCallable, meta=(HidePin = "WorldContextObject", DefaultToSelf = "WorldContextObject", BlueprintInternalUseOnly = "true", HideSpawnParms="Instigator"), Category="Ability|Tasks")
-	static UAbilityTask_WaitTargetData* WaitTargetData(UObject* WorldContextObject, FName TaskInstanceName, bool CanFireMultipleTimes, TEnumAsByte<EGameplayTargetingConfirmation::Type> ConfirmationType, TSubclassOf<AGameplayAbilityTargetActor> Class);
+	static UAbilityTask_WaitTargetData* WaitTargetData(UObject* WorldContextObject, FName TaskInstanceName, TEnumAsByte<EGameplayTargetingConfirmation::Type> ConfirmationType, TSubclassOf<AGameplayAbilityTargetActor> Class);
 
 	UFUNCTION(BlueprintCallable, meta = (HidePin = "WorldContextObject", DefaultToSelf = "WorldContextObject", BlueprintInternalUseOnly = "true"), Category = "Abilities")
 	bool BeginSpawningActor(UObject* WorldContextObject, TSubclassOf<AGameplayAbilityTargetActor> Class, AGameplayAbilityTargetActor*& SpawnedActor);
@@ -50,6 +46,9 @@ class UAbilityTask_WaitTargetData: public UAbilityTask
 	/** Called when the ability is asked to confirm from an outside node. What this means depends on the individual task. By default, this does nothing other than ending if bEndTask is true. */
 	virtual void ExternalConfirm(bool bEndTask) override;
 
+	/** Called when the ability is asked to cancel from an outside node. What this means depends on the individual task. By default, this does nothing other than ending the task. */
+	virtual void ExternalCancel() override;
+
 protected:
 
 	virtual void OnDestroy(bool AbilityEnded) override;
