Shuffled some type definitions between files.

diff --git a/Engine/Source/Runtime/GameplayAbilities/Public/Abilities/GameplayAbility.h b/Engine/Source/Runtime/GameplayAbilities/Public/Abilities/GameplayAbility.h
--- a/Engine/Source/Runtime/GameplayAbilities/Public/Abilities/GameplayAbility.h
+++ b/Engine/Source/Runtime/GameplayAbilities/Public/Abilities/GameplayAbility.h
@@ -5,6 +5,7 @@
 #include "GameplayTagContainer.h"
 #include "GameplayTagAssetInterface.h"
 #include "Runtime/Engine/Classes/Animation/AnimInstance.h"
+#include "GameplayAbilityTargetTypes.h"
 #include "GameplayAbilityTypes.h"
 #include "GameplayEffect.h"
 #include "GameplayAbility.generated.h"
diff --git a/Engine/Source/Runtime/GameplayAbilities/Public/Abilities/GameplayAbilityTargetTypes.h b/Engine/Source/Runtime/GameplayAbilities/Public/Abilities/GameplayAbilityTargetTypes.h
--- a/Engine/Source/Runtime/GameplayAbilities/Public/Abilities/GameplayAbilityTargetTypes.h
+++ b/Engine/Source/Runtime/GameplayAbilities/Public/Abilities/GameplayAbilityTargetTypes.h
@@ -25,6 +25,180 @@ namespace EGameplayTargetingConfirmation
 }
 
 
+/**
+*	A generic structure for targeting data. We want generic functions to produce this data and other generic
+*	functions to consume this data.
+*
+*	We expect this to be able to hold specific actors/object reference and also generic location/direction/origin
+*	information.
+*
+*	Some example producers:
+*		-Overlap/Hit collision event generates TargetData about who was hit in a melee attack
+*		-A mouse input causes a hit trace and the actor infront of the crosshair is turned into TargetData
+*		-A mouse input causes TargetData to be generated from the owner's crosshair view origin/direction
+*		-An AOE/aura pulses and all actors in a radius around the instigator are added to TargetData
+*		-Panzer Dragoon style 'painting' targeting mode
+*		-MMORPG style ground AOE targeting style (potentially both a location on the ground and actors that were targeted)
+*
+*	Some example consumers:
+*		-Apply a GameplayEffect to all actors in TargetData
+*		-Find closest actor from all in TargetData
+*		-Call some function on all actors in TargetData
+*		-Filter or merge TargetDatas
+*		-Spawn a new actor at a TargetData location
+*
+*
+*
+*	Maybe it is better to distinguish between actor list targeting vs positional targeting data?
+*		-AOE/aura type of targeting data blurs the line
+*
+*
+*/
+
+USTRUCT()
+struct GAMEPLAYABILITIES_API FGameplayAbilityTargetData
+{
+	GENERATED_USTRUCT_BODY()
+
+	virtual ~FGameplayAbilityTargetData() { }
+
+	void ApplyGameplayEffect(UGameplayEffect *GameplayEffect, const struct FGameplayAbilityActorInfo InstigatorInfo);
+
+	virtual TArray<AActor*>	GetActors() const
+	{
+		return TArray<AActor*>();
+	}
+
+	// -------------------------------------
+
+	virtual bool HasHitResult() const
+	{
+		return false;
+	}
+
+	virtual const FHitResult* GetHitResult() const
+	{
+		return NULL;
+	}
+
+	// -------------------------------------
+
+	virtual bool HasOrigin() const
+	{
+		return false;
+	}
+
+	virtual FTransform GetOrigin() const
+	{
+		return FTransform::Identity;
+	}
+
+	// -------------------------------------
+
+	virtual bool HasEndPoint() const
+	{
+		return false;
+	}
+
+	virtual FVector GetEndPoint() const
+	{
+		return FVector::ZeroVector;
+	}
+
+	// -------------------------------------
+
+	virtual UScriptStruct* GetScriptStruct()
+	{
+		return FGameplayAbilityTargetData::StaticStruct();
+	}
+
+	virtual FString ToString() const;
+};
+
+
+UENUM(BlueprintType)
+namespace EGameplayAbilityTargetingLocationType
+{
+	/**
+	*	What type of location calculation to use when an ability asks for our transform.
+	*/
+
+	enum Type
+	{
+		LiteralTransform		UMETA(DisplayName = "Literal Transform"),		// We report an actual raw transform. This is also the final fallback if other methods fail.
+		ActorTransform			UMETA(DisplayName = "Actor Transform"),			// We pull the transform from an associated actor directly.
+		SocketTransform			UMETA(DisplayName = "Socket Transform"),		// We aim from a named socket on the player's skeletal mesh component.
+	};
+}
+
+/**
+*	Handle for Targeting Data. This servers two main purposes:
+*		-Avoid us having to copy around the full targeting data structure in Blueprints
+*		-Allows us to leverage polymorphism in the target data structure
+*		-Allows us to implement NetSerialize and replicate by value between clients/server
+*
+*		-Avoid using UObjects could be used to give us polymorphism and by reference passing in blueprints.
+*		-However we would still be screwed when it came to replication
+*
+*		-Replication by value
+*		-Pass by reference in blueprints
+*		-Polymophism in TargetData structure
+*
+*/
+
+USTRUCT(BlueprintType)
+struct GAMEPLAYABILITIES_API FGameplayAbilityTargetDataHandle
+{
+	GENERATED_USTRUCT_BODY()
+
+	FGameplayAbilityTargetDataHandle() { }
+	FGameplayAbilityTargetDataHandle(struct FGameplayAbilityTargetData* DataPtr)
+		: Data(DataPtr)
+	{
+
+	}
+
+	TSharedPtr<FGameplayAbilityTargetData>	Data;
+
+	void Clear()
+	{
+		Data.Reset();
+	}
+
+	bool IsValid() const
+	{
+		return Data.IsValid();
+	}
+
+	bool NetSerialize(FArchive& Ar, class UPackageMap* Map, bool& bOutSuccess);
+
+	/** Comparison operator */
+	bool operator==(FGameplayAbilityTargetDataHandle const& Other) const
+	{
+		// Both invalid structs or both valid and Pointer compare (???) // deep comparison equality
+		bool bBothValid = IsValid() && Other.IsValid();
+		bool bBothInvalid = !IsValid() && !Other.IsValid();
+		return (bBothInvalid || (bBothValid && (Data.Get() == Other.Data.Get())));
+	}
+
+	/** Comparison operator */
+	bool operator!=(FGameplayAbilityTargetDataHandle const& Other) const
+	{
+		return !(FGameplayAbilityTargetDataHandle::operator==(Other));
+	}
+};
+
+template<>
+struct TStructOpsTypeTraits<FGameplayAbilityTargetDataHandle> : public TStructOpsTypeTraitsBase
+{
+	enum
+	{
+		WithCopy = true,		// Necessary so that TSharedPtr<FGameplayAbilityTargetData> Data is copied around
+		WithNetSerializer = true,
+		WithIdenticalViaEquality = true,
+	};
+};
+
 USTRUCT(BlueprintType)
 struct GAMEPLAYABILITIES_API FGameplayAbilityTargetData_Radius : public FGameplayAbilityTargetData
 {
@@ -34,7 +208,7 @@ struct GAMEPLAYABILITIES_API FGameplayAbilityTargetData_Radius : public FGamepla
 	: Origin(0.f) { }
 
 	FGameplayAbilityTargetData_Radius(const TArray<AActor*>& InActors, const FVector& InOrigin)
-	: Actors(InActors), Origin(InOrigin) { }
+		: Actors(InActors), Origin(InOrigin) { }
 
 	virtual TArray<AActor*>	GetActors() const { return Actors; }
 
@@ -62,4 +236,361 @@ struct TStructOpsTypeTraits<FGameplayAbilityTargetData_Radius> : public TStructO
 	{
 		WithNetSerializer = true	// For now this is REQUIRED for FGameplayAbilityTargetDataHandle net serialization to work
 	};
-};
\ No newline at end of file
+};
+
+USTRUCT(BlueprintType)
+struct GAMEPLAYABILITIES_API FGameplayAbilityTargetingLocationInfo
+{
+	GENERATED_USTRUCT_BODY()
+
+	FGameplayAbilityTargetingLocationInfo()
+	: LocationType(EGameplayAbilityTargetingLocationType::LiteralTransform)
+	{
+	};
+
+	void operator=(const FGameplayAbilityTargetingLocationInfo& Other)
+	{
+		LocationType = Other.LocationType;
+		LiteralTransform = Other.LiteralTransform;
+		SourceActor = Other.SourceActor;
+		SourceComponent = Other.SourceComponent;
+		SourceSocketName = Other.SourceSocketName;
+	}
+
+public:
+	FTransform GetTargetingTransform() const
+	{
+		//Return or calculate based on LocationType.
+		switch (LocationType)
+		{
+		case EGameplayAbilityTargetingLocationType::ActorTransform:
+			if (SourceActor)
+			{
+				return SourceActor->GetTransform();
+			}
+			break;
+		case EGameplayAbilityTargetingLocationType::SocketTransform:
+			if (SourceComponent)
+			{
+				return SourceComponent->GetSocketTransform(SourceSocketName);		//Bad socket name will just return component transform anyway, so we're safe
+			}
+			break;
+		case EGameplayAbilityTargetingLocationType::LiteralTransform:
+			return LiteralTransform;
+		default:
+			check(false);		//This case should not happen
+			break;
+		}
+		//Error
+		return FTransform::Identity;
+	}
+
+	FGameplayAbilityTargetDataHandle MakeTargetDataHandleFromHitResult(TWeakObjectPtr<UGameplayAbility> Ability, FHitResult HitResult) const;
+
+	FGameplayAbilityTargetDataHandle MakeTargetDataHandleFromActors(TArray<AActor*> TargetActors) const;
+
+	/** Type of location used - will determine what data is transmitted over the network and what fields are used when calculating position. */
+	UPROPERTY(BlueprintReadWrite, meta = (ExposeOnSpawn = true), Category = Targeting)
+		TEnumAsByte<EGameplayAbilityTargetingLocationType::Type> LocationType;
+
+	/** A literal world transform can be used, if one has been calculated outside of the actor using the ability. */
+	UPROPERTY(BlueprintReadWrite, meta = (ExposeOnSpawn = true), Category = Targeting)
+		FTransform LiteralTransform;
+
+	/** A source actor is needed for Actor-based targeting, but not for Socket-based targeting. */
+	UPROPERTY(BlueprintReadWrite, meta = (ExposeOnSpawn = true), Category = Targeting)
+		AActor* SourceActor;
+
+	/** Socket-based targeting requires a skeletal mesh component to check for the named socket. */
+	UPROPERTY(BlueprintReadWrite, meta = (ExposeOnSpawn = true), Category = Targeting)
+		USkeletalMeshComponent* SourceComponent;
+
+	/** If SourceComponent is valid, this is the name of the socket transform that will be used. If no Socket is provided, SourceComponent's transform will be used. */
+	UPROPERTY(BlueprintReadWrite, meta = (ExposeOnSpawn = true), Category = Targeting)
+		FName SourceSocketName;
+
+	// -------------------------------------
+
+	virtual FString ToString() const
+	{
+		return TEXT("FGameplayAbilityTargetingLocationInfo");
+	}
+
+	bool NetSerialize(FArchive& Ar, class UPackageMap* Map, bool& bOutSuccess);
+
+	virtual UScriptStruct* GetScriptStruct()
+	{
+		return FGameplayAbilityTargetingLocationInfo::StaticStruct();
+	}
+};
+
+template<>
+struct TStructOpsTypeTraits<FGameplayAbilityTargetingLocationInfo> : public TStructOpsTypeTraitsBase
+{
+	enum
+	{
+		WithNetSerializer = true	// For now this is REQUIRED for FGameplayAbilityTargetDataHandle net serialization to work
+	};
+};
+
+USTRUCT(BlueprintType)
+struct GAMEPLAYABILITIES_API FGameplayAbilityTargetData_LocationInfo : public FGameplayAbilityTargetData
+{
+	GENERATED_USTRUCT_BODY()
+
+	/** Generic location data for source */
+	UPROPERTY(BlueprintReadWrite, EditAnywhere, Category = Targeting)
+	FGameplayAbilityTargetingLocationInfo SourceLocation;
+
+	/** Generic location data for target */
+	UPROPERTY(BlueprintReadWrite, EditAnywhere, Category = Targeting)
+		FGameplayAbilityTargetingLocationInfo TargetLocation;
+
+	// -------------------------------------
+
+	virtual bool HasOrigin() const
+	{
+		return true;
+	}
+
+	virtual FTransform GetOrigin() const
+	{
+		return SourceLocation.GetTargetingTransform();
+	}
+
+	// -------------------------------------
+
+	virtual bool HasEndPoint() const
+	{
+		return true;
+	}
+
+	virtual FVector GetEndPoint() const
+	{
+		return TargetLocation.GetTargetingTransform().GetLocation();
+	}
+
+	// -------------------------------------
+
+	virtual UScriptStruct* GetScriptStruct()
+	{
+		return FGameplayAbilityTargetData_LocationInfo::StaticStruct();
+	}
+
+	virtual FString ToString() const override
+	{
+		return TEXT("FGameplayAbilityTargetData_LocationInfo");
+	}
+
+	bool NetSerialize(FArchive& Ar, class UPackageMap* Map, bool& bOutSuccess);
+};
+
+template<>
+struct TStructOpsTypeTraits<FGameplayAbilityTargetData_LocationInfo> : public TStructOpsTypeTraitsBase
+{
+	enum
+	{
+		WithNetSerializer = true	// For now this is REQUIRED for FGameplayAbilityTargetDataHandle net serialization to work
+	};
+};
+
+USTRUCT(BlueprintType)
+struct GAMEPLAYABILITIES_API FGameplayAbilityTargetData_ActorArray : public FGameplayAbilityTargetData
+{
+	GENERATED_USTRUCT_BODY()
+
+	/** We could be selecting this group of actors from any type of location, so use a generic location type */
+	UPROPERTY(BlueprintReadWrite, EditAnywhere, Category = Targeting)
+	FGameplayAbilityTargetingLocationInfo SourceLocation;
+
+	/** Rather than targeting a single point, this type of targeting selects multiple actors. */
+	UPROPERTY(BlueprintReadWrite, EditAnywhere, Category = Targeting)
+		TArray<AActor*> TargetActorArray;
+
+	virtual TArray<AActor*>	GetActors() const override
+	{
+		return TargetActorArray;
+	}
+
+	// -------------------------------------
+
+	virtual bool HasOrigin() const override
+	{
+		return true;
+	}
+
+	virtual FTransform GetOrigin() const override
+	{
+		return SourceLocation.GetTargetingTransform();		//No single aim-at target, since this is aimed at multiple characters.
+	}
+
+	// -------------------------------------
+
+	virtual UScriptStruct* GetScriptStruct() override
+	{
+		return FGameplayAbilityTargetData_ActorArray::StaticStruct();
+	}
+
+	virtual FString ToString() const override
+	{
+		return TEXT("FGameplayAbilityTargetData_ActorArray");
+	}
+
+	bool NetSerialize(FArchive& Ar, class UPackageMap* Map, bool& bOutSuccess);
+};
+
+template<>
+struct TStructOpsTypeTraits<FGameplayAbilityTargetData_ActorArray> : public TStructOpsTypeTraitsBase
+{
+	enum
+	{
+		WithNetSerializer = true	// For now this is REQUIRED for FGameplayAbilityTargetDataHandle net serialization to work
+	};
+};
+
+USTRUCT(BlueprintType)
+struct GAMEPLAYABILITIES_API FGameplayAbilityTargetData_Mesh : public FGameplayAbilityTargetData
+{
+	GENERATED_USTRUCT_BODY()
+
+	FGameplayAbilityTargetData_Mesh()
+	: SourceActor(NULL)
+	, SourceComponent(NULL)
+	{}
+
+
+	/** Actor who owns the named component. Actor's location is used as start point if component cannot be found. */
+	UPROPERTY(BlueprintReadWrite, EditAnywhere, Category = Targeting)
+		AActor* SourceActor;
+
+	/** Local skeletal mesh component that holds the socket. */
+	UPROPERTY(BlueprintReadWrite, EditAnywhere, Category = Targeting)
+		USkeletalMeshComponent* SourceComponent;
+
+	/** If SourceActor and SourceComponent are valid, this is the name of the socket that will be used instead of the actor's location. */
+	UPROPERTY(BlueprintReadWrite, EditAnywhere, Category = Targeting)
+		FName SourceSocketName;
+
+	/** Point being targeted. */
+	UPROPERTY(BlueprintReadWrite, EditAnywhere, Category = Targeting)
+		FVector_NetQuantize TargetPoint;
+
+	// -------------------------------------
+
+	virtual bool HasOrigin() const override
+	{
+		return (SourceActor != NULL);
+	}
+
+	virtual FTransform GetOrigin() const override
+	{
+		if (SourceActor)
+		{
+			FTransform ReturnTransform = SourceActor->GetTransform();
+			if (SourceComponent)
+			{
+				ReturnTransform.SetLocation(SourceComponent->GetSocketLocation(SourceSocketName));
+			}
+			ReturnTransform.SetRotation((TargetPoint - ReturnTransform.GetLocation()).SafeNormal().Rotation().Quaternion());
+			return ReturnTransform;
+		}
+		return FTransform::Identity;
+	}
+
+	// -------------------------------------
+
+	virtual bool HasEndPoint() const override
+	{
+		return true;
+	}
+
+	virtual FVector GetEndPoint() const override
+	{
+		return TargetPoint;
+	}
+
+	// -------------------------------------
+
+	virtual FString ToString() const override
+	{
+		return TEXT("FGameplayAbilityTargetData_Mesh");
+	}
+
+	bool NetSerialize(FArchive& Ar, class UPackageMap* Map, bool& bOutSuccess);
+
+	virtual UScriptStruct* GetScriptStruct() override
+	{
+		return FGameplayAbilityTargetData_Mesh::StaticStruct();
+	}
+};
+
+template<>
+struct TStructOpsTypeTraits<FGameplayAbilityTargetData_Mesh> : public TStructOpsTypeTraitsBase
+{
+	enum
+	{
+		WithNetSerializer = true	// For now this is REQUIRED for FGameplayAbilityTargetDataHandle net serialization to work
+	};
+};
+
+
+USTRUCT(BlueprintType)
+struct GAMEPLAYABILITIES_API FGameplayAbilityTargetData_SingleTargetHit : public FGameplayAbilityTargetData
+{
+	GENERATED_USTRUCT_BODY()
+
+	FGameplayAbilityTargetData_SingleTargetHit()
+	{ }
+
+	FGameplayAbilityTargetData_SingleTargetHit(const FHitResult InHitResult)
+		: HitResult(InHitResult)
+	{ }
+
+	// -------------------------------------
+
+	virtual TArray<AActor*>	GetActors() const override
+	{
+		TArray<AActor*>	Actors;
+		if (HitResult.Actor.IsValid())
+		{
+			Actors.Push(HitResult.Actor.Get());
+		}
+		return Actors;
+	}
+
+	// -------------------------------------
+
+	virtual bool HasHitResult() const override
+	{
+		return true;
+	}
+
+	virtual const FHitResult* GetHitResult() const override
+	{
+		return &HitResult;
+	}
+
+	// -------------------------------------
+
+	UPROPERTY()
+		FHitResult	HitResult;
+
+	bool NetSerialize(FArchive& Ar, class UPackageMap* Map, bool& bOutSuccess);
+
+	virtual UScriptStruct* GetScriptStruct() override
+	{
+		return FGameplayAbilityTargetData_SingleTargetHit::StaticStruct();
+	}
+};
+
+template<>
+struct TStructOpsTypeTraits<FGameplayAbilityTargetData_SingleTargetHit> : public TStructOpsTypeTraitsBase
+{
+	enum
+	{
+		WithNetSerializer = true	// Fow now this is REQUIRED for FGameplayAbilityTargetDataHandle net serialization to work
+	};
+};
+
+/** Generic callback for returning when target data is available */
+DECLARE_MULTICAST_DELEGATE_OneParam(FAbilityTargetData, FGameplayAbilityTargetDataHandle);
diff --git a/Engine/Source/Runtime/GameplayAbilities/Public/Abilities/GameplayAbilityTypes.h b/Engine/Source/Runtime/GameplayAbilities/Public/Abilities/GameplayAbilityTypes.h
--- a/Engine/Source/Runtime/GameplayAbilities/Public/Abilities/GameplayAbilityTypes.h
+++ b/Engine/Source/Runtime/GameplayAbilities/Public/Abilities/GameplayAbilityTypes.h
@@ -59,20 +59,6 @@ namespace EGameplayAbilityReplicationPolicy
 	};
 }
 
-UENUM(BlueprintType)
-namespace EGameplayAbilityTargetingLocationType
-{
-	/**
-	 *	What type of location calculation to use when an ability asks for our transform.
-	 */
-
-	enum Type
-	{
-		LiteralTransform		UMETA(DisplayName = "Literal Transform"),		// We report an actual raw transform. This is also the final fallback if other methods fail.
-		ActorTransform			UMETA(DisplayName = "Actor Transform"),			// We pull the transform from an associated actor directly.
-		SocketTransform			UMETA(DisplayName = "Socket Transform"),		// We aim from a named socket on the player's skeletal mesh component.
-	};
-}
 
 USTRUCT(BlueprintType)
 struct FGameplayAbilityHandle
@@ -196,74 +182,6 @@ struct GAMEPLAYABILITIES_API FGameplayAbilityActorInfo
 	virtual void ClearActorInfo();
 };
 
-/**
-*	Handle for Targeting Data. This servers two main purposes:
-*		-Avoid us having to copy around the full targeting data structure in Blueprints
-*		-Allows us to leverage polymorphism in the target data structure
-*		-Allows us to implement NetSerialize and replicate by value between clients/server
-*
-*		-Avoid using UObjects could be used to give us polymorphism and by reference passing in blueprints.
-*		-However we would still be screwed when it came to replication
-*
-*		-Replication by value
-*		-Pass by reference in blueprints
-*		-Polymophism in TargetData structure
-*
-*/
-
-USTRUCT(BlueprintType)
-struct GAMEPLAYABILITIES_API FGameplayAbilityTargetDataHandle
-{
-	GENERATED_USTRUCT_BODY()
-
-	FGameplayAbilityTargetDataHandle() { }
-	FGameplayAbilityTargetDataHandle(struct FGameplayAbilityTargetData* DataPtr)
-		: Data(DataPtr)
-	{
-
-	}
-
-	TSharedPtr<FGameplayAbilityTargetData>	Data;
-
-	void Clear()
-	{
-		Data.Reset();
-	}
-
-	bool IsValid() const
-	{
-		return Data.IsValid();
-	}
-
-	bool NetSerialize(FArchive& Ar, class UPackageMap* Map, bool& bOutSuccess);
-
-	/** Comparison operator */
-	bool operator==(FGameplayAbilityTargetDataHandle const& Other) const
-	{
-		// Both invalid structs or both valid and Pointer compare (???) // deep comparison equality
-		bool bBothValid = IsValid() && Other.IsValid();
-		bool bBothInvalid = !IsValid() && !Other.IsValid();
-		return (bBothInvalid || (bBothValid && (Data.Get() == Other.Data.Get())));
-	}
-
-	/** Comparison operator */
-	bool operator!=(FGameplayAbilityTargetDataHandle const& Other) const
-	{
-		return !(FGameplayAbilityTargetDataHandle::operator==(Other));
-	}
-};
-
-template<>
-struct TStructOpsTypeTraits<FGameplayAbilityTargetDataHandle> : public TStructOpsTypeTraitsBase
-{
-	enum
-	{
-		WithCopy = true,		// Necessary so that TSharedPtr<FGameplayAbilityTargetData> Data is copied around
-		WithNetSerializer = true,
-		WithIdenticalViaEquality = true,
-	};
-};
-
 /**
 *	FGameplayAbilityActorInfo
 *
@@ -271,101 +189,6 @@ struct TStructOpsTypeTraits<FGameplayAbilityTargetDataHandle> : public TStructOp
 *
 */
 
-USTRUCT(BlueprintType)
-struct GAMEPLAYABILITIES_API FGameplayAbilityTargetingLocationInfo
-{
-	GENERATED_USTRUCT_BODY()
-
-	FGameplayAbilityTargetingLocationInfo()
-	: LocationType(EGameplayAbilityTargetingLocationType::LiteralTransform)
-	{
-	};
-
-	void operator=(const FGameplayAbilityTargetingLocationInfo& Other)
-	{
-		LocationType = Other.LocationType;
-		LiteralTransform = Other.LiteralTransform;
-		SourceActor = Other.SourceActor;
-		SourceComponent = Other.SourceComponent;
-		SourceSocketName = Other.SourceSocketName;
-	}
-
-public:
-	FTransform GetTargetingTransform() const
-	{
-		//Return or calculate based on LocationType.
-		switch (LocationType)
-		{
-		case EGameplayAbilityTargetingLocationType::ActorTransform:
-			if (SourceActor)
-			{
-				return SourceActor->GetTransform();
-			}
-			break;
-		case EGameplayAbilityTargetingLocationType::SocketTransform:
-			if (SourceComponent)
-			{
-				return SourceComponent->GetSocketTransform(SourceSocketName);		//Bad socket name will just return component transform anyway, so we're safe
-			}
-			break;
-		case EGameplayAbilityTargetingLocationType::LiteralTransform:
-			return LiteralTransform;
-		default:
-			check(false);		//This case should not happen
-			break;
-		}
-		//Error
-		return FTransform::Identity;
-	}
-
-	FGameplayAbilityTargetDataHandle MakeTargetDataHandleFromHitResult(TWeakObjectPtr<UGameplayAbility> Ability, FHitResult HitResult) const;
-
-	FGameplayAbilityTargetDataHandle MakeTargetDataHandleFromActors(TArray<AActor*> TargetActors) const;
-
-	/** Type of location used - will determine what data is transmitted over the network and what fields are used when calculating position. */
-	UPROPERTY(BlueprintReadWrite, meta = (ExposeOnSpawn = true), Category = Targeting)
-	TEnumAsByte<EGameplayAbilityTargetingLocationType::Type> LocationType;
-
-	/** A literal world transform can be used, if one has been calculated outside of the actor using the ability. */
-	UPROPERTY(BlueprintReadWrite, meta = (ExposeOnSpawn = true), Category = Targeting)
-	FTransform LiteralTransform;
-
-	/** A source actor is needed for Actor-based targeting, but not for Socket-based targeting. */
-	UPROPERTY(BlueprintReadWrite, meta = (ExposeOnSpawn = true), Category = Targeting)
-	AActor* SourceActor;
-
-	/** Socket-based targeting requires a skeletal mesh component to check for the named socket. */
-	UPROPERTY(BlueprintReadWrite, meta = (ExposeOnSpawn = true), Category = Targeting)
-	USkeletalMeshComponent* SourceComponent;
-
-	/** If SourceComponent is valid, this is the name of the socket transform that will be used. If no Socket is provided, SourceComponent's transform will be used. */
-	UPROPERTY(BlueprintReadWrite, meta = (ExposeOnSpawn = true), Category = Targeting)
-	FName SourceSocketName;
-
-	// -------------------------------------
-
-	virtual FString ToString() const
-	{
-		return TEXT("FGameplayAbilityTargetingLocationInfo");
-	}
-
-	bool NetSerialize(FArchive& Ar, class UPackageMap* Map, bool& bOutSuccess);
-
-	virtual UScriptStruct* GetScriptStruct()
-	{
-		return FGameplayAbilityTargetingLocationInfo::StaticStruct();
-	}
-};
-
-template<>
-struct TStructOpsTypeTraits<FGameplayAbilityTargetingLocationInfo> : public TStructOpsTypeTraitsBase
-{
-	enum
-	{
-		WithNetSerializer = true	// For now this is REQUIRED for FGameplayAbilityTargetDataHandle net serialization to work
-	};
-};
-
 
 /**
 *	FGameplayAbilityActivationInfo
@@ -376,6 +199,7 @@ struct TStructOpsTypeTraits<FGameplayAbilityTargetingLocationInfo> : public TStr
 *		-Generally not meant to be subclassed in projects.
 *		-Passed around by value since the struct is small.
 */
+
 USTRUCT(BlueprintType)
 struct GAMEPLAYABILITIES_API FGameplayAbilityActivationInfo
 {
@@ -423,356 +247,6 @@ struct GAMEPLAYABILITIES_API FGameplayAbilityActivationInfo
 
 // ----------------------------------------------------
 
-/**
- *	A generic structure for targeting data. We want generic functions to produce this data and other generic
- *	functions to consume this data.
- *
- *	We expect this to be able to hold specific actors/object reference and also generic location/direction/origin
- *	information.
- *
- *	Some example producers:
- *		-Overlap/Hit collision event generates TargetData about who was hit in a melee attack
- *		-A mouse input causes a hit trace and the actor infront of the crosshair is turned into TargetData
- *		-A mouse input causes TargetData to be generated from the owner's crosshair view origin/direction
- *		-An AOE/aura pulses and all actors in a radius around the instigator are added to TargetData
- *		-Panzer Dragoon style 'painting' targeting mode
- *		-MMORPG style ground AOE targeting style (potentially both a location on the ground and actors that were targeted)
- *
- *	Some example consumers:
- *		-Apply a GameplayEffect to all actors in TargetData
- *		-Find closest actor from all in TargetData
- *		-Call some function on all actors in TargetData
- *		-Filter or merge TargetDatas
- *		-Spawn a new actor at a TargetData location
- *
- *
- *
- *	Maybe it is better to distinguish between actor list targeting vs positional targeting data?
- *		-AOE/aura type of targeting data blurs the line
- *
- *
-*/
-
-USTRUCT()
-struct GAMEPLAYABILITIES_API FGameplayAbilityTargetData
-{
-	GENERATED_USTRUCT_BODY()
-
-	virtual ~FGameplayAbilityTargetData() { }
-
-	void ApplyGameplayEffect(UGameplayEffect *GameplayEffect, const FGameplayAbilityActorInfo InstigatorInfo);
-
-	virtual TArray<AActor*>	GetActors() const
-	{
-		return TArray<AActor*>();
-	}
-
-	// -------------------------------------
-
-	virtual bool HasHitResult() const
-	{
-		return false;
-	}
-
-	virtual const FHitResult* GetHitResult() const
-	{
-		return NULL;
-	}
-
-	// -------------------------------------
-
-	virtual bool HasOrigin() const
-	{
-		return false;
-	}
-
-	virtual FTransform GetOrigin() const
-	{
-		return FTransform::Identity;
-	}
-
-	// -------------------------------------
-
-	virtual bool HasEndPoint() const
-	{
-		return false;
-	}
-
-	virtual FVector GetEndPoint() const
-	{
-		return FVector::ZeroVector;
-	}
-
-	// -------------------------------------
-
-	virtual UScriptStruct* GetScriptStruct()
-	{
-		return FGameplayAbilityTargetData::StaticStruct();
-	}
-
-	virtual FString ToString() const;
-};
-
-USTRUCT(BlueprintType)
-struct GAMEPLAYABILITIES_API FGameplayAbilityTargetData_LocationInfo : public FGameplayAbilityTargetData
-{
-	GENERATED_USTRUCT_BODY()
-
-	/** Generic location data for source */
-	UPROPERTY(BlueprintReadWrite, EditAnywhere, Category = Targeting)
-	FGameplayAbilityTargetingLocationInfo SourceLocation;
-
-	/** Generic location data for target */
-	UPROPERTY(BlueprintReadWrite, EditAnywhere, Category = Targeting)
-	FGameplayAbilityTargetingLocationInfo TargetLocation;
-
-	// -------------------------------------
-
-	virtual bool HasOrigin() const
-	{
-		return true;
-	}
-
-	virtual FTransform GetOrigin() const
-	{
-		return SourceLocation.GetTargetingTransform();
-	}
-
-	// -------------------------------------
-
-	virtual bool HasEndPoint() const
-	{
-		return true;
-	}
-
-	virtual FVector GetEndPoint() const
-	{
-		return TargetLocation.GetTargetingTransform().GetLocation();
-	}
-
-	// -------------------------------------
-
-	virtual UScriptStruct* GetScriptStruct()
-	{
-		return FGameplayAbilityTargetData_LocationInfo::StaticStruct();
-	}
-
-	virtual FString ToString() const override
-	{
-		return TEXT("FGameplayAbilityTargetData_LocationInfo");
-	}
-
-	bool NetSerialize(FArchive& Ar, class UPackageMap* Map, bool& bOutSuccess);
-};
-
-template<>
-struct TStructOpsTypeTraits<FGameplayAbilityTargetData_LocationInfo> : public TStructOpsTypeTraitsBase
-{
-	enum
-	{
-		WithNetSerializer = true	// For now this is REQUIRED for FGameplayAbilityTargetDataHandle net serialization to work
-	};
-};
-
-USTRUCT(BlueprintType)
-struct GAMEPLAYABILITIES_API FGameplayAbilityTargetData_ActorArray : public FGameplayAbilityTargetData
-{
-	GENERATED_USTRUCT_BODY()
-
-	/** We could be selecting this group of actors from any type of location, so use a generic location type */
-	UPROPERTY(BlueprintReadWrite, EditAnywhere, Category = Targeting)
-	FGameplayAbilityTargetingLocationInfo SourceLocation;
-
-	/** Rather than targeting a single point, this type of targeting selects multiple actors. */
-	UPROPERTY(BlueprintReadWrite, EditAnywhere, Category = Targeting)
-	TArray<AActor*> TargetActorArray;
-
-	virtual TArray<AActor*>	GetActors() const override
-	{
-		return TargetActorArray;
-	}
-
-	// -------------------------------------
-
-	virtual bool HasOrigin() const override
-	{
-		return true;
-	}
-
-	virtual FTransform GetOrigin() const override
-	{
-		return SourceLocation.GetTargetingTransform();		//No single aim-at target, since this is aimed at multiple characters.
-	}
-
-	// -------------------------------------
-
-	virtual UScriptStruct* GetScriptStruct() override
-	{
-		return FGameplayAbilityTargetData_ActorArray::StaticStruct();
-	}
-
-	virtual FString ToString() const override
-	{
-		return TEXT("FGameplayAbilityTargetData_ActorArray");
-	}
-
-	bool NetSerialize(FArchive& Ar, class UPackageMap* Map, bool& bOutSuccess);
-};
-
-template<>
-struct TStructOpsTypeTraits<FGameplayAbilityTargetData_ActorArray> : public TStructOpsTypeTraitsBase
-{
-	enum
-	{
-		WithNetSerializer = true	// For now this is REQUIRED for FGameplayAbilityTargetDataHandle net serialization to work
-	};
-};
-
-USTRUCT(BlueprintType)
-struct GAMEPLAYABILITIES_API FGameplayAbilityTargetData_Mesh : public FGameplayAbilityTargetData
-{
-	GENERATED_USTRUCT_BODY()
-
-	FGameplayAbilityTargetData_Mesh()
-	: SourceActor(NULL)
-	, SourceComponent(NULL)
-	{}
-
-
-	/** Actor who owns the named component. Actor's location is used as start point if component cannot be found. */
-	UPROPERTY(BlueprintReadWrite, EditAnywhere, Category = Targeting)
-	AActor* SourceActor;
-
-	/** Local skeletal mesh component that holds the socket. */
-	UPROPERTY(BlueprintReadWrite, EditAnywhere, Category = Targeting)
-	USkeletalMeshComponent* SourceComponent;
-
-	/** If SourceActor and SourceComponent are valid, this is the name of the socket that will be used instead of the actor's location. */
-	UPROPERTY(BlueprintReadWrite, EditAnywhere, Category=Targeting)
-	FName SourceSocketName;
-
-	/** Point being targeted. */
-	UPROPERTY(BlueprintReadWrite, EditAnywhere, Category = Targeting)
-	FVector_NetQuantize TargetPoint;
-
-	// -------------------------------------
-
-	virtual bool HasOrigin() const override
-	{
-		return (SourceActor != NULL);
-	}
-
-	virtual FTransform GetOrigin() const override
-	{
-		if (SourceActor)
-		{
-			FTransform ReturnTransform = SourceActor->GetTransform();
-			if (SourceComponent)
-			{
-				ReturnTransform.SetLocation(SourceComponent->GetSocketLocation(SourceSocketName));
-			}
-			ReturnTransform.SetRotation((TargetPoint - ReturnTransform.GetLocation()).SafeNormal().Rotation().Quaternion());
-			return ReturnTransform;
-		}
-		return FTransform::Identity;
-	}
-
-	// -------------------------------------
-
-	virtual bool HasEndPoint() const override
-	{
-		return true;
-	}
-
-	virtual FVector GetEndPoint() const override
-	{
-		return TargetPoint;
-	}
-
-	// -------------------------------------
-
-	virtual FString ToString() const override
-	{
-		return TEXT("FGameplayAbilityTargetData_Mesh");
-	}
-
-	bool NetSerialize(FArchive& Ar, class UPackageMap* Map, bool& bOutSuccess);
-
-	virtual UScriptStruct* GetScriptStruct() override
-	{
-		return FGameplayAbilityTargetData_Mesh::StaticStruct();
-	}
-};
-
-template<>
-struct TStructOpsTypeTraits<FGameplayAbilityTargetData_Mesh> : public TStructOpsTypeTraitsBase
-{
-	enum
-	{
-		WithNetSerializer = true	// For now this is REQUIRED for FGameplayAbilityTargetDataHandle net serialization to work
-	};
-};
-
-
-USTRUCT(BlueprintType)
-struct GAMEPLAYABILITIES_API FGameplayAbilityTargetData_SingleTargetHit : public FGameplayAbilityTargetData
-{
-	GENERATED_USTRUCT_BODY()
-
-	FGameplayAbilityTargetData_SingleTargetHit()
-	{ }
-
-	FGameplayAbilityTargetData_SingleTargetHit(const FHitResult InHitResult)
-	: HitResult(InHitResult)
-	{ }
-
-	// -------------------------------------
-
-	virtual TArray<AActor*>	GetActors() const override
-	{
-		TArray<AActor*>	Actors;
-		if (HitResult.Actor.IsValid())
-		{
-			Actors.Push(HitResult.Actor.Get());
-		}
-		return Actors;
-	}
-	
-	// -------------------------------------
-
-	virtual bool HasHitResult() const override
-	{
-		return true;
-	}
-
-	virtual const FHitResult* GetHitResult() const override
-	{
-		return &HitResult;
-	}
-
-	// -------------------------------------
-
-	UPROPERTY()
-	FHitResult	HitResult;
-
-	bool NetSerialize(FArchive& Ar, class UPackageMap* Map, bool& bOutSuccess);
-
-	virtual UScriptStruct* GetScriptStruct() override
-	{
-		return FGameplayAbilityTargetData_SingleTargetHit::StaticStruct();
-	}
-};
-
-template<>
-struct TStructOpsTypeTraits<FGameplayAbilityTargetData_SingleTargetHit> : public TStructOpsTypeTraitsBase
-{
-	enum
-	{
-		WithNetSerializer = true	// Fow now this is REQUIRED for FGameplayAbilityTargetDataHandle net serialization to work
-	};
-};
-
-
 USTRUCT(BlueprintType)
 struct GAMEPLAYABILITIES_API FGameplayEventData
 {
@@ -826,9 +300,6 @@ struct FGameplayAbiliyInputBinds
 	UEnum* GetBindEnum() { return FindObject<UEnum>(ANY_PACKAGE, *EnumName); }
 };
 
-/** Generic callback for returning when target data is available */
-DECLARE_MULTICAST_DELEGATE_OneParam(FAbilityTargetData, FGameplayAbilityTargetDataHandle);
-
 /** Used for cleaning up predicted data on network clients */
 DECLARE_MULTICAST_DELEGATE(FAbilitySystemComponentPredictionKeyClear);
 
