Fix non unity (needed FGameplayAbilityTargetDataHandle definition)

diff --git a/Engine/Source/Runtime/GameplayAbilities/Public/AbilitySystemComponent.h b/Engine/Source/Runtime/GameplayAbilities/Public/AbilitySystemComponent.h
--- a/Engine/Source/Runtime/GameplayAbilities/Public/AbilitySystemComponent.h
+++ b/Engine/Source/Runtime/GameplayAbilities/Public/AbilitySystemComponent.h
@@ -4,6 +4,7 @@
 #include "GameplayEffect.h"
 #include "GameplayCueInterface.h"
 #include "Abilities/GameplayAbilityTypes.h"
+#include "Abilities/GameplayAbilityTargetTypes.h"
 #include "AbilitySystemComponent.generated.h"
 
 GAMEPLAYABILITIES_API DECLARE_LOG_CATEGORY_EXTERN(LogAbilitySystemComponent, Log, All);
