-Refactor activation RPCs and functionality out of GameplayAbility
-Added GameplayAbilitySpec/Handle
-Bunch of misc cleanup

diff --git a/Engine/Source/Runtime/GameplayAbilities/Private/Abilities/GameplayAbility.cpp b/Engine/Source/Runtime/GameplayAbilities/Private/Abilities/GameplayAbility.cpp
--- a/Engine/Source/Runtime/GameplayAbilities/Private/Abilities/GameplayAbility.cpp
+++ b/Engine/Source/Runtime/GameplayAbilities/Private/Abilities/GameplayAbility.cpp
@@ -49,55 +49,30 @@ UGameplayAbility::UGameplayAbility(const class FPostConstructInitializePropertie
 #endif
 }
 
-void UGameplayAbility::TriggerAbilityFromGameplayEvent(FGameplayAbilityActorInfo* ActorInfo, FGameplayTag EventTag, FGameplayEventData* Payload, UAbilitySystemComponent& Component)
+int32 UGameplayAbility::GetFunctionCallspace(UFunction* Function, void* Parameters, FFrame* Stack)
 {
-	if (ShouldAbilityRespondToEvent(EventTag, Payload))
+	if (HasAnyFlags(RF_ClassDefaultObject))
 	{
-		int32 ExecutingAbilityIndex = -1;
-
-		// if we're the server and this is coming from a predicted event we should check if the client has already predicted it
-		if (Payload->CurrPredictionKey 
-			&& NetExecutionPolicy == EGameplayAbilityNetExecutionPolicy::Predictive
-			&& ActorInfo->Actor->Role == ROLE_Authority)
-		{
-			bool bPendingClientAbilityFound = false;
-			for (auto PendingAbilityInfo : Component.PendingClientAbilities)
-			{
-				if (Payload->CurrPredictionKey == PendingAbilityInfo.PrevPredictionKey && this == PendingAbilityInfo.Ability) // found a match
-				{
-					Payload->PrevPredictionKey = PendingAbilityInfo.PrevPredictionKey;
-					Payload->CurrPredictionKey = PendingAbilityInfo.CurrPredictionKey;
-
-					Component.PendingClientAbilities.RemoveSingleSwap(PendingAbilityInfo);
-					bPendingClientAbilityFound = true;
-					break;
-				}
-			}
-
-			// we haven't received the client's copy of the triggered ability
-			// keep track of this so we can associate the prediction keys when it comes in
-			if (bPendingClientAbilityFound == false)
-			{
-				UAbilitySystemComponent::FExecutingAbilityInfo Info;
-				Info.CurrPredictionKey = Payload->CurrPredictionKey;
-				Info.Ability = this;
+		return FunctionCallspace::Local;
+	}
+	check(GetOuter() != NULL);
+	return GetOuter()->GetFunctionCallspace(Function, Parameters, Stack);
+}
 
-				ExecutingAbilityIndex = Component.ExecutingServerAbilities.Add(Info);
-			}
-		}
+bool UGameplayAbility::CallRemoteFunction(UFunction* Function, void* Parameters, FOutParmRec* OutParms, FFrame* Stack)
+{
+	check(!HasAnyFlags(RF_ClassDefaultObject));
+	check(GetOuter() != NULL);
 
-		if (TryActivateAbility(ActorInfo, Payload->PrevPredictionKey, Payload->CurrPredictionKey))
-		{
-			if (ExecutingAbilityIndex >= 0)
-			{
-				Component.ExecutingServerAbilities[ExecutingAbilityIndex].State = UAbilitySystemComponent::EAbilityExecutionState::Succeeded;
-			}
-		}
-		else if (ExecutingAbilityIndex >= 0)
-		{
-			Component.ExecutingServerAbilities[ExecutingAbilityIndex].State = UAbilitySystemComponent::EAbilityExecutionState::Failed;
-		}
+	AActor* Owner = CastChecked<AActor>(GetOuter());
+	UNetDriver* NetDriver = Owner->GetNetDriver();
+	if (NetDriver)
+	{
+		NetDriver->ProcessRemoteFunction(Owner, Function, Parameters, OutParms, Stack, this);
+		return true;
 	}
+
+	return false;
 }
 
 // TODO: polymorphic payload
@@ -129,10 +104,10 @@ void UGameplayAbility::PostNetInit()
 	
 	if (CurrentActorInfo == NULL)
 	{
-		AActor *OwnerActor = Cast<AActor>(GetOuter());
+		AActor* OwnerActor = Cast<AActor>(GetOuter());
 		if (ensure(OwnerActor))
 		{
-			UAbilitySystemComponent* AbilitySystemComponent = OwnerActor->FindComponentByClass<UAbilitySystemComponent>();
+			UAbilitySystemComponent* AbilitySystemComponent = UAbilitySystemGlobals::GetAbilitySystemComponentFromActor(OwnerActor);
 			if (ensure(AbilitySystemComponent))
 			{
 				CurrentActorInfo = AbilitySystemComponent->AbilityActorInfo.Get();
@@ -171,9 +146,9 @@ bool UGameplayAbility::IsSupportedForNetworking() const
 	return Supported;
 }
 
-bool UGameplayAbility::CanActivateAbility(const FGameplayAbilityActorInfo* ActorInfo) const
+bool UGameplayAbility::CanActivateAbility(const FGameplayAbilitySpecHandle Handle, const FGameplayAbilityActorInfo* ActorInfo) const
 {
-	SetCurrentActorInfo(ActorInfo);
+	SetCurrentActorInfo(Handle, ActorInfo);
 
 	if (ActorInfo->AbilitySystemComponent->GetUserAbilityActivationInhibited())
 	{
@@ -225,7 +200,7 @@ bool UGameplayAbility::ShouldAbilityRespondToEvent(FGameplayTag EventTag, const
 	return true;
 }
 
-bool UGameplayAbility::CommitAbility(const FGameplayAbilityActorInfo* ActorInfo, const FGameplayAbilityActivationInfo ActivationInfo)
+bool UGameplayAbility::CommitAbility(const FGameplayAbilitySpecHandle Handle, const FGameplayAbilityActorInfo* ActorInfo, const FGameplayAbilityActivationInfo ActivationInfo)
 {
 	// Last chance to fail (maybe we no longer have resources to commit since we after we started this ability activation)
 	if (!CommitCheck(ActorInfo, ActivationInfo))
@@ -277,92 +252,12 @@ void UGameplayAbility::CommitExecute(const FGameplayAbilityActorInfo* ActorInfo,
 	}
 }
 
-bool UGameplayAbility::TryActivateAbility(const FGameplayAbilityActorInfo* ActorInfo, uint32 PrevPredictionKey, uint32 CurrPredictionKey, UGameplayAbility ** OutInstancedAbility)
-{	
-	// make sure the ActorInfo and then Actor on that FGameplayAbilityActorInfo are valid, if not bail out.
-	if (ActorInfo == NULL || ActorInfo->Actor == NULL)
-	{
-		return false;
-	}
-
-	// This should only come from button presses/local instigation
-	ENetRole NetMode = ActorInfo->Actor->Role;
-	ensure(NetMode != ROLE_SimulatedProxy);
-	
-	FGameplayAbilityActivationInfo	ActivationInfo(ActorInfo->Actor.Get(), PrevPredictionKey, CurrPredictionKey);
-
-	// Always do a non instanced CanActivate check
-	if (!CanActivateAbility(ActorInfo))
-	{
-		return false;
-	}
-
-	// If we are the server or this is a client authorative 
-	if (NetExecutionPolicy == EGameplayAbilityNetExecutionPolicy::Client || (NetMode == ROLE_Authority))
-	{
-		// Create instance of this ability if necessary
-		
-		if (GetInstancingPolicy() == EGameplayAbilityInstancingPolicy::InstancedPerExecution)
-		{
-			UGameplayAbility * Ability = ActorInfo->AbilitySystemComponent->CreateNewInstanceOfAbility(this);
-			Ability->CallActivateAbility(ActorInfo, ActivationInfo);
-			if (OutInstancedAbility)
-			{
-				*OutInstancedAbility = Ability;
-			}
-		}
-		else
-		{
-			CallActivateAbility(ActorInfo, ActivationInfo);
-		}
-	}
-	else if (NetExecutionPolicy == EGameplayAbilityNetExecutionPolicy::Server)
-	{
-		ServerTryActivateAbility(ActorInfo, ActivationInfo);
-	}
-	else if (NetExecutionPolicy == EGameplayAbilityNetExecutionPolicy::Predictive)
-	{
-		// This execution is now officially EGameplayAbilityActivationMode:Predicting and has a PredictionKey
-		ActivationInfo.GeneratePredictionKey(ActorInfo->AbilitySystemComponent.Get());
-		
-		// This must be called immediately after GeneratePredictionKey to prevent problems with recursively activating abilities
-		ServerTryActivateAbility(ActorInfo, ActivationInfo);
-
-		if (GetInstancingPolicy() == EGameplayAbilityInstancingPolicy::InstancedPerExecution)
-		{
-			// For now, only NonReplicated + InstancedPerExecution abilities can be Predictive.
-			// We lack the code to predict spawning an instance of the execution and then merge/combine
-			// with the server spawned version when it arrives.
-			
-			if(GetReplicationPolicy() == EGameplayAbilityReplicationPolicy::ReplicateNone)
-			{
-				UGameplayAbility * Ability = ActorInfo->AbilitySystemComponent->CreateNewInstanceOfAbility(this);
-				Ability->CallActivateAbility(ActorInfo, ActivationInfo);
-				if (OutInstancedAbility)
-				{
-					*OutInstancedAbility = Ability;
-				}
-			}
-			else
-			{
-				ensure(false);
-			}
-		}
-		else
-		{
-			CallActivateAbility(ActorInfo, ActivationInfo);
-		}
-	}
-
-	return true;	
-}
-
-void UGameplayAbility::CancelAbility(const FGameplayAbilityActorInfo* ActorInfo, const FGameplayAbilityActivationInfo ActivationInfo)
+void UGameplayAbility::CancelAbility(const FGameplayAbilitySpecHandle Handle, const FGameplayAbilityActorInfo* ActorInfo, const FGameplayAbilityActivationInfo ActivationInfo)
 {
-	EndAbility(ActorInfo);
+	EndAbility(Handle, ActorInfo, ActivationInfo);
 }
 
-void UGameplayAbility::EndAbility(const FGameplayAbilityActorInfo* ActorInfo)
+void UGameplayAbility::EndAbility(const FGameplayAbilitySpecHandle Handle, const FGameplayAbilityActorInfo* ActorInfo, const FGameplayAbilityActivationInfo ActivationInfo)
 {
 	bIsActive = false;
 
@@ -377,12 +272,12 @@ void UGameplayAbility::EndAbility(const FGameplayAbilityActorInfo* ActorInfo)
 	ActiveTasks.Reset();	// Empty the array but dont resize memory, since this object is probably going to be destroyed very soon anyways.
 
 	// Tell owning AbilitySystemComponent that we ended so it can do stuff (including MarkPendingKill us)
-	ActorInfo->AbilitySystemComponent->NotifyAbilityEnded(this);
+	ActorInfo->AbilitySystemComponent->NotifyAbilityEnded(Handle, this);
 }
 
-void UGameplayAbility::ActivateAbility(const FGameplayAbilityActorInfo* ActorInfo, const FGameplayAbilityActivationInfo ActivationInfo)
+void UGameplayAbility::ActivateAbility(const FGameplayAbilitySpecHandle Handle, const FGameplayAbilityActorInfo* ActorInfo, const FGameplayAbilityActivationInfo ActivationInfo)
 {
-	SetCurrentInfo(ActorInfo, ActivationInfo);
+	SetCurrentInfo(Handle, ActorInfo, ActivationInfo);
 
 	if (HasBlueprintActivate)
 	{
@@ -395,55 +290,49 @@ void UGameplayAbility::ActivateAbility(const FGameplayAbilityActorInfo* ActorInf
 
 		// Do stuff...
 
-		if (CommitAbility(ActorInfo, ActivationInfo))		// ..then commit the ability...
+		if (CommitAbility(Handle, ActorInfo, ActivationInfo))		// ..then commit the ability...
 		{			
 			//	Then do more stuff...
 		}
 	}
 }
 
-void UGameplayAbility::PreActivate(const FGameplayAbilityActorInfo* ActorInfo, const FGameplayAbilityActivationInfo ActivationInfo)
+void UGameplayAbility::PreActivate(const FGameplayAbilitySpecHandle Handle, const FGameplayAbilityActorInfo* ActorInfo, const FGameplayAbilityActivationInfo ActivationInfo)
 {
 	UAbilitySystemComponent* Comp = ActorInfo->AbilitySystemComponent.Get();
 
-	Comp->CancelAbilitiesWithTags(CancelAbilitiesWithTag, ActorInfo, ActivationInfo, this);
+	Comp->CancelAbilitiesWithTags(CancelAbilitiesWithTag, Handle, ActorInfo, ActivationInfo, this);
 
 	bIsActive = true;
-	Comp->NotifyAbilityActivated(this);
+	Comp->NotifyAbilityActivated(Handle, this);
 }
 
-void UGameplayAbility::CallActivateAbility(const FGameplayAbilityActorInfo* ActorInfo, const FGameplayAbilityActivationInfo ActivationInfo)
+void UGameplayAbility::CallActivateAbility(const FGameplayAbilitySpecHandle Handle, const FGameplayAbilityActorInfo* ActorInfo, const FGameplayAbilityActivationInfo ActivationInfo)
 {
-	PreActivate(ActorInfo, ActivationInfo);
-	ActivateAbility(ActorInfo, ActivationInfo);
+	PreActivate(Handle, ActorInfo, ActivationInfo);
+	ActivateAbility(Handle, ActorInfo, ActivationInfo);
 }
 
 void UGameplayAbility::ConfirmActivateSucceed()
 {
-	CurrentActivationInfo.SetActivationConfirmed();
-
-	OnConfirmDelegate.Broadcast(this);
-	OnConfirmDelegate.Clear();
-}
-
-void UGameplayAbility::ServerTryActivateAbility(const FGameplayAbilityActorInfo* ActorInfo, const FGameplayAbilityActivationInfo ActivationInfo)
-{
-	ActorInfo->AbilitySystemComponent->ServerTryActivateAbility(this, ActivationInfo.PrevPredictionKey, ActivationInfo.CurrPredictionKey);
-}
+	// On instanced abilities, update CurrentActivationInfo and call any registered delegates.
+	if (HasAnyFlags(RF_ClassDefaultObject) == false)
+	{
+		PostNetInit();
+		check(CurrentActorInfo);
+		CurrentActivationInfo.SetActivationConfirmed();
 
-void UGameplayAbility::ClientActivateAbility(const FGameplayAbilityActorInfo* ActorInfo, const FGameplayAbilityActivationInfo ActivationInfo)
-{ 
-	CallActivateAbility(ActorInfo, ActivationInfo);
+		OnConfirmDelegate.Broadcast(this);
+		OnConfirmDelegate.Clear();
+	}
 }
 
-void UGameplayAbility::InputPressed(const FGameplayAbilityActorInfo* ActorInfo)
+void UGameplayAbility::InputPressed(const FGameplayAbilitySpecHandle Handle, const FGameplayAbilityActorInfo* ActorInfo, const FGameplayAbilityActivationInfo ActivationInfo)
 {
-	TryActivateAbility(ActorInfo);
 }
 
-void UGameplayAbility::InputReleased(const FGameplayAbilityActorInfo* ActorInfo)
+void UGameplayAbility::InputReleased(const FGameplayAbilitySpecHandle Handle, const FGameplayAbilityActorInfo* ActorInfo, const FGameplayAbilityActivationInfo ActivationInfo)
 {
-	
 }
 
 bool UGameplayAbility::CheckCooldown(const FGameplayAbilityActorInfo* ActorInfo) const
@@ -534,40 +423,14 @@ FGameplayEffectSpecHandle UGameplayAbility::GetOutgoingSpec(UGameplayEffect* Gam
 bool UGameplayAbility::K2_CommitAbility()
 {
 	check(CurrentActorInfo);
-	return CommitAbility(CurrentActorInfo, CurrentActivationInfo);
+	return CommitAbility(CurrentSpecHandle, CurrentActorInfo, CurrentActivationInfo);
 }
 
 void UGameplayAbility::K2_EndAbility()
 {
 	check(CurrentActorInfo);
-	EndAbility(CurrentActorInfo);
-}
-
-
-int32 UGameplayAbility::GetFunctionCallspace(UFunction* Function, void* Parameters, FFrame* Stack)
-{
-	if (HasAnyFlags(RF_ClassDefaultObject))
-	{
-		return FunctionCallspace::Local;
-	}
-	check(GetOuter() != NULL);
-	return GetOuter()->GetFunctionCallspace(Function, Parameters, Stack);
-}
-
-bool UGameplayAbility::CallRemoteFunction(UFunction* Function, void* Parameters, FOutParmRec* OutParms, FFrame* Stack)
-{
-	check(!HasAnyFlags(RF_ClassDefaultObject));
-	check(GetOuter() != NULL);
-
-	AActor* Owner = CastChecked<AActor>(GetOuter());
-	UNetDriver* NetDriver = Owner->GetNetDriver();
-	if (NetDriver)
-	{
-		NetDriver->ProcessRemoteFunction(Owner, Function, Parameters, OutParms, Stack, this);
-		return true;
-	}
-
-	return false;
+	
+	EndAbility(CurrentSpecHandle, CurrentActorInfo, CurrentActivationInfo);
 }
 
 // --------------------------------------------------------------------
@@ -582,6 +445,8 @@ void UGameplayAbility::MontageBranchPoint_AbilityDecisionStart(const FGameplayAb
 
 }
 
+// --------------------------------------------------------------------
+
 FGameplayAbilityTargetingLocationInfo UGameplayAbility::MakeTargetLocationInfoFromOwnerActor()
 {
 	FGameplayAbilityTargetingLocationInfo ReturnLocation;
@@ -599,22 +464,6 @@ FGameplayAbilityTargetingLocationInfo UGameplayAbility::MakeTargetLocationInfoFr
 	return ReturnLocation;
 }
 
-void UGameplayAbility::ClientActivateAbilitySucceed_Implementation(int32 PredictionKey)
-{
-	PostNetInit();
-
-	// Child classes can't override ClientActivateAbilitySucceed_Implementation, so just call ClientActivateAbilitySucceed_Internal
-	ClientActivateAbilitySucceed_Internal(PredictionKey);
-}
-
-void UGameplayAbility::ClientActivateAbilitySucceed_Internal(int32 PredictionKey)
-{
-	check(CurrentActorInfo);
-	CurrentActivationInfo.SetActivationConfirmed();
-
-	CallActivateAbility(CurrentActorInfo, CurrentActivationInfo);
-}
-
 //----------------------------------------------------------------------
 
 void UGameplayAbility::TaskStarted(UAbilityTask* NewTask)
@@ -687,296 +536,6 @@ void UGameplayAbility::TaskEnded(UAbilityTask* Task)
 	ActiveTasks.Remove(Task);
 }
 
-//----------------------------------------------------------------------
-
-void FGameplayAbilityActorInfo::InitFromActor(AActor *InActor, UAbilitySystemComponent* InAbilitySystemComponent)
-{
-	check(InActor);
-	check(InAbilitySystemComponent);
-
-	Actor = InActor;
-	AbilitySystemComponent = InAbilitySystemComponent;
-
-	// Look for a player controller or pawn in the owner chain.
-	AActor *TestActor = InActor;
-	while(TestActor)
-	{
-		if (APlayerController * CastPC = Cast<APlayerController>(TestActor))
-		{
-			PlayerController = CastPC;
-			break;
-		}
-
-		if (APawn * Pawn = Cast<APawn>(TestActor))
-		{
-			PlayerController = Cast<APlayerController>(Pawn->GetController());
-			break;
-		}
-
-		TestActor = TestActor->GetOwner();
-	}
-
-	// Grab Components that we care about
-	USkeletalMeshComponent * SkelMeshComponent = InActor->FindComponentByClass<USkeletalMeshComponent>();
-	if (SkelMeshComponent)
-	{
-		this->AnimInstance = SkelMeshComponent->GetAnimInstance();
-	}
-
-	MovementComponent = InActor->FindComponentByClass<UMovementComponent>();
-}
-
-void FGameplayAbilityActorInfo::ClearActorInfo()
-{
-	Actor = NULL;
-	PlayerController = NULL;
-	AnimInstance = NULL;
-	MovementComponent = NULL;
-}
-
-bool FGameplayAbilityActorInfo::IsLocallyControlled() const
-{
-	if (PlayerController.IsValid())
-	{
-		return PlayerController->IsLocalController();
-	}
-
-	return false;
-}
-
-bool FGameplayAbilityActorInfo::IsNetAuthority() const
-{
-	if (Actor.IsValid())
-	{
-		return (Actor->Role == ROLE_Authority);
-	}
-
-	// If we encounter issues with this being called before or after the owning actor is destroyed,
-	// we may need to cache off the authority (or look for it on some global/world state).
-
-	ensure(false);
-	return false;
-}
-
-void FGameplayAbilityActivationInfo::GeneratePredictionKey(UAbilitySystemComponent * Component) const
-{
-	check(Component);
-	check(ActivationMode == EGameplayAbilityActivationMode::NonAuthority);
-
-	PrevPredictionKey = CurrPredictionKey;
-	CurrPredictionKey = Component->GetNextPredictionKey();
-	ActivationMode = EGameplayAbilityActivationMode::Predicting;
-}
-
-void FGameplayAbilityActivationInfo::SetActivationConfirmed()
-{
-	//check(ActivationMode == EGameplayAbilityActivationMode::NonAuthority);
-	ActivationMode = EGameplayAbilityActivationMode::Confirmed;
-}
-
-//----------------------------------------------------------------------
-
-void FGameplayAbilityTargetData::ApplyGameplayEffect(UGameplayEffect* GameplayEffect, const FGameplayAbilityActorInfo InstigatorInfo)
-{
-	// TODO: Improve relationship between InstigatorContext and FGameplayAbilityTargetData/FHitResult (or use something different between HitResult)
-	
-
-	FGameplayEffectSpec	SpecToApply(GameplayEffect,					// The UGameplayEffect data asset
-									InstigatorInfo.Actor.Get(),		// The actor who instigated this
-									1.f,							// FIXME: Leveling
-									NULL							// FIXME: CurveData override... should we just remove this?
-									);
-	if (HasHitResult())
-	{
-		SpecToApply.InstigatorContext.AddHitResult(*GetHitResult());
-	}
-
-	if (HasOrigin())
-	{
-		SpecToApply.InstigatorContext.AddOrigin(GetOrigin().GetLocation());
-	}
-
-	TArray<AActor*> Actors = GetActors();
-	for (AActor * TargetActor : Actors)
-	{
-		check(TargetActor);
-		UAbilitySystemComponent * TargetComponent = UAbilitySystemBlueprintLibrary::GetAbilitySystemComponent(TargetActor);
-		if (TargetComponent)
-		{
-			InstigatorInfo.AbilitySystemComponent->ApplyGameplayEffectSpecToTarget(SpecToApply, TargetComponent);
-		}
-	}
-}
-
-
-FString FGameplayAbilityTargetData::ToString() const
-{
-	return TEXT("BASE CLASS");
-}
-
-
-FGameplayAbilityTargetDataHandle FGameplayAbilityTargetingLocationInfo::MakeTargetDataHandleFromHitResult(TWeakObjectPtr<UGameplayAbility> Ability, FHitResult HitResult) const
-{
-	if (LocationType == EGameplayAbilityTargetingLocationType::Type::SocketTransform)
-	{
-		const FGameplayAbilityActorInfo* ActorInfo = Ability.IsValid() ? Ability.Get()->GetCurrentActorInfo() : NULL;
-		AActor* AISourceActor = ActorInfo ? (ActorInfo->Actor.IsValid() ? ActorInfo->Actor.Get() : NULL) : NULL;
-		UAnimInstance* AnimInstance = ActorInfo ? ActorInfo->AnimInstance.Get() : NULL;
-		USkeletalMeshComponent* AISourceComponent = AnimInstance ? AnimInstance->GetOwningComponent() : NULL;
-
-		if (AISourceActor && AISourceComponent)
-		{
-			FGameplayAbilityTargetData_Mesh* ReturnData = new FGameplayAbilityTargetData_Mesh();
-			ReturnData->SourceActor = AISourceActor;
-			ReturnData->SourceComponent = AISourceComponent;
-			ReturnData->SourceSocketName = SourceSocketName;
-			ReturnData->TargetPoint = HitResult.Location;
-			return FGameplayAbilityTargetDataHandle(ReturnData);
-		}
-	}
-
-	/** Note: These are cleaned up by the FGameplayAbilityTargetDataHandle (via an internal TSharedPtr) */
-	FGameplayAbilityTargetData_SingleTargetHit* ReturnData = new FGameplayAbilityTargetData_SingleTargetHit();
-	ReturnData->HitResult = HitResult;
-	return FGameplayAbilityTargetDataHandle(ReturnData);
-}
-
-FGameplayAbilityTargetDataHandle FGameplayAbilityTargetingLocationInfo::MakeTargetDataHandleFromActors(TArray<AActor*> TargetActors) const
-{
-	/** Note: This is cleaned up by the FGameplayAbilityTargetDataHandle (via an internal TSharedPtr) */
-	FGameplayAbilityTargetData_ActorArray* ReturnData = new FGameplayAbilityTargetData_ActorArray();
-	ReturnData->TargetActorArray = TargetActors;
-	ReturnData->SourceLocation = *this;
-	return FGameplayAbilityTargetDataHandle(ReturnData);
-}
-
-bool FGameplayAbilityTargetDataHandle::NetSerialize(FArchive& Ar, class UPackageMap* Map, bool& bOutSuccess)
-{
-
-	UScriptStruct * ScriptStruct = Data.IsValid() ? Data->GetScriptStruct() : NULL;
-	Ar << ScriptStruct;
-
-	if (ScriptStruct)
-	{
-		if (Ar.IsLoading())
-		{
-			// For now, just always reset/reallocate the data when loading.
-			// Longer term if we want to generalize this and use it for property replication, we should support
-			// only reallocating when necessary
-			check(!Data.IsValid());
-
-			FGameplayAbilityTargetData * NewData = (FGameplayAbilityTargetData*)FMemory::Malloc(ScriptStruct->GetCppStructOps()->GetSize());
-			ScriptStruct->InitializeScriptStruct(NewData);
-
-			Data = TSharedPtr<FGameplayAbilityTargetData>(NewData);
-		}
-
-		void* ContainerPtr = Data.Get();
-
-		if (ScriptStruct->StructFlags & STRUCT_NetSerializeNative)
-		{
-			ScriptStruct->GetCppStructOps()->NetSerialize(Ar, Map, bOutSuccess, Data.Get());
-		}
-		else
-		{
-			// This won't work since UStructProperty::NetSerializeItem is deprecrated.
-			//	1) we have to manually crawl through the topmost struct's fields since we don't have a UStructProperty for it (just the UScriptProperty)
-			//	2) if there are any UStructProperties in the topmost struct's fields, we will assert in UStructProperty::NetSerializeItem.
-			
-			ABILITY_LOG(Fatal, TEXT("FGameplayAbilityTargetDataHandle::NetSerialize called on data struct %s without a native NetSerialize"), *ScriptStruct->GetName());
-
-			for (TFieldIterator<UProperty> It(ScriptStruct); It; ++It)
-			{
-				if (It->PropertyFlags & CPF_RepSkip)
-				{
-					continue;
-				}
-
-				void * PropertyData = It->ContainerPtrToValuePtr<void*>(ContainerPtr);
-
-				It->NetSerializeItem(Ar, Map, PropertyData);
-			}
-		}
-	}
-	
-
-	//ABILITY_LOG(Warning, TEXT("FGameplayAbilityTargetDataHandle Serialized: %s"), ScriptStruct ? *ScriptStruct->GetName() : TEXT("NULL") );
-
-	bOutSuccess = true;
-	return true;
-}
-
-bool FGameplayAbilityTargetingLocationInfo::NetSerialize(FArchive& Ar, class UPackageMap* Map, bool& bOutSuccess)
-{
-	Ar << LocationType;
-
-	switch (LocationType)
-	{
-	case EGameplayAbilityTargetingLocationType::ActorTransform:
-		Ar << SourceActor;
-		break;
-	case EGameplayAbilityTargetingLocationType::SocketTransform:
-		Ar << SourceComponent;
-		Ar << SourceSocketName;
-		break;
-	case EGameplayAbilityTargetingLocationType::LiteralTransform:
-		Ar << LiteralTransform;
-		break;
-	default:
-		check(false);		//This case should not happen
-		break;
-	}
-
-	bOutSuccess = true;
-	return true;
-}
-
-bool FGameplayAbilityTargetData_LocationInfo::NetSerialize(FArchive& Ar, class UPackageMap* Map, bool& bOutSuccess)
-{
-	SourceLocation.NetSerialize(Ar, Map, bOutSuccess);
-	TargetLocation.NetSerialize(Ar, Map, bOutSuccess);
-
-	bOutSuccess = true;
-	return true;
-}
-
-bool FGameplayAbilityTargetData_ActorArray::NetSerialize(FArchive& Ar, class UPackageMap* Map, bool& bOutSuccess)
-{
-	SourceLocation.NetSerialize(Ar, Map, bOutSuccess);
-	Ar << TargetActorArray;
-
-	bOutSuccess = true;
-	return true;
-}
-
-bool FGameplayAbilityTargetData_Mesh::NetSerialize(FArchive& Ar, class UPackageMap* Map, bool& bOutSuccess)
-{
-	//SourceActor can be used as a backup if the component isn't found.
-	Ar << SourceActor;
-	Ar << SourceComponent;
-	Ar << SourceSocketName;
-	TargetPoint.NetSerialize(Ar, Map, bOutSuccess);
-
-	bOutSuccess = true;
-	return true;
-}
-
-bool FGameplayAbilityTargetData_SingleTargetHit::NetSerialize(FArchive& Ar, class UPackageMap* Map, bool& bOutSuccess)
-{
-	Ar << HitResult.Actor;
-
-	HitResult.Location.NetSerialize(Ar, Map, bOutSuccess);
-	HitResult.Normal.NetSerialize(Ar, Map, bOutSuccess);
-
-	bOutSuccess = true;
-	return true;
-}
-
-
-
-// -----------------------------------------------------------------
-
-
 /**
  *	Helper methods for adding GaAmeplayCues without having to go through GameplayEffects.
  *	For now, none of these will happen predictively. We can eventually build this out more to 
@@ -1006,3 +565,19 @@ void UGameplayAbility::RemoveGameplayCue(FGameplayTag GameplayCueTag)
 		CurrentActorInfo->AbilitySystemComponent->RemoveGameplayCue(GameplayCueTag);
 	}
 }
+
+int32 UGameplayAbility::GetAbilityLevel() const
+{
+	check(IsInstantiated()); // You should not call this on non instanced abilities.
+	check(CurrentActorInfo);
+	return GetAbilityLevel(CurrentSpecHandle, CurrentActorInfo);
+}
+
+/** Returns current ability level for non instanced abilities. You m ust call this version in these contexts! */
+int32 UGameplayAbility::GetAbilityLevel(FGameplayAbilitySpecHandle Handle, const FGameplayAbilityActorInfo* ActorInfo) const
+{
+	FGameplayAbilitySpec* Spec = ActorInfo->AbilitySystemComponent->FindAbilitySpecFromHandle(Handle);
+	check(Spec);
+
+	return Spec->Level;
+}
\ No newline at end of file
diff --git a/Engine/Source/Runtime/GameplayAbilities/Private/Abilities/GameplayAbilitySet.cpp b/Engine/Source/Runtime/GameplayAbilities/Private/Abilities/GameplayAbilitySet.cpp
--- a/Engine/Source/Runtime/GameplayAbilities/Private/Abilities/GameplayAbilitySet.cpp
+++ b/Engine/Source/Runtime/GameplayAbilities/Private/Abilities/GameplayAbilitySet.cpp
@@ -25,7 +25,7 @@ void UGameplayAbilitySet::GiveAbilities(UAbilitySystemComponent* AbilitySystemCo
 	{
 		if (BindInfo.GameplayAbilityClass)
 		{
-			AbilitySystemComponent->GiveAbility(BindInfo.GameplayAbilityClass->GetDefaultObject<UGameplayAbility>(), (int32)BindInfo.Command);
+			AbilitySystemComponent->GiveAbility(FGameplayAbilitySpec(BindInfo.GameplayAbilityClass->GetDefaultObject<UGameplayAbility>(), 1, (int32)BindInfo.Command));
 		}
 	}
 }
diff --git a/Engine/Source/Runtime/GameplayAbilities/Private/Abilities/GameplayAbility_CharacterJump.cpp b/Engine/Source/Runtime/GameplayAbilities/Private/Abilities/GameplayAbility_CharacterJump.cpp
--- a/Engine/Source/Runtime/GameplayAbilities/Private/Abilities/GameplayAbility_CharacterJump.cpp
+++ b/Engine/Source/Runtime/GameplayAbilities/Private/Abilities/GameplayAbility_CharacterJump.cpp
@@ -16,12 +16,12 @@ UGameplayAbility_CharacterJump::UGameplayAbility_CharacterJump(const class FPost
 	InstancingPolicy = EGameplayAbilityInstancingPolicy::NonInstanced;
 }
 
-void UGameplayAbility_CharacterJump::ActivateAbility(const FGameplayAbilityActorInfo* ActorInfo, const FGameplayAbilityActivationInfo ActivationInfo)
+void UGameplayAbility_CharacterJump::ActivateAbility(const FGameplayAbilitySpecHandle Handle, const FGameplayAbilityActorInfo* ActorInfo, const FGameplayAbilityActivationInfo ActivationInfo)
 {
 	if (ActivationInfo.ActivationMode == EGameplayAbilityActivationMode::Authority ||
 		ActivationInfo.ActivationMode == EGameplayAbilityActivationMode::Predicting)
 	{
-		if (!CommitAbility(ActorInfo, ActivationInfo))
+		if (!CommitAbility(Handle, ActorInfo, ActivationInfo))
 		{
 			return;
 		}
@@ -31,21 +31,17 @@ void UGameplayAbility_CharacterJump::ActivateAbility(const FGameplayAbilityActor
 	}
 }
 
-void UGameplayAbility_CharacterJump::InputReleased(const FGameplayAbilityActorInfo* ActorInfo)
+void UGameplayAbility_CharacterJump::InputReleased(const FGameplayAbilitySpecHandle Handle, const FGameplayAbilityActorInfo* ActorInfo, const FGameplayAbilityActivationInfo ActivationInfo)
 {
-	// This is the 'top' of the ability execution code path. We don't have an ActivationInfo yet.
-	// Normally for InputPressed, we do a bunch of stuff (for clients): predict, tell the server, wait for confirm, rollback, etc.
-	// For jump, we are just going to call CancelAbility with an authority ActivationInfo. This is all jump needs, but more complex
-	// abilities may need more work.
-
-	FGameplayAbilityActivationInfo	ActivationInfo(EGameplayAbilityActivationMode::Authority);
-
-	CancelAbility(ActorInfo, ActivationInfo);
+	if (ActorInfo != NULL && ActorInfo->Actor != NULL)
+	{
+		CancelAbility(Handle, ActorInfo, ActivationInfo);
+	}
 }
 
-bool UGameplayAbility_CharacterJump::CanActivateAbility(const FGameplayAbilityActorInfo* ActorInfo) const
+bool UGameplayAbility_CharacterJump::CanActivateAbility(const FGameplayAbilitySpecHandle Handle, const FGameplayAbilityActorInfo* ActorInfo) const
 {
-	if (!Super::CanActivateAbility(ActorInfo))
+	if (!Super::CanActivateAbility(Handle, ActorInfo))
 	{
 		return false;
 	}
@@ -60,8 +56,10 @@ bool UGameplayAbility_CharacterJump::CanActivateAbility(const FGameplayAbilityAc
  *	Montage that *it* played was still playing, and if so, to cancel it. If this is something we need to support, we may need some
  *	light weight data structure to represent 'non intanced abilities in action' with a way to cancel/end them.
  */
-void UGameplayAbility_CharacterJump::CancelAbility(const FGameplayAbilityActorInfo* ActorInfo, const FGameplayAbilityActivationInfo ActivationInfo)
+void UGameplayAbility_CharacterJump::CancelAbility(const FGameplayAbilitySpecHandle Handle, const FGameplayAbilityActorInfo* ActorInfo, const FGameplayAbilityActivationInfo ActivationInfo)
 {
+	Super::CancelAbility(Handle, ActorInfo, ActivationInfo);
+
 	if (ActivationInfo.ActivationMode == EGameplayAbilityActivationMode::Authority ||
 		ActivationInfo.ActivationMode == EGameplayAbilityActivationMode::Predicting)
 	{
diff --git a/Engine/Source/Runtime/GameplayAbilities/Private/Abilities/GameplayAbility_Montage.cpp b/Engine/Source/Runtime/GameplayAbilities/Private/Abilities/GameplayAbility_Montage.cpp
--- a/Engine/Source/Runtime/GameplayAbilities/Private/Abilities/GameplayAbility_Montage.cpp
+++ b/Engine/Source/Runtime/GameplayAbilities/Private/Abilities/GameplayAbility_Montage.cpp
@@ -18,9 +18,9 @@ UGameplayAbility_Montage::UGameplayAbility_Montage(const class FPostConstructIni
 	NetExecutionPolicy = EGameplayAbilityNetExecutionPolicy::Server;
 }
 
-void UGameplayAbility_Montage::ActivateAbility(const FGameplayAbilityActorInfo* ActorInfo, const FGameplayAbilityActivationInfo ActivationInfo)
+void UGameplayAbility_Montage::ActivateAbility(const FGameplayAbilitySpecHandle Handle, const FGameplayAbilityActorInfo* ActorInfo, const FGameplayAbilityActivationInfo ActivationInfo)
 {
-	if (!CommitAbility(ActorInfo, ActivationInfo))
+	if (!CommitAbility(Handle, ActorInfo, ActivationInfo))
 	{
 		return;
 	}
diff --git a/Engine/Source/Runtime/GameplayAbilities/Private/Abilities/Tasks/AbilityTask_PlayMontageAndWait.cpp b/Engine/Source/Runtime/GameplayAbilities/Private/Abilities/Tasks/AbilityTask_PlayMontageAndWait.cpp
--- a/Engine/Source/Runtime/GameplayAbilities/Private/Abilities/Tasks/AbilityTask_PlayMontageAndWait.cpp
+++ b/Engine/Source/Runtime/GameplayAbilities/Private/Abilities/Tasks/AbilityTask_PlayMontageAndWait.cpp
@@ -37,11 +37,17 @@ void UAbilityTask_PlayMontageAndWait::Activate()
 		const FGameplayAbilityActorInfo* ActorInfo = Ability->GetCurrentActorInfo();
 		if (ActorInfo->AnimInstance.IsValid())
 		{
-			FOnMontageEnded EndDelegate;
-			EndDelegate.BindUObject(this, &UAbilityTask_PlayMontageAndWait::OnMontageEnded);
-
 			ActorInfo->AnimInstance->Montage_Play(MontageToPlay, 1.f);
-			ActorInfo->AnimInstance->Montage_SetEndDelegate(EndDelegate);
+			if (MontageToPlay)
+			{
+				FOnMontageEnded EndDelegate;
+				EndDelegate.BindUObject(this, &UAbilityTask_PlayMontageAndWait::OnMontageEnded);
+				ActorInfo->AnimInstance->Montage_SetEndDelegate(EndDelegate);
+			}
+			else
+			{
+				ABILITY_LOG(Warning, TEXT("UAbilityTask_PlayMontageAndWait called in Ability %s with null montage."), *Ability->GetName());
+			}
 		}
 	}
 }
diff --git a/Engine/Source/Runtime/GameplayAbilities/Private/AbilitySystemComponent.cpp b/Engine/Source/Runtime/GameplayAbilities/Private/AbilitySystemComponent.cpp
--- a/Engine/Source/Runtime/GameplayAbilities/Private/AbilitySystemComponent.cpp
+++ b/Engine/Source/Runtime/GameplayAbilities/Private/AbilitySystemComponent.cpp
@@ -612,8 +612,7 @@ void UAbilitySystemComponent::GetLifetimeReplicatedProps(TArray< FLifetimeProper
 	DOREPLIFETIME(UAbilitySystemComponent, SpawnedAttributes);
 	DOREPLIFETIME(UAbilitySystemComponent, ActiveGameplayEffects);
 	DOREPLIFETIME(UAbilitySystemComponent, ActiveGameplayCues);
-
-	DOREPLIFETIME(UAbilitySystemComponent, ReplicatedInstancedAbilities);
+	
 	DOREPLIFETIME(UAbilitySystemComponent, ActivatableAbilities);
 
 	DOREPLIFETIME(UAbilitySystemComponent, AbilityActor);
@@ -633,7 +632,7 @@ bool UAbilitySystemComponent::ReplicateSubobjects(class UActorChannel *Channel,
 		}
 	}
 
-	for (UGameplayAbility* Ability : ReplicatedInstancedAbilities)
+	for (UGameplayAbility* Ability : AllReplicatedInstancedAbilities)
 	{
 		if (Ability && !Ability->HasAnyFlags(RF_PendingKill))
 		{
diff --git a/Engine/Source/Runtime/GameplayAbilities/Private/AbilitySystemComponent_Abilities.cpp b/Engine/Source/Runtime/GameplayAbilities/Private/AbilitySystemComponent_Abilities.cpp
--- a/Engine/Source/Runtime/GameplayAbilities/Private/AbilitySystemComponent_Abilities.cpp
+++ b/Engine/Source/Runtime/GameplayAbilities/Private/AbilitySystemComponent_Abilities.cpp
@@ -115,21 +115,21 @@ void UAbilitySystemComponent::RefreshAbilityActorInfo()
 	AbilityActorInfo->InitFromActor(AbilityActorInfo->Actor.Get(), this);
 }
 
-UGameplayAbility* UAbilitySystemComponent::GiveAbility(UGameplayAbility* Ability, int32 InputID)
-{
-	check(Ability);
+FGameplayAbilitySpecHandle UAbilitySystemComponent::GiveAbility(FGameplayAbilitySpec Spec)
+{	
+	check(Spec.Ability);
 	check(IsOwnerActorAuthoritative());	// Should be called on authority
 	
-	if (Ability->GetInstancingPolicy() == EGameplayAbilityInstancingPolicy::InstancedPerActor)
+	FGameplayAbilitySpec& OwnedSpec = ActivatableAbilities[ActivatableAbilities.Add(Spec)];
+	
+	if (OwnedSpec.Ability->GetInstancingPolicy() == EGameplayAbilityInstancingPolicy::InstancedPerActor)
 	{
-		Ability = CreateNewInstanceOfAbility(Ability);
+		OwnedSpec.Ability = CreateNewInstanceOfAbility(OwnedSpec, Spec.Ability);
 	}
+	
+	OnGiveAbility(OwnedSpec);
 
-	ActivatableAbilities.Add(FGameplayAbilityInputIDPair(Ability, InputID));
-
-	OnGiveAbility(Ability, InputID);
-
-	return Ability;
+	return OwnedSpec.Handle;
 }
 
 void UAbilitySystemComponent::ClearAllAbilities()
@@ -140,26 +140,41 @@ void UAbilitySystemComponent::ClearAllAbilities()
 	ActivatableAbilities.Empty(ActivatableAbilities.Num());
 }
 
-void UAbilitySystemComponent::OnGiveAbility(UGameplayAbility* Ability, int32 InputID)
+void UAbilitySystemComponent::OnGiveAbility(const FGameplayAbilitySpec Spec)
 {
-	for (FAbilityTriggerData TriggerData : Ability->AbilityTriggers)
+	for (const FAbilityTriggerData& TriggerData : Spec.Ability->AbilityTriggers)
 	{
 		FGameplayTag EventTag = TriggerData.TriggerTag;
 
 		if (GameplayEventTriggeredAbilities.Contains(EventTag))
 		{
-			GameplayEventTriggeredAbilities[EventTag].AddUnique(Ability);
+			GameplayEventTriggeredAbilities[EventTag].AddUnique(Spec.Handle);	// Fixme: is this right? Do we want to trigger the ability directly of the spec?
 		}
 		else
 		{
-			TArray<TWeakObjectPtr<UGameplayAbility> > Triggers;
-			Triggers.Add(Ability);
+			TArray<FGameplayAbilitySpecHandle> Triggers;
+			Triggers.Add(Spec.Handle);
 			GameplayEventTriggeredAbilities.Add(EventTag, Triggers);
 		}
 	}
 }
 
-UGameplayAbility* UAbilitySystemComponent::CreateNewInstanceOfAbility(UGameplayAbility *Ability)
+FGameplayAbilitySpec* UAbilitySystemComponent::FindAbilitySpecFromHandle(FGameplayAbilitySpecHandle Handle)
+{
+	SCOPE_CYCLE_COUNTER(STAT_FindAbilitySpecFromHandle);
+
+	for (FGameplayAbilitySpec& Spec : ActivatableAbilities)
+	{
+		if (Spec.Handle == Handle)
+		{
+			return &Spec;
+		}
+	}
+
+	return nullptr;
+}
+
+UGameplayAbility* UAbilitySystemComponent::CreateNewInstanceOfAbility(FGameplayAbilitySpec& Spec, UGameplayAbility* Ability)
 {
 	check(Ability);
 	check(Ability->HasAllFlags(RF_ClassDefaultObject));
@@ -173,18 +188,21 @@ UGameplayAbility* UAbilitySystemComponent::CreateNewInstanceOfAbility(UGameplayA
 	// Add it to one of our instance lists so that it doesn't GC.
 	if (AbilityInstance->GetReplicationPolicy() != EGameplayAbilityReplicationPolicy::ReplicateNone)
 	{
-		ReplicatedInstancedAbilities.Add(AbilityInstance);
+		Spec.ReplicatedInstances.Add(AbilityInstance);
+		AllReplicatedInstancedAbilities.Add(AbilityInstance);
 	}
 	else
 	{
-		NonReplicatedInstancedAbilities.Add(AbilityInstance);
+		Spec.NonReplicatedInstances.Add(AbilityInstance);
 	}
 	
 	return AbilityInstance;
 }
 
-void UAbilitySystemComponent::NotifyAbilityEnded(UGameplayAbility* Ability)
+void UAbilitySystemComponent::NotifyAbilityEnded(FGameplayAbilitySpecHandle Handle, UGameplayAbility* Ability)
 {
+	FGameplayAbilitySpec* Spec = FindAbilitySpecFromHandle(Handle);
+	check(Spec);
 	check(Ability);
 
 	// If AnimatingAbility ended, clear the pointer
@@ -192,6 +210,8 @@ void UAbilitySystemComponent::NotifyAbilityEnded(UGameplayAbility* Ability)
 	{
 		AnimatingAbility = NULL;
 	}
+
+	Spec->IsActive = false;
 	
 	/** If this is instanced per execution, mark pending kill and remove it from our instanced lists if we are the authority */
 	if (Ability->GetInstancingPolicy() == EGameplayAbilityInstancingPolicy::InstancedPerExecution)
@@ -202,80 +222,43 @@ void UAbilitySystemComponent::NotifyAbilityEnded(UGameplayAbility* Ability)
 		{
 			if (GetOwnerRole() == ROLE_Authority)
 			{
-				ReplicatedInstancedAbilities.Remove(Ability);
+				Spec->ReplicatedInstances.Remove(Ability);
+				AllReplicatedInstancedAbilities.Remove(Ability);
 				Ability->MarkPendingKill();
 			}
 		}
 		else
 		{
-			NonReplicatedInstancedAbilities.Remove(Ability);
+			Spec->NonReplicatedInstances.Remove(Ability);
 			Ability->MarkPendingKill();
 		}
 	}
 }
 
-bool UAbilitySystemComponent::ActivateAbility(TWeakObjectPtr<UGameplayAbility> Ability)
-{
-	check(AbilityActorInfo.IsValid());
-
-	if (Ability.IsValid())
-	{
-		Ability.Get()->InputPressed(AbilityActorInfo.Get());
-	}
-
-	return false;
-}
-
-void UAbilitySystemComponent::CancelAbilitiesWithTags(const FGameplayTagContainer Tags, const FGameplayAbilityActorInfo* ActorInfo, const FGameplayAbilityActivationInfo ActivationInfo, UGameplayAbility* Ignore)
+void UAbilitySystemComponent::CancelAbilitiesWithTags(const FGameplayTagContainer Tags, const FGameplayAbilitySpecHandle Handle, const FGameplayAbilityActorInfo* ActorInfo, const FGameplayAbilityActivationInfo ActivationInfo, UGameplayAbility* Ignore)
 {
-	/**
-	 *	FIXME
-	 *
-	 *	Right now we are canceling all activatable abilities that match Tags. This includes abilities that might not have been activated in the first place!
-	 *	For instanced-per-actor abilities this is fine. They could check if they were activated/still activating.
-	 *	For non-instanced abilities it is ambiguous. We have no way to know 'how many' non-instanced abilities are in flight.
-	 *	Likewise for instanced-per-execution abilities. Though they are present in Replicated/NonReplicatedInstancedAbilities list
-	 */
-
-	struct local
+	for (FGameplayAbilitySpec& Spec : ActivatableAbilities)
 	{
-		static void CancelAbilitiesWithTags(TArray<UGameplayAbility*>& AbilitiesToCancel, const FGameplayTagContainer InTags, TArray<FGameplayAbilityInputIDPair> Abilities, const FGameplayAbilityActorInfo* InActorInfo, const FGameplayAbilityActivationInfo InActivationInfo, UGameplayAbility* InIgnore)
+		if (Spec.IsActive && Spec.Ability && Spec.Ability->AbilityTags.MatchesAny(Tags, false))
 		{
-			for (int32 idx=0; idx < Abilities.Num(); ++idx)
+			if (Spec.Ability->GetInstancingPolicy() == EGameplayAbilityInstancingPolicy::InstancedPerExecution)
 			{
-				UGameplayAbility *Ability = Abilities[idx].Ability;
-				if (Ability && (Ability != InIgnore) && Ability->AbilityTags.MatchesAny(InTags, false))
+				// We need to cancel spawned instance, not the CDO
+				TArray<UGameplayAbility*> AbilitiesToCancel = Spec.GetAbilityInstances();
+				for (UGameplayAbility* InstanceAbility : AbilitiesToCancel)
 				{
-					if (!Ability->HasAnyFlags(RF_ClassDefaultObject) || Ability->GetInstancingPolicy() == EGameplayAbilityInstancingPolicy::NonInstanced)
+					if (InstanceAbility && Ignore != InstanceAbility)
 					{
-						AbilitiesToCancel.Add(Ability);
-					}					
+						InstanceAbility->CancelAbility(Spec.Handle, ActorInfo, ActivationInfo);
+					}
 				}
 			}
-		}
-
-		static void CancelAbilitiesWithTags(TArray<UGameplayAbility*>& AbilitiesToCancel, const FGameplayTagContainer InTags, TArray<UGameplayAbility*> Abilities, const FGameplayAbilityActorInfo* InActorInfo, const FGameplayAbilityActivationInfo InActivationInfo, UGameplayAbility* InIgnore)
-		{
-			for (int32 idx=0; idx < Abilities.Num(); ++idx)
+			else
 			{
-				UGameplayAbility *Ability = Abilities[idx];
-				if (Ability && (Ability != InIgnore) && Ability->AbilityTags.MatchesAny(InTags, false))
-				{					
-					AbilitiesToCancel.Add(Ability);
-				}
+				Spec.Ability->CancelAbility(Spec.Handle, ActorInfo, ActivationInfo);
+				check(!Spec.IsActive);
 			}
 		}
-	};
-
-	TArray<UGameplayAbility*>	AbilitiesToCancel;
-	
-	local::CancelAbilitiesWithTags(AbilitiesToCancel, Tags, ReplicatedInstancedAbilities, ActorInfo, ActivationInfo, Ignore);
-	local::CancelAbilitiesWithTags(AbilitiesToCancel, Tags, NonReplicatedInstancedAbilities, ActorInfo, ActivationInfo, Ignore);
-	local::CancelAbilitiesWithTags(AbilitiesToCancel, Tags, ActivatableAbilities, ActorInfo, ActivationInfo, Ignore);
-
-	for (UGameplayAbility* Ability : AbilitiesToCancel)
-	{
-		Ability->CancelAbility(ActorInfo, ActivationInfo);
 	}
 }
 
@@ -291,27 +274,140 @@ TEXT("AbilitySystem.DenyClientActivations"),
 
 void UAbilitySystemComponent::OnRep_ActivateAbilities()
 {
-	for (FGameplayAbilityInputIDPair& Pair : ActivatableAbilities)
+	for (FGameplayAbilitySpec& Spec : ActivatableAbilities)
+	{
+		OnGiveAbility(Spec);
+	}
+}
+
+/**
+ * Attempts to activate the ability.
+ *	-This function calls CanActivateAbility
+ *	-This function handles instancing
+ *	-This function handles networking and prediction
+ *	-If all goes well, CallActivateAbility is called next.
+ */
+bool UAbilitySystemComponent::TryActivateAbility(FGameplayAbilitySpecHandle Handle, uint32 PrevPredictionKey, uint32 CurrPredictionKey, UGameplayAbility** OutInstancedAbility )
+{
+	FGameplayAbilitySpec* Spec = FindAbilitySpecFromHandle(Handle);
+	if (!Spec)
+	{
+		ABILITY_LOG(Warning, TEXT("TryActivateAbility called with invalid Handle"));
+		return false;
+	}
+
+	const FGameplayAbilityActorInfo* ActorInfo = AbilityActorInfo.Get();
+
+	// make sure the ActorInfo and then Actor on that FGameplayAbilityActorInfo are valid, if not bail out.
+	if (ActorInfo == NULL || ActorInfo->Actor == NULL)
+	{
+		return false;
+	}
+
+	// This should only come from button presses/local instigation (AI, etc)
+	ENetRole NetMode = ActorInfo->Actor->Role;
+	ensure(NetMode != ROLE_SimulatedProxy);
+
+	UGameplayAbility* Ability = Spec->Ability;
+
+	// Always do a non instanced CanActivate check
+	if (!Ability->CanActivateAbility(Handle, ActorInfo))
+	{
+		return false;
+	}
+
+	if (Spec->IsActive)
+	{
+		ABILITY_LOG(Warning, TEXT("TryActivateAbility called when ability was already active. NetMode: %d. Ability: %s"), (int32)NetMode, *Ability->GetName());
+	}
+
+	//check(!Spec->IsActive);	// Shouldn't be here if the ability is already active.
+	Spec->IsActive = true;
+
+	// Setup a fresh ActivationInfo for this AbilitySpec.
+	Spec->ActivationInfo = FGameplayAbilityActivationInfo(ActorInfo->Actor.Get(), PrevPredictionKey, CurrPredictionKey);
+	FGameplayAbilityActivationInfo &ActivationInfo = Spec->ActivationInfo;
+
+	// If we are the server or this is a client authorative 
+	if (Ability->GetNetExecutionPolicy() == EGameplayAbilityNetExecutionPolicy::Client || (NetMode == ROLE_Authority))
+	{
+		// Create instance of this ability if necessary
+
+		if (Ability->GetInstancingPolicy() == EGameplayAbilityInstancingPolicy::InstancedPerExecution)
+		{
+			UGameplayAbility* InstancedAbility = CreateNewInstanceOfAbility(*Spec, Ability);
+			InstancedAbility->CallActivateAbility(Handle, ActorInfo, ActivationInfo);
+			if (OutInstancedAbility)
+			{
+				*OutInstancedAbility = InstancedAbility;
+			}
+		}
+		else
+		{
+			Ability->CallActivateAbility(Handle, ActorInfo, ActivationInfo);
+		}
+	}
+	else if (Ability->GetNetExecutionPolicy() == EGameplayAbilityNetExecutionPolicy::Server)
+	{
+		ServerTryActivateAbility(Handle, ActivationInfo.PrevPredictionKey, ActivationInfo.CurrPredictionKey);
+	}
+	else if (Ability->GetNetExecutionPolicy() == EGameplayAbilityNetExecutionPolicy::Predictive)
 	{
-		UGameplayAbility* Ability = Pair.Ability;
-		if (Ability)
+		// This execution is now officially EGameplayAbilityActivationMode:Predicting and has a PredictionKey
+		ActivationInfo.GeneratePredictionKey(this);
+
+		// This must be called immediately after GeneratePredictionKey to prevent problems with recursively activating abilities
+		ServerTryActivateAbility(Handle, ActivationInfo.PrevPredictionKey, ActivationInfo.CurrPredictionKey);
+
+		if (Ability->GetInstancingPolicy() == EGameplayAbilityInstancingPolicy::InstancedPerExecution)
 		{
-			OnGiveAbility(Ability, Pair.InputID);
+			// For now, only NonReplicated + InstancedPerExecution abilities can be Predictive.
+			// We lack the code to predict spawning an instance of the execution and then merge/combine
+			// with the server spawned version when it arrives.
+
+			if (Ability->GetReplicationPolicy() == EGameplayAbilityReplicationPolicy::ReplicateNone)
+			{
+				UGameplayAbility* InstancedAbility = CreateNewInstanceOfAbility(*Spec, Ability);
+				InstancedAbility->CallActivateAbility(Handle, ActorInfo, ActivationInfo);
+				if (OutInstancedAbility)
+				{
+					*OutInstancedAbility = InstancedAbility;
+				}
+			}
+			else
+			{
+				ensure(false);
+			}
+		}
+		else
+		{
+			Ability->CallActivateAbility(Handle, ActorInfo, ActivationInfo);
 		}
 	}
+
+	return true;
 }
 
-void UAbilitySystemComponent::ServerTryActivateAbility_Implementation(UGameplayAbility * AbilityToActivate, uint32 PrevPredictionKey, uint32 CurrPredictionKey)
+void UAbilitySystemComponent::ServerTryActivateAbility_Implementation(FGameplayAbilitySpecHandle Handle, uint32 PrevPredictionKey, uint32 CurrPredictionKey)
 {
 #if !(UE_BUILD_SHIPPING || UE_BUILD_TEST)
 	if (DenyClientActivation > 0)
 	{
 		DenyClientActivation--;
-		ClientActivateAbilityFailed(AbilityToActivate, CurrPredictionKey);
+		ClientActivateAbilityFailed(Handle, CurrPredictionKey);
 		return;
 	}
 #endif
 
+	FGameplayAbilitySpec* Spec = FindAbilitySpecFromHandle(Handle);
+	if (!Spec)
+	{
+		// Can potentially happen in race conditions where client tries to activate ability that is removed server side before it is received.
+		return;
+	}
+
+	UGameplayAbility* AbilityToActivate = Spec->Ability;
+
 	ensure(AbilityToActivate);
 	ensure(AbilityActorInfo.IsValid());
 
@@ -321,19 +417,25 @@ void UAbilitySystemComponent::ServerTryActivateAbility_Implementation(UGameplayA
 		// first check if the server has already started this ability
 		for (auto ExecutingAbilityInfo : ExecutingServerAbilities)
 		{
-			if (ExecutingAbilityInfo.CurrPredictionKey == PrevPredictionKey && ExecutingAbilityInfo.Ability == AbilityToActivate)
+			if (ExecutingAbilityInfo.CurrPredictionKey == PrevPredictionKey && ExecutingAbilityInfo.Handle == Handle)
 			{
 				switch(ExecutingAbilityInfo.State)
 				{
 					case EAbilityExecutionState::Executing:
-						ExecutingAbilityInfo.Ability->CurrentActivationInfo.PrevPredictionKey = PrevPredictionKey;
-						ExecutingAbilityInfo.Ability->CurrentActivationInfo.CurrPredictionKey = CurrPredictionKey;
+					{
+						FGameplayAbilitySpec* Spec = FindAbilitySpecFromHandle(ExecutingAbilityInfo.Handle);
+						if (Spec)
+						{
+							Spec->ActivationInfo.PrevPredictionKey = PrevPredictionKey;
+							Spec->ActivationInfo.CurrPredictionKey = CurrPredictionKey;
+						}
 						break;
+					}
 					case EAbilityExecutionState::Failed:
-						ClientActivateAbilityFailed(AbilityToActivate, CurrPredictionKey);
+						ClientActivateAbilityFailed(Handle, CurrPredictionKey);
 						break;
 					case EAbilityExecutionState::Succeeded:
-						ClientActivateAbilitySucceed(AbilityToActivate, CurrPredictionKey);
+						ClientActivateAbilitySucceed(Handle, CurrPredictionKey);
 						break;
 				}
 
@@ -345,7 +447,7 @@ void UAbilitySystemComponent::ServerTryActivateAbility_Implementation(UGameplayA
 		FPendingAbilityInfo AbilityInfo;
 		AbilityInfo.PrevPredictionKey = PrevPredictionKey;
 		AbilityInfo.CurrPredictionKey = CurrPredictionKey;
-		AbilityInfo.Ability = AbilityToActivate;
+		AbilityInfo.Handle = Handle;
 
 		PendingClientAbilities.Add(AbilityInfo);
 		return;
@@ -353,16 +455,10 @@ void UAbilitySystemComponent::ServerTryActivateAbility_Implementation(UGameplayA
 
 	UGameplayAbility *InstancedAbility = NULL;
 
-	if (AbilityToActivate->TryActivateAbility(AbilityActorInfo.Get(), PrevPredictionKey, CurrPredictionKey, &InstancedAbility))
+	// Attempt to activate the ability (server side) and tell the client if it succeeded or failed.
+	if (TryActivateAbility(Handle, PrevPredictionKey, CurrPredictionKey, &InstancedAbility))
 	{
-		if (InstancedAbility && InstancedAbility->GetReplicationPolicy() != EGameplayAbilityReplicationPolicy::ReplicateNone)
-		{
-			InstancedAbility->ClientActivateAbilitySucceed(CurrPredictionKey);
-		}
-		else
-		{
-			ClientActivateAbilitySucceed(AbilityToActivate, CurrPredictionKey);
-		}
+		ClientActivateAbilitySucceed(Handle, CurrPredictionKey);
 
 		// Update our ReplicatedPredictionKey. When the client gets value, he will know his state (actor+all components/subobjects) are up to do date and he can
 		// remove any necessary predictive work.
@@ -375,16 +471,16 @@ void UAbilitySystemComponent::ServerTryActivateAbility_Implementation(UGameplayA
 	}
 	else
 	{
-		ClientActivateAbilityFailed(AbilityToActivate, CurrPredictionKey);
+		ClientActivateAbilityFailed(Handle, CurrPredictionKey);
 	}
 }
 
-bool UAbilitySystemComponent::ServerTryActivateAbility_Validate(UGameplayAbility * AbilityToActivate, uint32 PrevPredictionKey, uint32 CurrPredictionKey)
+bool UAbilitySystemComponent::ServerTryActivateAbility_Validate(FGameplayAbilitySpecHandle AbilityToActivate, uint32 PrevPredictionKey, uint32 CurrPredictionKey)
 {
 	return true;
 }
 
-void UAbilitySystemComponent::ClientActivateAbilityFailed_Implementation(UGameplayAbility * AbilityToActivate, uint32 PredictionKey)
+void UAbilitySystemComponent::ClientActivateAbilityFailed_Implementation(FGameplayAbilitySpecHandle Handle, uint32 PredictionKey)
 {
 	if (PredictionKey > 0)
 	{
@@ -394,11 +490,17 @@ void UAbilitySystemComponent::ClientActivateAbilityFailed_Implementation(UGamepl
 			if (Key == PredictionKey)
 			{
 				ABILITY_LOG(Warning, TEXT("Failed ActivateAbility, clearing prediction data %d"), PredictionKey);
-
 				PredictionDelegates[idx].Value.PredictionKeyClearDelegate.Broadcast();
-				for (uint32 DependentKey : PredictionDelegates[idx].Value.DependentPredictionKeys)
+
+				// Find the actual UGameplayAbility
+				FGameplayAbilitySpec* Spec = FindAbilitySpecFromHandle(Handle);
+				if (Spec)
 				{
-					ClientActivateAbilityFailed(AbilityToActivate, DependentKey);
+					UGameplayAbility* AbilityToActivate = Spec->Ability;
+					for (uint32 DependentKey : PredictionDelegates[idx].Value.DependentPredictionKeys)
+					{
+						ClientActivateAbilityFailed(Handle, DependentKey);
+					}
 				}
 				PredictionDelegates.RemoveAt(idx);
 				break;
@@ -411,11 +513,22 @@ void UAbilitySystemComponent::ClientActivateAbilityFailed_Implementation(UGamepl
 	}
 }
 
-void UAbilitySystemComponent::ClientActivateAbilitySucceed_Implementation(UGameplayAbility* AbilityToActivate, uint32 PredictionKey)
+void UAbilitySystemComponent::ClientActivateAbilitySucceed_Implementation(FGameplayAbilitySpecHandle Handle, uint32 PredictionKey)
 {
+	FGameplayAbilitySpec* Spec = FindAbilitySpecFromHandle(Handle);
+	if (!Spec)
+	{
+		// Can potentially happen in race conditions where client tries to activate ability that is removed server side before it is received.
+		return;
+	}
+
+	UGameplayAbility* AbilityToActivate = Spec->Ability;
+
 	ensure(AbilityToActivate);
 	ensure(AbilityActorInfo.IsValid());
 
+	Spec->ActivationInfo.SetActivationConfirmed();
+
 	// Fixme: We need a better way to link up/reconcile predictive replicated abilities. It would be ideal if we could predictively spawn an
 	// ability and then replace/link it with the server spawned one once the server has confirmed it.
 
@@ -429,7 +542,7 @@ void UAbilitySystemComponent::ClientActivateAbilitySucceed_Implementation(UGamep
 		{
 			// Find the one we predictively spawned, tell them we are confirmed
 			bool found = false;
-			for (UGameplayAbility* LocalAbility : NonReplicatedInstancedAbilities)				// Fixme: this has to be updated once predictive abilities can replicate
+			for (UGameplayAbility* LocalAbility : Spec->NonReplicatedInstances)				// Fixme: this has to be updated once predictive abilities can replicate
 			{
 				if (LocalAbility->GetCurrentActivationInfo().CurrPredictionKey == PredictionKey)
 				{
@@ -447,21 +560,89 @@ void UAbilitySystemComponent::ClientActivateAbilitySucceed_Implementation(UGamep
 	}
 	else
 	{
-		FGameplayAbilityActivationInfo	ActivationInfo(EGameplayAbilityActivationMode::Confirmed, 0, PredictionKey);
 		// We haven't already executed this ability at all, so kick it off.
+		
 		if (AbilityToActivate->GetInstancingPolicy() == EGameplayAbilityInstancingPolicy::InstancedPerExecution)
 		{
 			// Need to instantiate this in order to execute
-			UGameplayAbility* InstancedAbility = CreateNewInstanceOfAbility(AbilityToActivate);
-			InstancedAbility->CallActivateAbility(AbilityActorInfo.Get(), ActivationInfo);
+			UGameplayAbility* InstancedAbility = CreateNewInstanceOfAbility(*Spec, AbilityToActivate);
+			InstancedAbility->CallActivateAbility(Handle, AbilityActorInfo.Get(), Spec->ActivationInfo);
 		}
 		else
 		{
-			AbilityToActivate->CallActivateAbility(AbilityActorInfo.Get(), ActivationInfo);
+			AbilityToActivate->CallActivateAbility(Handle, AbilityActorInfo.Get(), Spec->ActivationInfo);
 		}
 	}
 }
 
+void UAbilitySystemComponent::TriggerAbilityFromGameplayEvent(FGameplayAbilitySpecHandle Handle, FGameplayAbilityActorInfo* ActorInfo, FGameplayTag EventTag, FGameplayEventData* Payload, UAbilitySystemComponent& Component)
+{
+	FGameplayAbilitySpec* Spec = FindAbilitySpecFromHandle(Handle);
+	if (!ensure(Spec))
+	{
+		return;
+	}
+
+	UGameplayAbility* Ability = Spec->Ability;
+	if (!ensure(Ability))
+	{
+		return;
+	}
+	
+
+	if (Ability->ShouldAbilityRespondToEvent(EventTag, Payload))
+	{
+		int32 ExecutingAbilityIndex = -1;
+
+		// if we're the server and this is coming from a predicted event we should check if the client has already predicted it
+		if (Payload->CurrPredictionKey
+			&& Ability->GetNetExecutionPolicy() == EGameplayAbilityNetExecutionPolicy::Predictive
+			&& ActorInfo->Actor->Role == ROLE_Authority)
+		{
+			bool bPendingClientAbilityFound = false;
+			for (auto PendingAbilityInfo : Component.PendingClientAbilities)
+			{
+				if (Payload->CurrPredictionKey == PendingAbilityInfo.PrevPredictionKey && Handle == PendingAbilityInfo.Handle) // found a match
+				{
+					Payload->PrevPredictionKey = PendingAbilityInfo.PrevPredictionKey;
+					Payload->CurrPredictionKey = PendingAbilityInfo.CurrPredictionKey;
+
+					Component.PendingClientAbilities.RemoveSingleSwap(PendingAbilityInfo);
+					bPendingClientAbilityFound = true;
+					break;
+				}
+			}
+
+			// we haven't received the client's copy of the triggered ability
+			// keep track of this so we can associate the prediction keys when it comes in
+			if (bPendingClientAbilityFound == false)
+			{
+				UAbilitySystemComponent::FExecutingAbilityInfo Info;
+				Info.CurrPredictionKey = Payload->CurrPredictionKey;
+				Info.Handle = Handle;
+
+				ExecutingAbilityIndex = Component.ExecutingServerAbilities.Add(Info);
+			}
+		}
+
+		if (TryActivateAbility(Handle, Payload->PrevPredictionKey, Payload->CurrPredictionKey))
+		{
+			if (ExecutingAbilityIndex >= 0)
+			{
+				Component.ExecutingServerAbilities[ExecutingAbilityIndex].State = UAbilitySystemComponent::EAbilityExecutionState::Succeeded;
+			}
+		}
+		else if (ExecutingAbilityIndex >= 0)
+		{
+			Component.ExecutingServerAbilities[ExecutingAbilityIndex].State = UAbilitySystemComponent::EAbilityExecutionState::Failed;
+		}
+	}
+}
+
+// ----------------------------------------------------------------------------------------------------------------------------------------------------
+//								Input 
+// ----------------------------------------------------------------------------------------------------------------------------------------------------
+
 void UAbilitySystemComponent::MontageBranchPoint_AbilityDecisionStop()
 {
 	if (AnimatingAbility)
@@ -502,7 +683,7 @@ void UAbilitySystemComponent::NotifyAbilityCommit(UGameplayAbility* Ability)
 	AbilityCommitedCallbacks.Broadcast(Ability);
 }
 
-void UAbilitySystemComponent::NotifyAbilityActivated(UGameplayAbility* Ability)
+void UAbilitySystemComponent::NotifyAbilityActivated(FGameplayAbilitySpecHandle Handle, UGameplayAbility* Ability)
 {
 	AbilityActivatedCallbacks.Broadcast(Ability);
 }
@@ -510,15 +691,13 @@ void UAbilitySystemComponent::NotifyAbilityActivated(UGameplayAbility* Ability)
 void UAbilitySystemComponent::HandleGameplayEvent(FGameplayTag EventTag, FGameplayEventData* Payload)
 {
 	if (GameplayEventTriggeredAbilities.Contains(EventTag))
-	{
-		TArray<TWeakObjectPtr<UGameplayAbility> > TriggeredAbilities = GameplayEventTriggeredAbilities[EventTag];
+	{		
+		TArray<FGameplayAbilitySpecHandle> TriggeredAbilityHandles = GameplayEventTriggeredAbilities[EventTag];
 
-		for (auto Ability : TriggeredAbilities)
+		for (auto AbilityHandle : TriggeredAbilityHandles)
 		{
-			if (Ability.IsValid())
-			{
-				Ability.Get()->TriggerAbilityFromGameplayEvent(AbilityActorInfo.Get(), EventTag, Payload, *this);
-			}
+			TriggerAbilityFromGameplayEvent(AbilityHandle, AbilityActorInfo.Get(), EventTag, Payload, *this);
+						
 		}
 	}
 }
@@ -592,11 +771,34 @@ void UAbilitySystemComponent::BindAbilityActivationToInputComponent(UInputCompon
 void UAbilitySystemComponent::AbilityInputPressed(int32 InputID)
 {
 	// FIXME: Maps or multicast delegate to actually do this
-	for (FGameplayAbilityInputIDPair& Pair: ActivatableAbilities)
+	for (FGameplayAbilitySpec& Spec : ActivatableAbilities)
 	{
-		if (Pair.InputID == InputID)
+		if (Spec.InputID == InputID)
 		{
-			Pair.Ability->InputPressed(AbilityActorInfo.Get());
+			if (Spec.Ability)
+			{
+				if (Spec.IsActive)
+				{
+					// The ability is active, so just pipe the input event to it
+					if (Spec.Ability->GetInstancingPolicy() ==  EGameplayAbilityInstancingPolicy::NonInstanced)
+					{
+						Spec.Ability->InputPressed(Spec.Handle, AbilityActorInfo.Get(), Spec.ActivationInfo);
+					}
+					else
+					{
+						TArray<UGameplayAbility*> Instances = Spec.GetAbilityInstances();
+						for (UGameplayAbility* Instance : Instances)
+						{
+							Instance->InputPressed(Spec.Handle, AbilityActorInfo.Get(), Spec.ActivationInfo);
+						}						
+					}
+				}
+				else
+				{
+					// Ability is not active, so try to activate it
+					TryActivateAbility(Spec.Handle);
+				}
+			}
 		}
 	}
 }
@@ -604,11 +806,33 @@ void UAbilitySystemComponent::AbilityInputPressed(int32 InputID)
 void UAbilitySystemComponent::AbilityInputReleased(int32 InputID)
 {
 	// FIXME: Maps or multicast delegate to actually do this
-	for (FGameplayAbilityInputIDPair& Pair : ActivatableAbilities)
+	for (FGameplayAbilitySpec& Spec : ActivatableAbilities)
 	{
-		if (Pair.InputID == InputID)
+		if (Spec.InputID == InputID)
 		{
-			Pair.Ability->InputReleased(AbilityActorInfo.Get());
+			if (Spec.Ability)
+			{
+				if (Spec.IsActive)
+				{
+					// The ability is active, so just pipe the input event to it
+					if (Spec.Ability->GetInstancingPolicy() == EGameplayAbilityInstancingPolicy::NonInstanced)
+					{
+						Spec.Ability->InputReleased(Spec.Handle, AbilityActorInfo.Get(), Spec.ActivationInfo);
+					}
+					else
+					{
+						TArray<UGameplayAbility*> Instances = Spec.GetAbilityInstances();
+						for (UGameplayAbility* Instance : Instances)
+						{
+							Instance->InputReleased(Spec.Handle, AbilityActorInfo.Get(), Spec.ActivationInfo);
+						}
+					}
+				}
+				else
+				{
+
+				}
+			}
 		}
 	}
 }
diff --git a/Engine/Source/Runtime/GameplayAbilities/Private/AbilitySystemStats.cpp b/Engine/Source/Runtime/GameplayAbilities/Private/AbilitySystemStats.cpp
--- a/Engine/Source/Runtime/GameplayAbilities/Private/AbilitySystemStats.cpp
+++ b/Engine/Source/Runtime/GameplayAbilities/Private/AbilitySystemStats.cpp
@@ -15,5 +15,6 @@ DEFINE_STAT(STAT_GetGameplayCueFunction);
 DEFINE_STAT(STAT_GetOutgoingSpec);
 DEFINE_STAT(STAT_InitAttributeSetDefaults);
 DEFINE_STAT(STAT_TickAbilityTasks);
+DEFINE_STAT(STAT_FindAbilitySpecFromHandle);
 
 
diff --git a/Engine/Source/Runtime/GameplayAbilities/Private/GameplayAbilityTargetTypes.cpp b/Engine/Source/Runtime/GameplayAbilities/Private/GameplayAbilityTargetTypes.cpp
--- a/Engine/Source/Runtime/GameplayAbilities/Private/GameplayAbilityTargetTypes.cpp
+++ b/Engine/Source/Runtime/GameplayAbilities/Private/GameplayAbilityTargetTypes.cpp
@@ -3,6 +3,187 @@
 #include "AbilitySystemPrivatePCH.h"
 #include "GameplayAbilityTargetTypes.h"
 
+void FGameplayAbilityTargetData::ApplyGameplayEffect(UGameplayEffect* GameplayEffect, const FGameplayAbilityActorInfo InstigatorInfo)
+{
+	// TODO: Improve relationship between InstigatorContext and FGameplayAbilityTargetData/FHitResult (or use something different between HitResult)
+
+
+	FGameplayEffectSpec	SpecToApply(GameplayEffect,					// The UGameplayEffect data asset
+		InstigatorInfo.Actor.Get(),		// The actor who instigated this
+		1.f,							// FIXME: Leveling
+		NULL							// FIXME: CurveData override... should we just remove this?
+		);
+	if (HasHitResult())
+	{
+		SpecToApply.InstigatorContext.AddHitResult(*GetHitResult());
+	}
+
+	if (HasOrigin())
+	{
+		SpecToApply.InstigatorContext.AddOrigin(GetOrigin().GetLocation());
+	}
+
+	TArray<AActor*> Actors = GetActors();
+	for (AActor * TargetActor : Actors)
+	{
+		check(TargetActor);
+		UAbilitySystemComponent * TargetComponent = UAbilitySystemBlueprintLibrary::GetAbilitySystemComponent(TargetActor);
+		if (TargetComponent)
+		{
+			InstigatorInfo.AbilitySystemComponent->ApplyGameplayEffectSpecToTarget(SpecToApply, TargetComponent);
+		}
+	}
+}
+
+
+FString FGameplayAbilityTargetData::ToString() const
+{
+	return TEXT("BASE CLASS");
+}
+
+
+FGameplayAbilityTargetDataHandle FGameplayAbilityTargetingLocationInfo::MakeTargetDataHandleFromHitResult(TWeakObjectPtr<UGameplayAbility> Ability, FHitResult HitResult) const
+{
+	if (LocationType == EGameplayAbilityTargetingLocationType::Type::SocketTransform)
+	{
+		const FGameplayAbilityActorInfo* ActorInfo = Ability.IsValid() ? Ability.Get()->GetCurrentActorInfo() : NULL;
+		AActor* AISourceActor = ActorInfo ? (ActorInfo->Actor.IsValid() ? ActorInfo->Actor.Get() : NULL) : NULL;
+		UAnimInstance* AnimInstance = ActorInfo ? ActorInfo->AnimInstance.Get() : NULL;
+		USkeletalMeshComponent* AISourceComponent = AnimInstance ? AnimInstance->GetOwningComponent() : NULL;
+
+		if (AISourceActor && AISourceComponent)
+		{
+			FGameplayAbilityTargetData_Mesh* ReturnData = new FGameplayAbilityTargetData_Mesh();
+			ReturnData->SourceActor = AISourceActor;
+			ReturnData->SourceComponent = AISourceComponent;
+			ReturnData->SourceSocketName = SourceSocketName;
+			ReturnData->TargetPoint = HitResult.Location;
+			return FGameplayAbilityTargetDataHandle(ReturnData);
+		}
+	}
+
+	/** Note: These are cleaned up by the FGameplayAbilityTargetDataHandle (via an internal TSharedPtr) */
+	FGameplayAbilityTargetData_SingleTargetHit* ReturnData = new FGameplayAbilityTargetData_SingleTargetHit();
+	ReturnData->HitResult = HitResult;
+	return FGameplayAbilityTargetDataHandle(ReturnData);
+}
+
+FGameplayAbilityTargetDataHandle FGameplayAbilityTargetingLocationInfo::MakeTargetDataHandleFromActors(TArray<AActor*> TargetActors) const
+{
+	/** Note: This is cleaned up by the FGameplayAbilityTargetDataHandle (via an internal TSharedPtr) */
+	FGameplayAbilityTargetData_ActorArray* ReturnData = new FGameplayAbilityTargetData_ActorArray();
+	ReturnData->TargetActorArray = TargetActors;
+	ReturnData->SourceLocation = *this;
+	return FGameplayAbilityTargetDataHandle(ReturnData);
+}
+
+bool FGameplayAbilityTargetDataHandle::NetSerialize(FArchive& Ar, class UPackageMap* Map, bool& bOutSuccess)
+{
+
+	UScriptStruct* ScriptStruct = Data.IsValid() ? Data->GetScriptStruct() : NULL;
+	Ar << ScriptStruct;
+
+	if (ScriptStruct)
+	{
+		if (Ar.IsLoading())
+		{
+			// For now, just always reset/reallocate the data when loading.
+			// Longer term if we want to generalize this and use it for property replication, we should support
+			// only reallocating when necessary
+			check(!Data.IsValid());
+
+			FGameplayAbilityTargetData* NewData = (FGameplayAbilityTargetData*)FMemory::Malloc(ScriptStruct->GetCppStructOps()->GetSize());
+			ScriptStruct->InitializeScriptStruct(NewData);
+
+			Data = TSharedPtr<FGameplayAbilityTargetData>(NewData);
+		}
+
+		void* ContainerPtr = Data.Get();
+
+		if (ScriptStruct->StructFlags & STRUCT_NetSerializeNative)
+		{
+			ScriptStruct->GetCppStructOps()->NetSerialize(Ar, Map, bOutSuccess, Data.Get());
+		}
+		else
+		{
+			// This won't work since UStructProperty::NetSerializeItem is deprecrated.
+			//	1) we have to manually crawl through the topmost struct's fields since we don't have a UStructProperty for it (just the UScriptProperty)
+			//	2) if there are any UStructProperties in the topmost struct's fields, we will assert in UStructProperty::NetSerializeItem.
+
+			ABILITY_LOG(Fatal, TEXT("FGameplayAbilityTargetDataHandle::NetSerialize called on data struct %s without a native NetSerialize"), *ScriptStruct->GetName());
+		}
+	}
+
+	//ABILITY_LOG(Warning, TEXT("FGameplayAbilityTargetDataHandle Serialized: %s"), ScriptStruct ? *ScriptStruct->GetName() : TEXT("NULL") );
+	bOutSuccess = true;
+	return true;
+}
+
+bool FGameplayAbilityTargetingLocationInfo::NetSerialize(FArchive& Ar, class UPackageMap* Map, bool& bOutSuccess)
+{
+	Ar << LocationType;
+
+	switch (LocationType)
+	{
+	case EGameplayAbilityTargetingLocationType::ActorTransform:
+		Ar << SourceActor;
+		break;
+	case EGameplayAbilityTargetingLocationType::SocketTransform:
+		Ar << SourceComponent;
+		Ar << SourceSocketName;
+		break;
+	case EGameplayAbilityTargetingLocationType::LiteralTransform:
+		Ar << LiteralTransform;
+		break;
+	default:
+		check(false);		//This case should not happen
+		break;
+	}
+
+	bOutSuccess = true;
+	return true;
+}
+
+bool FGameplayAbilityTargetData_LocationInfo::NetSerialize(FArchive& Ar, class UPackageMap* Map, bool& bOutSuccess)
+{
+	SourceLocation.NetSerialize(Ar, Map, bOutSuccess);
+	TargetLocation.NetSerialize(Ar, Map, bOutSuccess);
+
+	bOutSuccess = true;
+	return true;
+}
+
+bool FGameplayAbilityTargetData_ActorArray::NetSerialize(FArchive& Ar, class UPackageMap* Map, bool& bOutSuccess)
+{
+	SourceLocation.NetSerialize(Ar, Map, bOutSuccess);
+	Ar << TargetActorArray;
+
+	bOutSuccess = true;
+	return true;
+}
+
+bool FGameplayAbilityTargetData_Mesh::NetSerialize(FArchive& Ar, class UPackageMap* Map, bool& bOutSuccess)
+{
+	//SourceActor can be used as a backup if the component isn't found.
+	Ar << SourceActor;
+	Ar << SourceComponent;
+	Ar << SourceSocketName;
+	TargetPoint.NetSerialize(Ar, Map, bOutSuccess);
+
+	bOutSuccess = true;
+	return true;
+}
+
+bool FGameplayAbilityTargetData_SingleTargetHit::NetSerialize(FArchive& Ar, class UPackageMap* Map, bool& bOutSuccess)
+{
+	Ar << HitResult.Actor;
+
+	HitResult.Location.NetSerialize(Ar, Map, bOutSuccess);
+	HitResult.Normal.NetSerialize(Ar, Map, bOutSuccess);
+
+	bOutSuccess = true;
+	return true;
+}
 
 bool FGameplayAbilityTargetData_Radius::NetSerialize(FArchive& Ar, class UPackageMap* Map, bool& bOutSuccess)
 {
diff --git a/Engine/Source/Runtime/GameplayAbilities/Private/GameplayAbilityTypes.cpp b/Engine/Source/Runtime/GameplayAbilities/Private/GameplayAbilityTypes.cpp
new file mode 100644
--- /dev/null
+++ b/Engine/Source/Runtime/GameplayAbilities/Private/GameplayAbilityTypes.cpp
@@ -0,0 +1,94 @@
+// Copyright 1998-2014 Epic Games, Inc. All Rights Reserved.
+
+#include "AbilitySystemPrivatePCH.h"
+
+//----------------------------------------------------------------------
+
+void FGameplayAbilityActorInfo::InitFromActor(AActor *InActor, UAbilitySystemComponent* InAbilitySystemComponent)
+{
+	check(InActor);
+	check(InAbilitySystemComponent);
+
+	Actor = InActor;
+	AbilitySystemComponent = InAbilitySystemComponent;
+
+	// Look for a player controller or pawn in the owner chain.
+	AActor *TestActor = InActor;
+	while (TestActor)
+	{
+		if (APlayerController * CastPC = Cast<APlayerController>(TestActor))
+		{
+			PlayerController = CastPC;
+			break;
+		}
+
+		if (APawn * Pawn = Cast<APawn>(TestActor))
+		{
+			PlayerController = Cast<APlayerController>(Pawn->GetController());
+			break;
+		}
+
+		TestActor = TestActor->GetOwner();
+	}
+
+	// Grab Components that we care about
+	USkeletalMeshComponent * SkelMeshComponent = InActor->FindComponentByClass<USkeletalMeshComponent>();
+	if (SkelMeshComponent)
+	{
+		this->AnimInstance = SkelMeshComponent->GetAnimInstance();
+	}
+
+	MovementComponent = InActor->FindComponentByClass<UMovementComponent>();
+}
+
+void FGameplayAbilityActorInfo::ClearActorInfo()
+{
+	Actor = NULL;
+	PlayerController = NULL;
+	AnimInstance = NULL;
+	MovementComponent = NULL;
+}
+
+bool FGameplayAbilityActorInfo::IsLocallyControlled() const
+{
+	if (PlayerController.IsValid())
+	{
+		return PlayerController->IsLocalController();
+	}
+
+	return false;
+}
+
+bool FGameplayAbilityActorInfo::IsNetAuthority() const
+{
+	if (Actor.IsValid())
+	{
+		return (Actor->Role == ROLE_Authority);
+	}
+
+	// If we encounter issues with this being called before or after the owning actor is destroyed,
+	// we may need to cache off the authority (or look for it on some global/world state).
+
+	ensure(false);
+	return false;
+}
+
+void FGameplayAbilityActivationInfo::GeneratePredictionKey(UAbilitySystemComponent * Component) const
+{
+	check(Component);
+	check(ActivationMode == EGameplayAbilityActivationMode::NonAuthority);
+
+	PrevPredictionKey = CurrPredictionKey;
+	CurrPredictionKey = Component->GetNextPredictionKey();
+	ActivationMode = EGameplayAbilityActivationMode::Predicting;
+}
+
+void FGameplayAbilityActivationInfo::SetActivationConfirmed()
+{
+	ActivationMode = EGameplayAbilityActivationMode::Confirmed;
+}
+
+
+// -----------------------------------------------------------------
+
+
diff --git a/Engine/Source/Runtime/GameplayAbilities/Public/Abilities/GameplayAbility.h b/Engine/Source/Runtime/GameplayAbilities/Public/Abilities/GameplayAbility.h
--- a/Engine/Source/Runtime/GameplayAbilities/Public/Abilities/GameplayAbility.h
+++ b/Engine/Source/Runtime/GameplayAbilities/Public/Abilities/GameplayAbility.h
@@ -63,6 +63,7 @@
  */
 
 
+/** TriggerData */
 USTRUCT()
 struct FAbilityTriggerData
 {
@@ -108,22 +109,13 @@ public:
 	// ----------------------------------------------------------------------------------------------------------------
 	
 	/** Input binding. Base implementation calls TryActivateAbility */
-	virtual void InputPressed(const FGameplayAbilityActorInfo* ActorInfo);
+	virtual void InputPressed(const FGameplayAbilitySpecHandle Handle, const FGameplayAbilityActorInfo* ActorInfo, const FGameplayAbilityActivationInfo ActivationInfo);
 
 	/** Input binding. Base implementation does nothing */
-	virtual void InputReleased(const FGameplayAbilityActorInfo* ActorInfo);
-	
-	/**
-	 * Attempts to activate the ability.
-	 *	-This function calls CanActivateAbility
-	 *	-This function handles instancing
-	 *	-This function handles networking and prediction
-	 *	-If all goes well, CallActivateAbility is called next.
-	 */
-	bool TryActivateAbility(const FGameplayAbilityActorInfo* ActorInfo, uint32 PrevPredictionKey = 0, uint32 CurrPredictionKey = 0, UGameplayAbility ** OutInstancedAbility = NULL);
+	virtual void InputReleased(const FGameplayAbilitySpecHandle Handle, const FGameplayAbilityActorInfo* ActorInfo, const FGameplayAbilityActivationInfo ActivationInfo);
 
 	/** Returns true if this ability can be activated right now. Has no side effects */
-	virtual bool CanActivateAbility(const FGameplayAbilityActorInfo* ActorInfo) const;
+	virtual bool CanActivateAbility(const FGameplayAbilitySpecHandle Handle, const FGameplayAbilityActorInfo* ActorInfo) const;
 
 	/** Returns true if this ability can be triggered right now. Has no side effects */
 	virtual bool ShouldAbilityRespondToEvent(FGameplayTag EventTag, const FGameplayEventData* Payload) const;
@@ -219,16 +211,14 @@ protected:
 	UFUNCTION(BlueprintImplementableEvent, Category = Ability, FriendlyName = "ActivateAbility")
 	virtual void K2_ActivateAbility();	
 
-	virtual void ActivateAbility(const FGameplayAbilityActorInfo * ActorInfo, const FGameplayAbilityActivationInfo ActivationInfo);
+	virtual void ActivateAbility(const FGameplayAbilitySpecHandle Handle, const FGameplayAbilityActorInfo* ActorInfo, const FGameplayAbilityActivationInfo ActivationInfo);
 
 	bool HasBlueprintActivate;
 
-	void CallActivateAbility(const FGameplayAbilityActorInfo* ActorInfo, const FGameplayAbilityActivationInfo ActivationInfo);
+	void CallActivateAbility(const FGameplayAbilitySpecHandle Handle, const FGameplayAbilityActorInfo* ActorInfo, const FGameplayAbilityActivationInfo ActivationInfo);
 
 	/** Called on a predictive ability when the server confirms its execution */
-	void ConfirmActivateSucceed();
-
-	virtual void TriggerAbilityFromGameplayEvent(FGameplayAbilityActorInfo* ActorInfo, FGameplayTag Tag, FGameplayEventData* Payload, UAbilitySystemComponent& Component);
+	virtual void ConfirmActivateSucceed();
 
 	UFUNCTION(BlueprintCallable, Category = "Abilities")
 	void SendGameplayEvent(FGameplayTag EventTag, FGameplayEventData Payload);
@@ -244,7 +234,7 @@ protected:
 	UFUNCTION(BlueprintCallable, Category = Ability, FriendlyName = "CommitAbility")
 	virtual bool K2_CommitAbility();
 
-	virtual bool CommitAbility(const FGameplayAbilityActorInfo* ActorInfo, const FGameplayAbilityActivationInfo ActivationInfo);
+	virtual bool CommitAbility(const FGameplayAbilitySpecHandle Handle, const FGameplayAbilityActorInfo* ActorInfo, const FGameplayAbilityActivationInfo ActivationInfo);
 
 	/**
 	 * The last chance to fail before commiting
@@ -263,14 +253,14 @@ protected:
 	virtual void CommitExecute(const FGameplayAbilityActorInfo* ActorInfo, const FGameplayAbilityActivationInfo ActivationInfo);
 
 	/** Do boilerplate init stuff and then call ActivateAbility */
-	void PreActivate(const FGameplayAbilityActorInfo* ActorInfo, const FGameplayAbilityActivationInfo ActivationInfo);
+	void PreActivate(const FGameplayAbilitySpecHandle Handle, const FGameplayAbilityActorInfo* ActorInfo, const FGameplayAbilityActivationInfo ActivationInfo);
 
 	// --------------------------------------
 	//	CancelAbility
 	// --------------------------------------
 
 	/** Destroys intanced-per-execution abilities. Instance-per-actor abilities should 'reset'. Non instance abilities - what can we do? */
-	virtual void CancelAbility(const FGameplayAbilityActorInfo* ActorInfo, const FGameplayAbilityActivationInfo ActivationInfo);
+	virtual void CancelAbility(const FGameplayAbilitySpecHandle Handle, const FGameplayAbilityActorInfo* ActorInfo, const FGameplayAbilityActivationInfo ActivationInfo);
 
 	/** Destroys intanced-per-execution abilities. Instance-per-actor abilities should 'reset'. Non instance abilities - what can we do? */
 	UFUNCTION(BlueprintCallable, Category = Ability)
@@ -296,15 +286,7 @@ protected:
 	UFUNCTION(BlueprintCallable, Category = Ability, FriendlyName="EndAbility")
 	virtual void K2_EndAbility();
 
-	virtual void EndAbility(const FGameplayAbilityActorInfo* ActorInfo);
-
-	// -------------------------------------
-	
-	/** Ask server to try to activate ability */
-	virtual void ServerTryActivateAbility(const FGameplayAbilityActorInfo* ActorInfo, const FGameplayAbilityActivationInfo ActivationInfo);
-
-	/** Called once server OKs a client ability activation */
-	virtual void ClientActivateAbility(const FGameplayAbilityActorInfo* ActorInfo, const FGameplayAbilityActivationInfo ActivationInfo);
+	virtual void EndAbility(const FGameplayAbilitySpecHandle Handle, const FGameplayAbilityActorInfo* ActorInfo, const FGameplayAbilityActivationInfo ActivationInfo);
 
 	// -------------------------------------
 	//	GameplayCue
@@ -329,20 +311,22 @@ protected:
 		return !HasAllFlags(RF_ClassDefaultObject);
 	}
 
-	void SetCurrentActorInfo(const FGameplayAbilityActorInfo* ActorInfo) const
+	void SetCurrentActorInfo(const FGameplayAbilitySpecHandle Handle, const FGameplayAbilityActorInfo* ActorInfo) const
 	{
 		if (IsInstantiated())
 		{
 			CurrentActorInfo = ActorInfo;
+			CurrentSpecHandle = Handle;
 		}
 	}
 
-	void SetCurrentInfo(const FGameplayAbilityActorInfo* ActorInfo, const FGameplayAbilityActivationInfo ActivationInfo)
+	void SetCurrentInfo(const FGameplayAbilitySpecHandle Handle, const FGameplayAbilityActorInfo* ActorInfo, const FGameplayAbilityActivationInfo ActivationInfo)
 	{
 		if (IsInstantiated())
 		{
 			CurrentActivationInfo = ActivationInfo;
 			CurrentActorInfo = ActorInfo;
+			CurrentSpecHandle = Handle;
 		}
 	}
 
@@ -353,12 +337,7 @@ protected:
 	USkeletalMeshComponent* GetOwningComponentFromActorInfo() const;
 
 	UFUNCTION(BlueprintCallable, Category=Ability)
-	FGameplayEffectSpecHandle GetOutgoingSpec(UGameplayEffect* GameplayEffect) const;
-
-	UFUNCTION(Client, Reliable)
-	void ClientActivateAbilitySucceed(int32 PredictionKey);
-
-	virtual void ClientActivateAbilitySucceed_Internal(int32 PredictionKey);
+	FGameplayEffectSpecHandle GetOutgoingSpec(UGameplayEffect* GameplayEffect) const;	
 
 	bool IsSupportedForNetworking() const override;
 
@@ -432,7 +411,22 @@ protected:
 
 	UFUNCTION(BlueprintPure, Category = Ability, meta = (HidePin = "WorldContextObject", DefaultToSelf = "WorldContextObject"))
 	FGameplayAbilityTargetingLocationInfo MakeTargetLocationInfoFromOwnerSkeletalMeshComponent(FName SocketName);
-private:
+
+	// ----------------------------------------------------------------------------------------------------------------
+	//
+	//	
+	//
+	// 
+	// ----------------------------------------------------------------------------------------------------------------
+	
+	/** Returns current level of the Ability */
+	UFUNCTION(BlueprintCallable, Category = Ability, meta = (HidePin = "WorldContextObject", DefaultToSelf = "WorldContextObject"))
+	int32 GetAbilityLevel() const;
+
+	/** Returns current ability level for non instanced abilities. You m ust call this version in these contexts! */
+	int32 GetAbilityLevel(FGameplayAbilitySpecHandle Handle, const FGameplayAbilityActorInfo* ActorInfo) const;
+
+protected:
 
 	/** 
 	 *  This is shared, cached information about the thing using us
@@ -443,6 +437,8 @@ private:
 	 */
 	mutable const FGameplayAbilityActorInfo* CurrentActorInfo;
 
+	mutable FGameplayAbilitySpecHandle CurrentSpecHandle;
+
 	/** True if the ability is currently active. For instance per owner abilities */
 	bool bIsActive;
 };
diff --git a/Engine/Source/Runtime/GameplayAbilities/Public/Abilities/GameplayAbilityTargetTypes.h b/Engine/Source/Runtime/GameplayAbilities/Public/Abilities/GameplayAbilityTargetTypes.h
--- a/Engine/Source/Runtime/GameplayAbilities/Public/Abilities/GameplayAbilityTargetTypes.h
+++ b/Engine/Source/Runtime/GameplayAbilities/Public/Abilities/GameplayAbilityTargetTypes.h
@@ -245,6 +245,8 @@ struct GAMEPLAYABILITIES_API FGameplayAbilityTargetingLocationInfo
 
 	FGameplayAbilityTargetingLocationInfo()
 	: LocationType(EGameplayAbilityTargetingLocationType::LiteralTransform)
+	, SourceActor(nullptr)
+	, SourceComponent(nullptr)
 	{
 	};
 
@@ -291,23 +293,23 @@ public:
 
 	/** Type of location used - will determine what data is transmitted over the network and what fields are used when calculating position. */
 	UPROPERTY(BlueprintReadWrite, meta = (ExposeOnSpawn = true), Category = Targeting)
-		TEnumAsByte<EGameplayAbilityTargetingLocationType::Type> LocationType;
+	TEnumAsByte<EGameplayAbilityTargetingLocationType::Type> LocationType;
 
 	/** A literal world transform can be used, if one has been calculated outside of the actor using the ability. */
 	UPROPERTY(BlueprintReadWrite, meta = (ExposeOnSpawn = true), Category = Targeting)
-		FTransform LiteralTransform;
+	FTransform LiteralTransform;
 
 	/** A source actor is needed for Actor-based targeting, but not for Socket-based targeting. */
 	UPROPERTY(BlueprintReadWrite, meta = (ExposeOnSpawn = true), Category = Targeting)
-		AActor* SourceActor;
+	AActor* SourceActor;
 
 	/** Socket-based targeting requires a skeletal mesh component to check for the named socket. */
 	UPROPERTY(BlueprintReadWrite, meta = (ExposeOnSpawn = true), Category = Targeting)
-		USkeletalMeshComponent* SourceComponent;
+	USkeletalMeshComponent* SourceComponent;
 
 	/** If SourceComponent is valid, this is the name of the socket transform that will be used. If no Socket is provided, SourceComponent's transform will be used. */
 	UPROPERTY(BlueprintReadWrite, meta = (ExposeOnSpawn = true), Category = Targeting)
-		FName SourceSocketName;
+	FName SourceSocketName;
 
 	// -------------------------------------
 
@@ -573,7 +575,7 @@ struct GAMEPLAYABILITIES_API FGameplayAbilityTargetData_SingleTargetHit : public
 	// -------------------------------------
 
 	UPROPERTY()
-		FHitResult	HitResult;
+	FHitResult	HitResult;
 
 	bool NetSerialize(FArchive& Ar, class UPackageMap* Map, bool& bOutSuccess);
 
diff --git a/Engine/Source/Runtime/GameplayAbilities/Public/Abilities/GameplayAbilityTypes.h b/Engine/Source/Runtime/GameplayAbilities/Public/Abilities/GameplayAbilityTypes.h
--- a/Engine/Source/Runtime/GameplayAbilities/Public/Abilities/GameplayAbilityTypes.h
+++ b/Engine/Source/Runtime/GameplayAbilities/Public/Abilities/GameplayAbilityTypes.h
@@ -60,97 +60,69 @@ namespace EGameplayAbilityReplicationPolicy
 }
 
 
+UENUM(BlueprintType)
+namespace EGameplayAbilityActivationMode
+{
+	enum Type
+	{
+		Authority,			// We are the authority activating this ability
+		NonAuthority,		// We are not the authority but aren't predicting yet. This is a mostly invalid state to be in.
+
+		Predicting,			// We are predicting the activation of this ability
+		Confirmed,			// We are not the authority, but the authority has confirmed this activation
+	};
+}
+
+// ----------------------------------------------------
+
 USTRUCT(BlueprintType)
-struct FGameplayAbilityHandle
+struct FGameplayAbilitySpecHandle
 {
 	GENERATED_USTRUCT_BODY()
 
-	FGameplayAbilityHandle()
+	FGameplayAbilitySpecHandle()
 	: Handle(INDEX_NONE)
 	{
 
 	}
 
-	FGameplayAbilityHandle(int32 InHandle)
-		: Handle(InHandle)
-	{
-
-	}
-
 	bool IsValid() const
 	{
 		return Handle != INDEX_NONE;
 	}
 
-	FGameplayAbilityHandle GetNextHandle()
+	void GenerateNewHandle()
 	{
-		return FGameplayAbilityHandle(Handle + 1);
+		static int32 GHandle = 1;
+		Handle = GHandle++;
 	}
 
-	bool operator==(const FGameplayAbilityHandle& Other) const
+	bool operator==(const FGameplayAbilitySpecHandle& Other) const
 	{
 		return Handle == Other.Handle;
 	}
 
-	bool operator!=(const FGameplayAbilityHandle& Other) const
+	bool operator!=(const FGameplayAbilitySpecHandle& Other) const
 	{
 		return Handle != Other.Handle;
 	}
 
-	FString ToString() const
-	{
-		return FString::Printf(TEXT("%d"), Handle);
-	}
-
 private:
 
 	UPROPERTY()
 	int32 Handle;
 };
 
-
-USTRUCT()
-struct FGameplayAbilityInputIDPair
-{
-	GENERATED_USTRUCT_BODY()
-
-	FGameplayAbilityInputIDPair()
-	: Ability(nullptr), InputID(-1) { }
-
-	FGameplayAbilityInputIDPair(UGameplayAbility* InAbility, int32 InInputID)
-	: Ability(InAbility), InputID(InInputID) { }
-
-	UPROPERTY()
-	UGameplayAbility* Ability;
-
-	UPROPERTY()
-	int32	InputID;
-};
-
-
-UENUM(BlueprintType)
-namespace EGameplayAbilityActivationMode
-{
-	enum Type
-	{
-		Authority,			// We are the authority activating this ability
-		NonAuthority,		// We are not the authority but aren't predicting yet. This is a mostly invalid state to be in.
-
-		Predicting,			// We are predicting the activation of this ability
-		Confirmed,			// We are not the authority, but the authority has confirmed this activation
-	};
-}
-
 /**
-*	FGameplayAbilityActorInfo
-*
-*	Cached data associated with an Actor using an Ability.
-*		-Initialized from an AActor* in InitFromActor
-*		-Abilities use this to know what to actor upon. E.g., instead of being coupled to a specific actor class.
-*		-These are generally passed around as pointers to support polymorphism.
-*		-Projects can override UAbilitySystemGlobals::AllocAbilityActorInfo to override the default struct type that is created.
-*
-*/
+ *	FGameplayAbilityActorInfo
+ *
+ *	Cached data associated with an Actor using an Ability.
+ *		-Initialized from an AActor* in InitFromActor
+ *		-Abilities use this to know what to actor upon. E.g., instead of being coupled to a specific actor class.
+ *		-These are generally passed around as pointers to support polymorphism.
+ *		-Projects can override UAbilitySystemGlobals::AllocAbilityActorInfo to override the default struct type that is created.
+ *
+ */
 USTRUCT(BlueprintType)
 struct GAMEPLAYABILITIES_API FGameplayAbilityActorInfo
 {
@@ -183,32 +155,23 @@ struct GAMEPLAYABILITIES_API FGameplayAbilityActorInfo
 };
 
 /**
-*	FGameplayAbilityActorInfo
-*
-*	Contains data used to determine locations for targeting systems.
-*
-*/
-
-
-/**
-*	FGameplayAbilityActivationInfo
-*
-*	Data tied to a specific activation of an ability.
-*		-Tell us whether we are the authority, if we are predicting, confirmed, etc.
-*		-Holds current and previous PredictionKey
-*		-Generally not meant to be subclassed in projects.
-*		-Passed around by value since the struct is small.
-*/
-
+ *	FGameplayAbilityActivationInfo
+ *
+ *	Data tied to a specific activation of an ability.
+ *		-Tell us whether we are the authority, if we are predicting, confirmed, etc.
+ *		-Holds current and previous PredictionKey
+ *		-Generally not meant to be subclassed in projects.
+ *		-Passed around by value since the struct is small.
+ */
 USTRUCT(BlueprintType)
 struct GAMEPLAYABILITIES_API FGameplayAbilityActivationInfo
 {
 	GENERATED_USTRUCT_BODY()
 
 	FGameplayAbilityActivationInfo()
-	: ActivationMode(EGameplayAbilityActivationMode::Authority)
-	, PrevPredictionKey(0)
-	, CurrPredictionKey(0)
+		: ActivationMode(EGameplayAbilityActivationMode::Authority)
+		, PrevPredictionKey(0)
+		, CurrPredictionKey(0)
 	{
 
 	}
@@ -229,7 +192,6 @@ struct GAMEPLAYABILITIES_API FGameplayAbilityActivationInfo
 	{
 	}
 
-
 	void GeneratePredictionKey(UAbilitySystemComponent * Component) const;
 
 	void SetActivationConfirmed();
@@ -244,6 +206,65 @@ struct GAMEPLAYABILITIES_API FGameplayAbilityActivationInfo
 	mutable uint32 CurrPredictionKey;
 };
 
+// ----------------------------------------------------
+
+USTRUCT()
+struct GAMEPLAYABILITIES_API FGameplayAbilitySpec
+{
+	GENERATED_USTRUCT_BODY()
+
+	FGameplayAbilitySpec()
+	: IsActive(false)
+	{
+		
+	}
+
+	FGameplayAbilitySpec(UGameplayAbility* InAbility, int32 InLevel=1, int32 InInputID=INDEX_NONE)
+	: Ability(InAbility), Level(InLevel), InputID(InInputID), IsActive(false)
+	{
+		Handle.GenerateNewHandle();
+	}
+
+	/** Handle for outside sources to refer to this spec by */
+	UPROPERTY()
+	FGameplayAbilitySpecHandle Handle;
+	
+	/** Ability of the spec (either CDO or instanced) */
+	UPROPERTY()
+	UGameplayAbility* Ability;
+	
+	/** Level of Ability */
+	UPROPERTY()
+	int32	Level;
+
+	/** InputID, if bound */
+	UPROPERTY()
+	int32	InputID;
+
+	/** Is this ability active? (Literally: has EndAbility been called on it since ActivateAbility was called on it */
+	UPROPERTY()
+	bool	IsActive;
+
+	/** Activation state of this ability. This is not replicated since it needs to be overwritten locally on clients during prediction. */
+	UPROPERTY(NotReplicated)
+	FGameplayAbilityActivationInfo	ActivationInfo;
+
+	/** Non replicating instances of this ability. */
+	UPROPERTY(NotReplicated)
+	TArray<UGameplayAbility*> NonReplicatedInstances;
+
+	/** Replicated instances of this ability.. */
+	UPROPERTY()
+	TArray<UGameplayAbility*> ReplicatedInstances;
+
+	TArray<UGameplayAbility*> GetAbilityInstances()
+	{
+		TArray<UGameplayAbility*> Abilities;
+		Abilities.Append(ReplicatedInstances);
+		Abilities.Append(NonReplicatedInstances);
+		return Abilities;
+	}
+};
 
 // ----------------------------------------------------
 
diff --git a/Engine/Source/Runtime/GameplayAbilities/Public/Abilities/GameplayAbility_CharacterJump.h b/Engine/Source/Runtime/GameplayAbilities/Public/Abilities/GameplayAbility_CharacterJump.h
--- a/Engine/Source/Runtime/GameplayAbilities/Public/Abilities/GameplayAbility_CharacterJump.h
+++ b/Engine/Source/Runtime/GameplayAbilities/Public/Abilities/GameplayAbility_CharacterJump.h
@@ -15,11 +15,11 @@ class GAMEPLAYABILITIES_API UGameplayAbility_CharacterJump : public UGameplayAbi
 
 public:
 
-	virtual bool CanActivateAbility(const FGameplayAbilityActorInfo* ActorInfo) const override;
+	virtual bool CanActivateAbility(const FGameplayAbilitySpecHandle Handle, const FGameplayAbilityActorInfo* ActorInfo) const override;
 	
-	virtual void ActivateAbility(const FGameplayAbilityActorInfo* OwnerInfo, const FGameplayAbilityActivationInfo ActivationInfo) override;
+	virtual void ActivateAbility(const FGameplayAbilitySpecHandle Handle, const FGameplayAbilityActorInfo* OwnerInfo, const FGameplayAbilityActivationInfo ActivationInfo) override;
 
-	virtual void InputReleased(const FGameplayAbilityActorInfo* ActorInfo) override;
+	virtual void InputReleased(const FGameplayAbilitySpecHandle Handle, const FGameplayAbilityActorInfo* ActorInfo, const FGameplayAbilityActivationInfo ActivationInfo) override;
 
-	virtual void CancelAbility(const FGameplayAbilityActorInfo* ActorInfo, const FGameplayAbilityActivationInfo ActivationInfo) override;
+	virtual void CancelAbility(const FGameplayAbilitySpecHandle Handle, const FGameplayAbilityActorInfo* ActorInfo, const FGameplayAbilityActivationInfo ActivationInfo) override;
 };
\ No newline at end of file
diff --git a/Engine/Source/Runtime/GameplayAbilities/Public/Abilities/GameplayAbility_Montage.h b/Engine/Source/Runtime/GameplayAbilities/Public/Abilities/GameplayAbility_Montage.h
--- a/Engine/Source/Runtime/GameplayAbilities/Public/Abilities/GameplayAbility_Montage.h
+++ b/Engine/Source/Runtime/GameplayAbilities/Public/Abilities/GameplayAbility_Montage.h
@@ -15,7 +15,7 @@ class GAMEPLAYABILITIES_API UGameplayAbility_Montage : public UGameplayAbility
 
 public:
 	
-	virtual void ActivateAbility(const FGameplayAbilityActorInfo* OwnerInfo, const FGameplayAbilityActivationInfo ActivationInfo) override;
+	virtual void ActivateAbility(const FGameplayAbilitySpecHandle Handle, const FGameplayAbilityActorInfo* OwnerInfo, const FGameplayAbilityActivationInfo ActivationInfo) override;
 
 	UPROPERTY(EditDefaultsOnly, Category = MontageAbility)
 	UAnimMontage *	MontageToPlay;
diff --git a/Engine/Source/Runtime/GameplayAbilities/Public/AbilitySystemComponent.h b/Engine/Source/Runtime/GameplayAbilities/Public/AbilitySystemComponent.h
--- a/Engine/Source/Runtime/GameplayAbilities/Public/AbilitySystemComponent.h
+++ b/Engine/Source/Runtime/GameplayAbilities/Public/AbilitySystemComponent.h
@@ -55,6 +55,10 @@ struct GAMEPLAYABILITIES_API FAttributeDefaults
  * 
  */
 
+
+/**
+ *	The core ActorComponent for interfacing with the GameplayAbilities System
+ */
 UCLASS(ClassGroup=AbilitySystem, hidecategories=(Object,LOD,Lighting,Transform,Sockets,TextureStreaming), editinlinenew, meta=(BlueprintSpawnableComponent))
 class GAMEPLAYABILITIES_API UAbilitySystemComponent : public UActorComponent, public IGameplayTagAssetInterface
 {
@@ -175,12 +179,12 @@ class GAMEPLAYABILITIES_API UAbilitySystemComponent : public UActorComponent, pu
 		{
 			return PrevPredictionKey == Other.PrevPredictionKey
 				&& CurrPredictionKey == Other.CurrPredictionKey
-				&& Ability == Other.Ability;
+				&& Handle == Other.Handle;
 		}
 
 		uint32 PrevPredictionKey;
 		uint32 CurrPredictionKey;
-		UGameplayAbility* Ability;
+		FGameplayAbilitySpecHandle Handle;
 	};
 
 	// This is a list of GameplayAbilities that are predicted by the client and were triggered by abilities that were also predicted by the client
@@ -206,7 +210,7 @@ class GAMEPLAYABILITIES_API UAbilitySystemComponent : public UActorComponent, pu
 
 		uint32 CurrPredictionKey;
 		EAbilityExecutionState State;
-		UGameplayAbility* Ability;
+		FGameplayAbilitySpecHandle Handle;
 	};
 
 	TArray<FExecutingAbilityInfo> ExecutingServerAbilities;
@@ -370,27 +374,29 @@ class GAMEPLAYABILITIES_API UAbilitySystemComponent : public UActorComponent, pu
 	 *	
 	 */
 
-	UGameplayAbility* GiveAbility(UGameplayAbility* Ability, int32 InputID);
+	/** Grants Ability. Returns handle that can be used in TryActivateAbility, etc. */
+	FGameplayAbilitySpecHandle GiveAbility(FGameplayAbilitySpec AbilitySpec);
+
+	/** Attempts to activate the given ability */
+	bool TryActivateAbility(FGameplayAbilitySpecHandle AbilityToActivate, uint32 PrevPredictionKey = 0, uint32 CurrPredictionKey = 0, UGameplayAbility ** OutInstancedAbility = nullptr);
+
+	void TriggerAbilityFromGameplayEvent(FGameplayAbilitySpecHandle AbilityToTrigger, FGameplayAbilityActorInfo* ActorInfo, FGameplayTag Tag, FGameplayEventData* Payload, UAbilitySystemComponent& Component);
 
 	/** Wipes all 'given' abilities. */
 	void ClearAllAbilities();
 
 	/** Will be called from GiveAbility or from OnRep. Initializes events (triggers and inputs) with the given ability */
-	void OnGiveAbility(UGameplayAbility* Ability, int32 InputID);
+	void OnGiveAbility(const FGameplayAbilitySpec AbilitySpec);
 
-	UGameplayAbility* CreateNewInstanceOfAbility(UGameplayAbility* Ability);
+	UGameplayAbility* CreateNewInstanceOfAbility(FGameplayAbilitySpec& Spec, UGameplayAbility* Ability);
 
-	void CancelAbilitiesWithTags(const FGameplayTagContainer Tags, const FGameplayAbilityActorInfo* ActorInfo, const FGameplayAbilityActivationInfo ActivationInfo, UGameplayAbility* Ignore);
+	void CancelAbilitiesWithTags(const FGameplayTagContainer Tags, const FGameplayAbilitySpecHandle Handle, const FGameplayAbilityActorInfo* ActorInfo, const FGameplayAbilityActivationInfo ActivationInfo, UGameplayAbility* Ignore);
 	
-	/** References to non-replicating abilities that we instanced. We need to keep these references to avoid GC */
+	/** FUll list of all instance-per-execution gameplay abilities associated with this component */
 	UPROPERTY()
-	TArray<UGameplayAbility*>	NonReplicatedInstancedAbilities;
-
-	/** References to replicating abilities that we instanced. We need to keep these references to avoid GC */
-	UPROPERTY(Replicated)
-	TArray<UGameplayAbility*>	ReplicatedInstancedAbilities;
+	TArray<UGameplayAbility*>	AllReplicatedInstancedAbilities;
 
-	void NotifyAbilityEnded(UGameplayAbility* Ability);
+	void NotifyAbilityEnded(FGameplayAbilitySpecHandle Handle, UGameplayAbility* Ability);
 
 	/**
 	 *	The abilities we can activate. 
@@ -400,22 +406,25 @@ class GAMEPLAYABILITIES_API UAbilitySystemComponent : public UActorComponent, pu
 	 *	This array is not vital for things to work. It is a convenience thing for 'giving abilities to the actor'. But abilities could also work on things
 	 *	without an AbilitySystemComponent. For example an ability could be written to execute on a StaticMeshActor. As long as the ability doesn't require 
 	 *	instancing or anything else that the AbilitySystemComponent would provide, then it doesn't need the component to function.
-	 */ 
-	UPROPERTY(ReplicatedUsing = OnRep_ActivateAbilities, BlueprintReadOnly, Category = "Abilities")
-	TArray<FGameplayAbilityInputIDPair>	ActivatableAbilities;
+	 */
+
+	UPROPERTY(Replicated=OnRep_ActivateAbilities, BlueprintReadOnly, Category = "Abilities")
+	TArray<FGameplayAbilitySpec>	ActivatableAbilities;
 
+	FGameplayAbilitySpec* FindAbilitySpecFromHandle(FGameplayAbilitySpecHandle Handle);
 	
 	UFUNCTION()
 	void	OnRep_ActivateAbilities();
 
 	UFUNCTION(Server, reliable, WithValidation)
-	void	ServerTryActivateAbility(class UGameplayAbility* AbilityToActivate, uint32 PrevPredictionKey, uint32 CurrPredictionKey);
+	void	ServerTryActivateAbility(FGameplayAbilitySpecHandle AbilityToActivate, uint32 PrevPredictionKey, uint32 CurrPredictionKey);
 
 	UFUNCTION(Client, Reliable)
-	void	ClientActivateAbilityFailed(class UGameplayAbility* AbilityToActivate, uint32 PredictionKey);
+	void	ClientActivateAbilityFailed(FGameplayAbilitySpecHandle AbilityToActivate, uint32 PredictionKey);
 
 	UFUNCTION(Client, Reliable)
-	void	ClientActivateAbilitySucceed(class UGameplayAbility* AbilityToActivate, uint32 PredictionKey);
+	void	ClientActivateAbilitySucceed(FGameplayAbilitySpecHandle AbilityToActivate, uint32 PredictionKey);
+	
 
 	// ----------------------------------------------------------------------------------------------------------------
 
@@ -456,10 +465,10 @@ class GAMEPLAYABILITIES_API UAbilitySystemComponent : public UActorComponent, pu
 
 	void HandleGameplayEvent(FGameplayTag EventTag, FGameplayEventData* Payload);
 
-	TMap<FGameplayTag, TArray<TWeakObjectPtr<UGameplayAbility> > > GameplayEventTriggeredAbilities;
+	TMap<FGameplayTag, TArray<FGameplayAbilitySpecHandle > > GameplayEventTriggeredAbilities;
 
 	void NotifyAbilityCommit(UGameplayAbility* Ability);
-	void NotifyAbilityActivated(UGameplayAbility* Ability);
+	void NotifyAbilityActivated(const FGameplayAbilitySpecHandle Handle, UGameplayAbility* Ability);
 
 
 	UPROPERTY()
@@ -475,9 +484,6 @@ class GAMEPLAYABILITIES_API UAbilitySystemComponent : public UActorComponent, pu
 
 	// ----------------------------------------------------------------------------------------------------------------
 
-	/** This is temp for testing. We want to think a bit more on the API for outside stuff activating abilities */
-	bool	ActivateAbility(TWeakObjectPtr<UGameplayAbility> Ability);
-
 	/** Adds a UAbilityTask task to the list of tasks to be ticked */
 	void TickingTaskStarted(UAbilityTask* NewTask);
 
diff --git a/Engine/Source/Runtime/GameplayAbilities/Public/AbilitySystemStats.h b/Engine/Source/Runtime/GameplayAbilities/Public/AbilitySystemStats.h
--- a/Engine/Source/Runtime/GameplayAbilities/Public/AbilitySystemStats.h
+++ b/Engine/Source/Runtime/GameplayAbilities/Public/AbilitySystemStats.h
@@ -18,4 +18,5 @@ DECLARE_CYCLE_STAT_EXTERN(TEXT("CreateNewActiveGameplayEffect"), STAT_CreateNewA
 DECLARE_CYCLE_STAT_EXTERN(TEXT("GetGameplayCueFunction"), STAT_GetGameplayCueFunction, STATGROUP_AbilitySystem, );
 DECLARE_CYCLE_STAT_EXTERN(TEXT("GetOutgoingSpec"), STAT_GetOutgoingSpec, STATGROUP_AbilitySystem, );
 DECLARE_CYCLE_STAT_EXTERN(TEXT("InitAttributeSetDefaults"), STAT_InitAttributeSetDefaults, STATGROUP_AbilitySystem, );
-DECLARE_CYCLE_STAT_EXTERN(TEXT("TickAbilityTasks"), STAT_TickAbilityTasks, STATGROUP_AbilitySystem, );
\ No newline at end of file
+DECLARE_CYCLE_STAT_EXTERN(TEXT("TickAbilityTasks"), STAT_TickAbilityTasks, STATGROUP_AbilitySystem, );
+DECLARE_CYCLE_STAT_EXTERN(TEXT("FindAbilitySpecFromHandle"), STAT_FindAbilitySpecFromHandle, STATGROUP_AbilitySystem, );
\ No newline at end of file
