nonunity fix

diff --git a/Engine/Source/Runtime/GameplayAbilities/Private/GameplayAbilityTargetTypes.cpp b/Engine/Source/Runtime/GameplayAbilities/Private/GameplayAbilityTargetTypes.cpp
--- a/Engine/Source/Runtime/GameplayAbilities/Private/GameplayAbilityTargetTypes.cpp
+++ b/Engine/Source/Runtime/GameplayAbilities/Private/GameplayAbilityTargetTypes.cpp
@@ -2,6 +2,9 @@
 
 #include "AbilitySystemPrivatePCH.h"
 #include "GameplayAbilityTargetTypes.h"
+#include "GameplayEffect.h"
+#include "AbilitySystemComponent.h"
+
 
 void FGameplayAbilityTargetData::ApplyGameplayEffect(UGameplayEffect* GameplayEffect, const FGameplayAbilityActorInfo InstigatorInfo)
 {
