move nonunity fix

diff --git a/Engine/Source/Runtime/GameplayAbilities/Private/GameplayAbilityTargetTypes.cpp b/Engine/Source/Runtime/GameplayAbilities/Private/GameplayAbilityTargetTypes.cpp
--- a/Engine/Source/Runtime/GameplayAbilities/Private/GameplayAbilityTargetTypes.cpp
+++ b/Engine/Source/Runtime/GameplayAbilities/Private/GameplayAbilityTargetTypes.cpp
@@ -4,6 +4,7 @@
 #include "GameplayAbilityTargetTypes.h"
 #include "GameplayEffect.h"
 #include "AbilitySystemComponent.h"
+#include "AbilitySystemBlueprintLibrary.h"
 
 
 void FGameplayAbilityTargetData::ApplyGameplayEffect(UGameplayEffect* GameplayEffect, const FGameplayAbilityActorInfo InstigatorInfo)
diff --git a/Engine/Source/Runtime/GameplayAbilities/Private/GameplayAbilityTypes.cpp b/Engine/Source/Runtime/GameplayAbilities/Private/GameplayAbilityTypes.cpp
--- a/Engine/Source/Runtime/GameplayAbilities/Private/GameplayAbilityTypes.cpp
+++ b/Engine/Source/Runtime/GameplayAbilities/Private/GameplayAbilityTypes.cpp
@@ -1,6 +1,7 @@
 // Copyright 1998-2014 Epic Games, Inc. All Rights Reserved.
 
 #include "AbilitySystemPrivatePCH.h"
+#include "AbilitySystemComponent.h"
 
 //----------------------------------------------------------------------
 
