first pass on towers

diff --git a/Engine/Source/Runtime/GameplayAbilities/Private/AbilitySystemBlueprintLibrary.cpp b/Engine/Source/Runtime/GameplayAbilities/Private/AbilitySystemBlueprintLibrary.cpp
--- a/Engine/Source/Runtime/GameplayAbilities/Private/AbilitySystemBlueprintLibrary.cpp
+++ b/Engine/Source/Runtime/GameplayAbilities/Private/AbilitySystemBlueprintLibrary.cpp
@@ -22,7 +22,7 @@ UAbilitySystemComponent* UAbilitySystemBlueprintLibrary::GetAbilitySystemCompone
 
 void UAbilitySystemBlueprintLibrary::ApplyGameplayEffectToTargetData(FGameplayAbilityTargetDataHandle Target, UGameplayEffect *GameplayEffect, const FGameplayAbilityActorInfo InstigatorInfo)
 {
-	if (Target.Data.IsValid())
+	if (Target.Data.IsValid() && GameplayEffect != NULL)
 	{
 		Target.Data->ApplyGameplayEffect(GameplayEffect, InstigatorInfo);
 	}
