Added WaitInputRelease task

diff --git a/Engine/Source/Runtime/GameplayAbilities/Private/Abilities/GameplayAbility.cpp b/Engine/Source/Runtime/GameplayAbilities/Private/Abilities/GameplayAbility.cpp
--- a/Engine/Source/Runtime/GameplayAbilities/Private/Abilities/GameplayAbility.cpp
+++ b/Engine/Source/Runtime/GameplayAbilities/Private/Abilities/GameplayAbility.cpp
@@ -333,6 +333,8 @@ void UGameplayAbility::InputPressed(const FGameplayAbilitySpecHandle Handle, con
 
 void UGameplayAbility::InputReleased(const FGameplayAbilitySpecHandle Handle, const FGameplayAbilityActorInfo* ActorInfo, const FGameplayAbilityActivationInfo ActivationInfo)
 {
+	OnInputRelease.Broadcast(this);
+	OnInputRelease.Clear();
 }
 
 bool UGameplayAbility::CheckCooldown(const FGameplayAbilityActorInfo* ActorInfo) const
@@ -580,4 +582,11 @@ int32 UGameplayAbility::GetAbilityLevel(FGameplayAbilitySpecHandle Handle, const
 	check(Spec);
 
 	return Spec->Level;
+}
+
+FGameplayAbilitySpec* UGameplayAbility::GetCurrentAbilitySpec() const
+{
+	check(IsInstantiated()); // You should not call this on non instanced abilities.
+	check(CurrentActorInfo);
+	return CurrentActorInfo->AbilitySystemComponent->FindAbilitySpecFromHandle(CurrentSpecHandle);
 }
\ No newline at end of file
diff --git a/Engine/Source/Runtime/GameplayAbilities/Private/Abilities/Tasks/AbilityTask_WaitInputRelease.cpp b/Engine/Source/Runtime/GameplayAbilities/Private/Abilities/Tasks/AbilityTask_WaitInputRelease.cpp
new file mode 100644
--- /dev/null
+++ b/Engine/Source/Runtime/GameplayAbilities/Private/Abilities/Tasks/AbilityTask_WaitInputRelease.cpp
@@ -0,0 +1,63 @@
+
+#include "AbilitySystemPrivatePCH.h"
+#include "Abilities/Tasks/AbilityTask_WaitInputRelease.h"
+
+UAbilityTask_WaitInputRelease::UAbilityTask_WaitInputRelease(const class FPostConstructInitializeProperties& PCIP)
+	: Super(PCIP)
+{
+	StartTime = 0.f;
+	RegisteredCallback= false;
+}
+
+void UAbilityTask_WaitInputRelease::OnReleaseCallback(UGameplayAbility* InAbility)
+{
+	float ElapsedTime = GetWorld()->GetTimeSeconds() - StartTime;
+
+	if (Ability->GetCurrentActivationInfo().ActivationMode != EGameplayAbilityActivationMode::Authority)
+	{
+		// Tell the server we released.
+		AbilitySystemComponent->ServerInputRelease(Ability->GetCurrentAbilitySpecHandle());
+	}
+
+	// We are done. Kill us so we don't keep getting broadcast messages
+	OnRelease.Broadcast(ElapsedTime);
+	EndTask();
+}
+
+UAbilityTask_WaitInputRelease* UAbilityTask_WaitInputRelease::WaitInputRelease(class UObject* WorldContextObject)
+{
+	return NewTask<UAbilityTask_WaitInputRelease>(WorldContextObject);
+}
+
+void UAbilityTask_WaitInputRelease::Activate()
+{
+	if (Ability.IsValid())
+	{
+		FGameplayAbilitySpec* Spec = Ability->GetCurrentAbilitySpec();
+		if (Spec)
+		{
+			if (!Spec->InputPressed)
+			{
+				// Input was already released before we got here, we are done.
+				OnRelease.Broadcast(0.f);
+				EndTask();
+			}
+			else
+			{
+				StartTime = GetWorld()->GetTimeSeconds();
+				Ability->OnInputRelease.AddUObject(this, &UAbilityTask_WaitInputRelease::OnReleaseCallback);
+				RegisteredCallback = true;
+			}
+		}
+	}
+}
+
+void UAbilityTask_WaitInputRelease::OnDestroy(bool AbilityEnded)
+{
+	if (RegisteredCallback && Ability.IsValid())
+	{
+		Ability->OnInputRelease.RemoveUObject(this, &UAbilityTask_WaitInputRelease::OnReleaseCallback);
+	}
+
+	Super::OnDestroy(AbilityEnded);
+}
diff --git a/Engine/Source/Runtime/GameplayAbilities/Private/AbilitySystemComponent_Abilities.cpp b/Engine/Source/Runtime/GameplayAbilities/Private/AbilitySystemComponent_Abilities.cpp
--- a/Engine/Source/Runtime/GameplayAbilities/Private/AbilitySystemComponent_Abilities.cpp
+++ b/Engine/Source/Runtime/GameplayAbilities/Private/AbilitySystemComponent_Abilities.cpp
@@ -453,7 +453,8 @@ void UAbilitySystemComponent::ServerTryActivateAbility_Implementation(FGameplayA
 		return;
 	}
 
-	UGameplayAbility *InstancedAbility = NULL;
+	UGameplayAbility* InstancedAbility = NULL;
+	Spec->InputPressed = true; // Pretend input was pressed. Allows UAbilityTask_WaitInputRelease to work.
 
 	// Attempt to activate the ability (server side) and tell the client if it succeeded or failed.
 	if (TryActivateAbility(Handle, PrevPredictionKey, CurrPredictionKey, &InstancedAbility))
@@ -471,6 +472,7 @@ void UAbilitySystemComponent::ServerTryActivateAbility_Implementation(FGameplayA
 	}
 	else
 	{
+		Spec->InputPressed = false;
 		ClientActivateAbilityFailed(Handle, CurrPredictionKey);
 	}
 }
@@ -777,6 +779,7 @@ void UAbilitySystemComponent::AbilityInputPressed(int32 InputID)
 		{
 			if (Spec.Ability)
 			{
+				Spec.InputPressed = true;
 				if (Spec.IsActive)
 				{
 					// The ability is active, so just pipe the input event to it
@@ -812,26 +815,28 @@ void UAbilitySystemComponent::AbilityInputReleased(int32 InputID)
 		{
 			if (Spec.Ability)
 			{
-				if (Spec.IsActive)
-				{
-					// The ability is active, so just pipe the input event to it
-					if (Spec.Ability->GetInstancingPolicy() == EGameplayAbilityInstancingPolicy::NonInstanced)
-					{
-						Spec.Ability->InputReleased(Spec.Handle, AbilityActorInfo.Get(), Spec.ActivationInfo);
-					}
-					else
-					{
-						TArray<UGameplayAbility*> Instances = Spec.GetAbilityInstances();
-						for (UGameplayAbility* Instance : Instances)
-						{
-							Instance->InputReleased(Spec.Handle, AbilityActorInfo.Get(), Spec.ActivationInfo);
-						}
-					}
-				}
-				else
-				{
+				AbilitySpectInputReleased(Spec);
+			}
+		}
+	}
+}
 
-				}
+void UAbilitySystemComponent::AbilitySpectInputReleased(FGameplayAbilitySpec& Spec)
+{
+	Spec.InputPressed = false;
+	if (Spec.IsActive)
+	{
+		// The ability is active, so just pipe the input event to it
+		if (Spec.Ability->GetInstancingPolicy() == EGameplayAbilityInstancingPolicy::NonInstanced)
+		{
+			Spec.Ability->InputReleased(Spec.Handle, AbilityActorInfo.Get(), Spec.ActivationInfo);
+		}
+		else
+		{
+			TArray<UGameplayAbility*> Instances = Spec.GetAbilityInstances();
+			for (UGameplayAbility* Instance : Instances)
+			{
+				Instance->InputReleased(Spec.Handle, AbilityActorInfo.Get(), Spec.ActivationInfo);
 			}
 		}
 	}
@@ -859,6 +864,20 @@ void UAbilitySystemComponent::InputCancel()
 	CancelCallbacks.Broadcast();
 }
 
+bool UAbilitySystemComponent::ServerInputRelease_Validate(FGameplayAbilitySpecHandle Handle)
+{
+	return true;
+}
+
+void UAbilitySystemComponent::ServerInputRelease_Implementation(FGameplayAbilitySpecHandle Handle)
+{
+	FGameplayAbilitySpec* Spec = FindAbilitySpecFromHandle(Handle);
+	if (Spec)
+	{
+		AbilitySpectInputReleased(*Spec);
+	}
+}
+
 void UAbilitySystemComponent::TargetConfirm()
 {
 	for (AGameplayAbilityTargetActor* TargetActor : SpawnedTargetActors)
diff --git a/Engine/Source/Runtime/GameplayAbilities/Public/Abilities/GameplayAbility.h b/Engine/Source/Runtime/GameplayAbilities/Public/Abilities/GameplayAbility.h
--- a/Engine/Source/Runtime/GameplayAbilities/Public/Abilities/GameplayAbility.h
+++ b/Engine/Source/Runtime/GameplayAbilities/Public/Abilities/GameplayAbility.h
@@ -138,19 +138,29 @@ public:
 	}
 
 	/** Gets the current actor info bound to this ability - can only be called on instanced abilities. */
-	const FGameplayAbilityActorInfo* GetCurrentActorInfo()
+	const FGameplayAbilityActorInfo* GetCurrentActorInfo() const
 	{
 		check(IsInstantiated());
 		return CurrentActorInfo;
 	}
 
 	/** Gets the current activation info bound to this ability - can only be called on instanced abilities. */
-	FGameplayAbilityActivationInfo GetCurrentActivationInfo()
+	FGameplayAbilityActivationInfo GetCurrentActivationInfo() const
 	{
 		check(IsInstantiated());
 		return CurrentActivationInfo;
 	}
 
+	/** Gets the current AbilitySpecHandle- can only be called on instanced abilities. */
+	FGameplayAbilitySpecHandle GetCurrentAbilitySpecHandle() const
+	{
+		check(IsInstantiated());
+		return CurrentSpecHandle;
+	}
+
+	/** Retrieves the actual AbilitySpec for this ability. Can only be called on instanced abilities. */
+	FGameplayAbilitySpec* GetCurrentAbilitySpec() const;
+
 	virtual UWorld* GetWorld() const override
 	{
 		return GetOuter()->GetWorld();
@@ -175,7 +185,9 @@ public:
 	void TaskStarted(UAbilityTask* NewTask);
 
 	void TaskEnded(UAbilityTask* Task);
-	
+
+	/** GenericCallback for Input being released. Auto cleared on broadcast */
+	FGenericAbilityDelegate	OnInputRelease;
 
 protected:
 
diff --git a/Engine/Source/Runtime/GameplayAbilities/Public/Abilities/GameplayAbilityTypes.h b/Engine/Source/Runtime/GameplayAbilities/Public/Abilities/GameplayAbilityTypes.h
--- a/Engine/Source/Runtime/GameplayAbilities/Public/Abilities/GameplayAbilityTypes.h
+++ b/Engine/Source/Runtime/GameplayAbilities/Public/Abilities/GameplayAbilityTypes.h
@@ -214,13 +214,13 @@ struct GAMEPLAYABILITIES_API FGameplayAbilitySpec
 	GENERATED_USTRUCT_BODY()
 
 	FGameplayAbilitySpec()
-	: IsActive(false)
+	: Ability(nullptr), Level(1), InputID(INDEX_NONE), InputPressed(false), IsActive(false)
 	{
 		
 	}
 
 	FGameplayAbilitySpec(UGameplayAbility* InAbility, int32 InLevel=1, int32 InInputID=INDEX_NONE)
-	: Ability(InAbility), Level(InLevel), InputID(InInputID), IsActive(false)
+	: Ability(InAbility), Level(InLevel), InputID(InInputID), InputPressed(false), IsActive(false)
 	{
 		Handle.GenerateNewHandle();
 	}
@@ -241,6 +241,10 @@ struct GAMEPLAYABILITIES_API FGameplayAbilitySpec
 	UPROPERTY()
 	int32	InputID;
 
+	/** Is input currently pressed. Set to false when input is released */
+	UPROPERTY(NotReplicated)
+	bool	InputPressed;
+
 	/** Is this ability active? (Literally: has EndAbility been called on it since ActivateAbility was called on it */
 	UPROPERTY()
 	bool	IsActive;
diff --git a/Engine/Source/Runtime/GameplayAbilities/Public/Abilities/Tasks/AbilityTask_WaitInputRelease.h b/Engine/Source/Runtime/GameplayAbilities/Public/Abilities/Tasks/AbilityTask_WaitInputRelease.h
new file mode 100644
--- /dev/null
+++ b/Engine/Source/Runtime/GameplayAbilities/Public/Abilities/Tasks/AbilityTask_WaitInputRelease.h
@@ -0,0 +1,36 @@
+// Copyright 1998-2014 Epic Games, Inc. All Rights Reserved.
+#pragma once
+#include "AbilityTask.h"
+#include "Abilities/GameplayAbilityTypes.h"
+#include "AbilityTask_WaitInputRelease.generated.h"
+
+/**
+ *	Waits until the input is released from activating an ability. Clients will replicate a 'release input' event to the server, but not the exact time it was held locally.
+ *	We expect server to execute this task in parrallel and keep its own time.
+ */
+UCLASS(MinimalAPI)
+class UAbilityTask_WaitInputRelease : public UAbilityTask
+{
+	GENERATED_UCLASS_BODY()
+
+	DECLARE_DYNAMIC_MULTICAST_DELEGATE_OneParam(FInputReleaseDelegate, float, TimeHeld);
+
+	UPROPERTY(BlueprintAssignable)
+	FInputReleaseDelegate	OnRelease;
+
+	UFUNCTION()
+	void OnReleaseCallback(UGameplayAbility* InAbility);
+
+	virtual void Activate() override;
+
+	/** Wait until the user releases the input button for this ability's activation. Returns time from hitting this node, till release. Will return 0 if input was already released. */
+	UFUNCTION(BlueprintCallable, Category="Ability|Tasks", meta = (HidePin = "WorldContextObject", DefaultToSelf = "WorldContextObject", BlueprintInternalUseOnly = "TRUE"))
+	static UAbilityTask_WaitInputRelease* WaitInputRelease(UObject* WorldContextObject);
+
+protected:
+
+	virtual void OnDestroy(bool AbilityEnded) override;
+
+	float StartTime;
+	bool RegisteredCallback;
+};
\ No newline at end of file
diff --git a/Engine/Source/Runtime/GameplayAbilities/Public/AbilitySystemComponent.h b/Engine/Source/Runtime/GameplayAbilities/Public/AbilitySystemComponent.h
--- a/Engine/Source/Runtime/GameplayAbilities/Public/AbilitySystemComponent.h
+++ b/Engine/Source/Runtime/GameplayAbilities/Public/AbilitySystemComponent.h
@@ -451,6 +451,12 @@ class GAMEPLAYABILITIES_API UAbilitySystemComponent : public UActorComponent, pu
 
 	void AbilityInputReleased(int32 InputID);
 
+	void AbilitySpectInputReleased(FGameplayAbilitySpec& Spec);
+
+	/** Sent by abilities when they *need* to tell server when activation input is released. (Not sent by defauly only sent when using AbilityTask_WaitInputRelease) */
+	UFUNCTION(Server, reliable, WithValidation)
+	void ServerInputRelease(FGameplayAbilitySpecHandle AbilityToActivate);
+
 	UFUNCTION(BlueprintCallable, Category="Abilities")
 	void InputConfirm();
 
