NonUnity + style fix

diff --git a/Engine/Source/Runtime/GameplayAbilities/Public/Abilities/Tasks/AbilityTask_WaitInputRelease.h b/Engine/Source/Runtime/GameplayAbilities/Public/Abilities/Tasks/AbilityTask_WaitInputRelease.h
--- a/Engine/Source/Runtime/GameplayAbilities/Public/Abilities/Tasks/AbilityTask_WaitInputRelease.h
+++ b/Engine/Source/Runtime/GameplayAbilities/Public/Abilities/Tasks/AbilityTask_WaitInputRelease.h
@@ -2,6 +2,7 @@
 #pragma once
 #include "AbilityTask.h"
 #include "Abilities/GameplayAbilityTypes.h"
+#include "AbilitySystemComponent.h"
 #include "AbilityTask_WaitInputRelease.generated.h"
 
 /**
