allow WaitAbilityActivate tasks to ignore triggered abilities

diff --git a/Engine/Source/Runtime/GameplayAbilities/Private/Abilities/GameplayAbility.cpp b/Engine/Source/Runtime/GameplayAbilities/Private/Abilities/GameplayAbility.cpp
--- a/Engine/Source/Runtime/GameplayAbilities/Private/Abilities/GameplayAbility.cpp
+++ b/Engine/Source/Runtime/GameplayAbilities/Private/Abilities/GameplayAbility.cpp
@@ -589,4 +589,11 @@ FGameplayAbilitySpec* UGameplayAbility::GetCurrentAbilitySpec() const
 	check(IsInstantiated()); // You should not call this on non instanced abilities.
 	check(CurrentActorInfo);
 	return CurrentActorInfo->AbilitySystemComponent->FindAbilitySpecFromHandle(CurrentSpecHandle);
+}
+
+bool UGameplayAbility::IsTriggered() const
+{
+	// Assume that if there is triggered data, then we are triggered. 
+	// If we need to support abilities that can be both, this will need to be expanded.
+	return AbilityTriggers.Num() > 0;
 }
\ No newline at end of file
diff --git a/Engine/Source/Runtime/GameplayAbilities/Private/Abilities/Tasks/AbilityTask_WaitAbilityActivate.cpp b/Engine/Source/Runtime/GameplayAbilities/Private/Abilities/Tasks/AbilityTask_WaitAbilityActivate.cpp
--- a/Engine/Source/Runtime/GameplayAbilities/Private/Abilities/Tasks/AbilityTask_WaitAbilityActivate.cpp
+++ b/Engine/Source/Runtime/GameplayAbilities/Private/Abilities/Tasks/AbilityTask_WaitAbilityActivate.cpp
@@ -9,14 +9,15 @@
 UAbilityTask_WaitAbilityActivate::UAbilityTask_WaitAbilityActivate(const class FPostConstructInitializeProperties& PCIP)
 	: Super(PCIP)
 {
+	IncludeTriggeredAbilities = false;
 }
 
-UAbilityTask_WaitAbilityActivate* UAbilityTask_WaitAbilityActivate::WaitForAbilityActivate(UObject* WorldContextObject, FGameplayTag InWithTag, FGameplayTag InWithoutTag)
+UAbilityTask_WaitAbilityActivate* UAbilityTask_WaitAbilityActivate::WaitForAbilityActivate(UObject* WorldContextObject, FGameplayTag InWithTag, FGameplayTag InWithoutTag, bool InIncludeTriggeredAbilities)
 {
 	auto MyObj = NewTask<UAbilityTask_WaitAbilityActivate>(WorldContextObject);
 	MyObj->WithTag = InWithTag;
 	MyObj->WithoutTag = InWithoutTag;
-
+	MyObj->IncludeTriggeredAbilities = InIncludeTriggeredAbilities;
 	return MyObj;
 }
 
@@ -28,8 +29,13 @@ void UAbilityTask_WaitAbilityActivate::Activate()
 	}
 }
 
-void UAbilityTask_WaitAbilityActivate::OnAbilityActivate(UGameplayAbility *ActivatedAbility)
+void UAbilityTask_WaitAbilityActivate::OnAbilityActivate(UGameplayAbility* ActivatedAbility)
 {
+	if (!IncludeTriggeredAbilities && ActivatedAbility->IsTriggered())
+	{
+		return;
+	}
+
 	if ((WithTag.IsValid() && !ActivatedAbility->AbilityTags.HasTag(WithTag, EGameplayTagMatchType::IncludeParentTags, EGameplayTagMatchType::Explicit)) ||
 		(WithoutTag.IsValid() && ActivatedAbility->AbilityTags.HasTag(WithoutTag, EGameplayTagMatchType::IncludeParentTags, EGameplayTagMatchType::Explicit)))
 	{
diff --git a/Engine/Source/Runtime/GameplayAbilities/Public/Abilities/GameplayAbility.h b/Engine/Source/Runtime/GameplayAbilities/Public/Abilities/GameplayAbility.h
--- a/Engine/Source/Runtime/GameplayAbilities/Public/Abilities/GameplayAbility.h
+++ b/Engine/Source/Runtime/GameplayAbilities/Public/Abilities/GameplayAbility.h
@@ -189,6 +189,9 @@ public:
 	/** GenericCallback for Input being released. Auto cleared on broadcast */
 	FGenericAbilityDelegate	OnInputRelease;
 
+	/** Is this ability triggered from TriggerData (or is it triggered explicitly through input/game code) */
+	bool IsTriggered() const;
+
 protected:
 
 	// --------------------------------------
diff --git a/Engine/Source/Runtime/GameplayAbilities/Public/Abilities/Tasks/AbilityTask_WaitAbilityActivate.h b/Engine/Source/Runtime/GameplayAbilities/Public/Abilities/Tasks/AbilityTask_WaitAbilityActivate.h
--- a/Engine/Source/Runtime/GameplayAbilities/Public/Abilities/Tasks/AbilityTask_WaitAbilityActivate.h
+++ b/Engine/Source/Runtime/GameplayAbilities/Public/Abilities/Tasks/AbilityTask_WaitAbilityActivate.h
@@ -24,12 +24,13 @@ class UAbilityTask_WaitAbilityActivate : public UAbilityTask
 	UFUNCTION()
 	void OnAbilityActivate(UGameplayAbility *ActivatedAbility);
 
-	/** Wait until a new ability (of the same or different type) is activated. */
+	/** Wait until a new ability (of the same or different type) is activated. Only input based abilities will be counted unless IncludeTriggeredAbilities is true. */
 	UFUNCTION(BlueprintCallable, Category="Ability|Tasks", meta = (HidePin = "WorldContextObject", DefaultToSelf = "WorldContextObject", BlueprintInternalUseOnly = "TRUE"))
-	static UAbilityTask_WaitAbilityActivate* WaitForAbilityActivate(UObject* WorldContextObject, FGameplayTag WithTag, FGameplayTag WithoutTag);
+	static UAbilityTask_WaitAbilityActivate* WaitForAbilityActivate(UObject* WorldContextObject, FGameplayTag WithTag, FGameplayTag WithoutTag, bool IncludeTriggeredAbilities=false);
 
 	FGameplayTag WithTag;
 	FGameplayTag WithoutTag;
+	bool IncludeTriggeredAbilities;
 
 protected:
 
