changing from Actor* to TWeakObjectPtr<>

diff --git a/Engine/Source/Runtime/GameplayAbilities/Private/Abilities/GameplayAbilityTargetActor_Radius.cpp b/Engine/Source/Runtime/GameplayAbilities/Private/Abilities/GameplayAbilityTargetActor_Radius.cpp
--- a/Engine/Source/Runtime/GameplayAbilities/Private/Abilities/GameplayAbilityTargetActor_Radius.cpp
+++ b/Engine/Source/Runtime/GameplayAbilities/Private/Abilities/GameplayAbilityTargetActor_Radius.cpp
@@ -39,7 +39,7 @@ void AGameplayAbilityTargetActor_Radius::ConfirmTargetingAndContinue()
 	}
 }
 
-FGameplayAbilityTargetDataHandle AGameplayAbilityTargetActor_Radius::MakeTargetData(const TArray<AActor*> Actors, const FVector& Origin) const
+FGameplayAbilityTargetDataHandle AGameplayAbilityTargetActor_Radius::MakeTargetData(const TArray<TWeakObjectPtr<AActor>> Actors, const FVector& Origin) const
 {
 	if (OwningAbility)
 	{
@@ -51,12 +51,12 @@ FGameplayAbilityTargetDataHandle AGameplayAbilityTargetActor_Radius::MakeTargetD
 	return FGameplayAbilityTargetDataHandle();
 }
 
-TArray<AActor*>	AGameplayAbilityTargetActor_Radius::PerformOverlap(const FVector& Origin)
+TArray<TWeakObjectPtr<AActor>>	AGameplayAbilityTargetActor_Radius::PerformOverlap(const FVector& Origin)
 {
 	static FName RadiusTargetingOverlap = FName(TEXT("RadiusTargetingOverlap"));
 	bool bTraceComplex = false;
 	
-	TArray<AActor*> ActorsToIgnore;
+	TArray<TWeakObjectPtr<AActor>> ActorsToIgnore;
 	ActorsToIgnore.Add(SourceActor);
 
 	FCollisionQueryParams Params(RadiusTargetingOverlap, bTraceComplex);
@@ -69,7 +69,7 @@ TArray<AActor*>	AGameplayAbilityTargetActor_Radius::PerformOverlap(const FVector
 	SourceActor->GetWorld()->OverlapMulti(Overlaps, Origin, FQuat::Identity, ECC_Pawn, FCollisionShape::MakeSphere(Radius), Params);
 
 
-	TArray<AActor*>	HitActors;
+	TArray<TWeakObjectPtr<AActor>>	HitActors;
 
 	for (int32 i=0; i < Overlaps.Num(); ++i)
 	{
diff --git a/Engine/Source/Runtime/GameplayAbilities/Private/AbilitySystemBlueprintLibrary.cpp b/Engine/Source/Runtime/GameplayAbilities/Private/AbilitySystemBlueprintLibrary.cpp
--- a/Engine/Source/Runtime/GameplayAbilities/Private/AbilitySystemBlueprintLibrary.cpp
+++ b/Engine/Source/Runtime/GameplayAbilities/Private/AbilitySystemBlueprintLibrary.cpp
@@ -64,14 +64,14 @@ FGameplayAbilityTargetDataHandle UAbilitySystemBlueprintLibrary::AbilityTargetDa
 	return Handle;
 }
 
-TArray<AActor*> UAbilitySystemBlueprintLibrary::GetActorsFromTargetData(FGameplayAbilityTargetDataHandle TargetData)
+TArray<TWeakObjectPtr<AActor>> UAbilitySystemBlueprintLibrary::GetActorsFromTargetData(FGameplayAbilityTargetDataHandle TargetData)
 {
 	FGameplayAbilityTargetData* Data = TargetData.Data.Get();
 	if (Data)
 	{
 		return Data->GetActors();
 	}
-	return TArray<AActor*>();
+	return TArray<TWeakObjectPtr<AActor>>();
 }
 
 bool UAbilitySystemBlueprintLibrary::TargetDataHasHitResult(FGameplayAbilityTargetDataHandle TargetData)
diff --git a/Engine/Source/Runtime/GameplayAbilities/Private/GameplayAbilityTargetTypes.cpp b/Engine/Source/Runtime/GameplayAbilities/Private/GameplayAbilityTargetTypes.cpp
--- a/Engine/Source/Runtime/GameplayAbilities/Private/GameplayAbilityTargetTypes.cpp
+++ b/Engine/Source/Runtime/GameplayAbilities/Private/GameplayAbilityTargetTypes.cpp
@@ -27,14 +27,16 @@ void FGameplayAbilityTargetData::ApplyGameplayEffect(UGameplayEffect* GameplayEf
 		SpecToApply.InstigatorContext.AddOrigin(GetOrigin().GetLocation());
 	}
 
-	TArray<AActor*> Actors = GetActors();
-	for (AActor* TargetActor : Actors)
+	TArray<TWeakObjectPtr<AActor>> Actors = GetActors();
+	for (TWeakObjectPtr<AActor> TargetActor : Actors)
 	{
-		check(TargetActor);
-		UAbilitySystemComponent * TargetComponent = UAbilitySystemBlueprintLibrary::GetAbilitySystemComponent(TargetActor);
-		if (TargetComponent)
+		if (TargetActor.IsValid())
 		{
-			InstigatorInfo.AbilitySystemComponent->ApplyGameplayEffectSpecToTarget(SpecToApply, TargetComponent);
+			UAbilitySystemComponent * TargetComponent = UAbilitySystemBlueprintLibrary::GetAbilitySystemComponent(TargetActor.Get());
+			if (TargetComponent)
+			{
+				InstigatorInfo.AbilitySystemComponent->ApplyGameplayEffectSpecToTarget(SpecToApply, TargetComponent);
+			}
 		}
 	}
 }
@@ -72,7 +74,7 @@ FGameplayAbilityTargetDataHandle FGameplayAbilityTargetingLocationInfo::MakeTarg
 	return FGameplayAbilityTargetDataHandle(ReturnData);
 }
 
-FGameplayAbilityTargetDataHandle FGameplayAbilityTargetingLocationInfo::MakeTargetDataHandleFromActors(TArray<AActor*> TargetActors) const
+FGameplayAbilityTargetDataHandle FGameplayAbilityTargetingLocationInfo::MakeTargetDataHandleFromActors(TArray<TWeakObjectPtr<AActor>> TargetActors) const
 {
 	/** Note: This is cleaned up by the FGameplayAbilityTargetDataHandle (via an internal TSharedPtr) */
 	FGameplayAbilityTargetData_ActorArray* ReturnData = new FGameplayAbilityTargetData_ActorArray();
@@ -195,4 +197,4 @@ bool FGameplayAbilityTargetData_Radius::NetSerialize(FArchive& Ar, class UPackag
 	Ar << Origin;
 
 	return true;
-}
\ No newline at end of file
+}
diff --git a/Engine/Source/Runtime/GameplayAbilities/Public/Abilities/GameplayAbilityTargetActor_Radius.h b/Engine/Source/Runtime/GameplayAbilities/Public/Abilities/GameplayAbilityTargetActor_Radius.h
--- a/Engine/Source/Runtime/GameplayAbilities/Public/Abilities/GameplayAbilityTargetActor_Radius.h
+++ b/Engine/Source/Runtime/GameplayAbilities/Public/Abilities/GameplayAbilityTargetActor_Radius.h
@@ -23,8 +23,8 @@ public:
 
 protected:
 
-	TArray<AActor*>	PerformOverlap(const FVector& Origin);
+	TArray<TWeakObjectPtr<AActor>>	PerformOverlap(const FVector& Origin);
 
-	FGameplayAbilityTargetDataHandle MakeTargetData(const TArray<AActor*> Actors, const FVector& Origin) const;
+	FGameplayAbilityTargetDataHandle MakeTargetData(const TArray<TWeakObjectPtr<AActor>> Actors, const FVector& Origin) const;
 
 };
\ No newline at end of file
diff --git a/Engine/Source/Runtime/GameplayAbilities/Public/Abilities/GameplayAbilityTargetTypes.h b/Engine/Source/Runtime/GameplayAbilities/Public/Abilities/GameplayAbilityTargetTypes.h
--- a/Engine/Source/Runtime/GameplayAbilities/Public/Abilities/GameplayAbilityTargetTypes.h
+++ b/Engine/Source/Runtime/GameplayAbilities/Public/Abilities/GameplayAbilityTargetTypes.h
@@ -64,9 +64,9 @@ struct GAMEPLAYABILITIES_API FGameplayAbilityTargetData
 
 	void ApplyGameplayEffect(UGameplayEffect *GameplayEffect, const struct FGameplayAbilityActorInfo InstigatorInfo);
 
-	virtual TArray<AActor*>	GetActors() const
+	virtual TArray<TWeakObjectPtr<AActor>>	GetActors() const
 	{
-		return TArray<AActor*>();
+		return TArray<TWeakObjectPtr<AActor>>();
 	}
 
 	// -------------------------------------
@@ -207,10 +207,10 @@ struct GAMEPLAYABILITIES_API FGameplayAbilityTargetData_Radius : public FGamepla
 	FGameplayAbilityTargetData_Radius()
 	: Origin(0.f) { }
 
-	FGameplayAbilityTargetData_Radius(const TArray<AActor*>& InActors, const FVector& InOrigin)
+	FGameplayAbilityTargetData_Radius(const TArray<TWeakObjectPtr<AActor>> InActors, const FVector& InOrigin)
 		: Actors(InActors), Origin(InOrigin) { }
 
-	virtual TArray<AActor*>	GetActors() const { return Actors; }
+	virtual TArray<TWeakObjectPtr<AActor>>	GetActors() const { return Actors; }
 
 	virtual bool HasOrigin() const { return true; }
 
@@ -225,7 +225,7 @@ struct GAMEPLAYABILITIES_API FGameplayAbilityTargetData_Radius : public FGamepla
 
 private:
 
-	TArray<AActor*> Actors;
+	TArray<TWeakObjectPtr<AActor>> Actors;
 	FVector Origin;
 };
 
@@ -289,7 +289,7 @@ public:
 
 	FGameplayAbilityTargetDataHandle MakeTargetDataHandleFromHitResult(TWeakObjectPtr<UGameplayAbility> Ability, FHitResult HitResult) const;
 
-	FGameplayAbilityTargetDataHandle MakeTargetDataHandleFromActors(TArray<AActor*> TargetActors) const;
+	FGameplayAbilityTargetDataHandle MakeTargetDataHandleFromActors(TArray<TWeakObjectPtr<AActor>> TargetActors) const;
 
 	/** Type of location used - will determine what data is transmitted over the network and what fields are used when calculating position. */
 	UPROPERTY(BlueprintReadWrite, meta = (ExposeOnSpawn = true), Category = Targeting)
@@ -407,9 +407,9 @@ struct GAMEPLAYABILITIES_API FGameplayAbilityTargetData_ActorArray : public FGam
 
 	/** Rather than targeting a single point, this type of targeting selects multiple actors. */
 	UPROPERTY(BlueprintReadWrite, EditAnywhere, Category = Targeting)
-	TArray<AActor*> TargetActorArray;
+	TArray<TWeakObjectPtr<AActor>> TargetActorArray;
 
-	virtual TArray<AActor*>	GetActors() const override
+	virtual TArray<TWeakObjectPtr<AActor>>	GetActors() const override
 	{
 		return TargetActorArray;
 	}
@@ -550,9 +550,9 @@ struct GAMEPLAYABILITIES_API FGameplayAbilityTargetData_SingleTargetHit : public
 
 	// -------------------------------------
 
-	virtual TArray<AActor*>	GetActors() const override
+	virtual TArray<TWeakObjectPtr<AActor>>	GetActors() const override
 	{
-		TArray<AActor*>	Actors;
+		TArray<TWeakObjectPtr<AActor>>	Actors;
 		if (HitResult.Actor.IsValid())
 		{
 			Actors.Push(HitResult.Actor.Get());
diff --git a/Engine/Source/Runtime/GameplayAbilities/Public/AbilitySystemBlueprintLibrary.h b/Engine/Source/Runtime/GameplayAbilities/Public/AbilitySystemBlueprintLibrary.h
--- a/Engine/Source/Runtime/GameplayAbilities/Public/AbilitySystemBlueprintLibrary.h
+++ b/Engine/Source/Runtime/GameplayAbilities/Public/AbilitySystemBlueprintLibrary.h
@@ -38,7 +38,7 @@ class UAbilitySystemBlueprintLibrary : public UBlueprintFunctionLibrary
 	static FGameplayAbilityTargetDataHandle	AbilityTargetDataFromHitResult(FHitResult HitResult);
 
 	UFUNCTION(BlueprintPure, Category = "Ability|TargetData")
-	static TArray<AActor*> GetActorsFromTargetData(FGameplayAbilityTargetDataHandle TargetData);
+	static TArray<TWeakObjectPtr<AActor>> GetActorsFromTargetData(FGameplayAbilityTargetDataHandle TargetData);
 
 	UFUNCTION(BlueprintPure, Category = "Ability|TargetData")
 	static bool TargetDataHasHitResult(FGameplayAbilityTargetDataHandle HitResult);
