Target data handles now contain an array of data instead of just a single member.

diff --git a/Engine/Source/Runtime/GameplayAbilities/Private/Abilities/Tasks/AbilityTask_WaitOverlap.cpp b/Engine/Source/Runtime/GameplayAbilities/Private/Abilities/Tasks/AbilityTask_WaitOverlap.cpp
--- a/Engine/Source/Runtime/GameplayAbilities/Private/Abilities/Tasks/AbilityTask_WaitOverlap.cpp
+++ b/Engine/Source/Runtime/GameplayAbilities/Private/Abilities/Tasks/AbilityTask_WaitOverlap.cpp
@@ -26,7 +26,7 @@ void UAbilityTask_WaitOverlap::OnHitCallback(AActor* OtherActor, UPrimitiveCompo
 
 		// Give it a handle and return
 		FGameplayAbilityTargetDataHandle	Handle;
-		Handle.Data = TSharedPtr<FGameplayAbilityTargetData>(TargetData);
+		Handle.Data.Add(TSharedPtr<FGameplayAbilityTargetData>(TargetData));
 		OnOverlap.Broadcast(Handle);
 
 		// We are done. Kill us so we don't keep getting broadcast messages
diff --git a/Engine/Source/Runtime/GameplayAbilities/Private/Abilities/Tasks/AbilityTask_WaitTargetData.cpp b/Engine/Source/Runtime/GameplayAbilities/Private/Abilities/Tasks/AbilityTask_WaitTargetData.cpp
--- a/Engine/Source/Runtime/GameplayAbilities/Private/Abilities/Tasks/AbilityTask_WaitTargetData.cpp
+++ b/Engine/Source/Runtime/GameplayAbilities/Private/Abilities/Tasks/AbilityTask_WaitTargetData.cpp
@@ -82,7 +82,7 @@ bool UAbilityTask_WaitTargetData::BeginSpawningActor(UObject* WorldContextObject
 			if (!CDO->ShouldProduceTargetDataOnServer)
 			{
 				//Problem here - if there's targeting data just sitting around, fire events don't get hooked up because of this if-else, even if we don't end the task.
-				if (AbilitySystemComponent->ReplicatedTargetData.IsValid())
+				if (AbilitySystemComponent->ReplicatedTargetData.IsValid(0))
 				{
 					ValidData.Broadcast(AbilitySystemComponent->ReplicatedTargetData);
 					if (ConfirmationType == EGameplayTargetingConfirmation::CustomMulti)
diff --git a/Engine/Source/Runtime/GameplayAbilities/Private/AbilitySystemBlueprintLibrary.cpp b/Engine/Source/Runtime/GameplayAbilities/Private/AbilitySystemBlueprintLibrary.cpp
--- a/Engine/Source/Runtime/GameplayAbilities/Private/AbilitySystemBlueprintLibrary.cpp
+++ b/Engine/Source/Runtime/GameplayAbilities/Private/AbilitySystemBlueprintLibrary.cpp
@@ -22,12 +22,24 @@ UAbilitySystemComponent* UAbilitySystemBlueprintLibrary::GetAbilitySystemCompone
 
 void UAbilitySystemBlueprintLibrary::ApplyGameplayEffectToTargetData(FGameplayAbilityTargetDataHandle Target, UGameplayEffect *GameplayEffect, const FGameplayAbilityActorInfo InstigatorInfo)
 {
-	if (Target.Data.IsValid() && GameplayEffect != NULL)
+	if (GameplayEffect != NULL)
 	{
-		Target.Data->ApplyGameplayEffect(GameplayEffect, InstigatorInfo);
+		for (int32 i = 0; i < Target.Data.Num(); ++i)
+		{
+			if (Target.Data[i].IsValid())
+			{
+				Target.Data[i]->ApplyGameplayEffect(GameplayEffect, InstigatorInfo);
+			}
+		}
 	}
 }
 
+FGameplayAbilityTargetDataHandle UAbilitySystemBlueprintLibrary::AppendTargetDataHandle(FGameplayAbilityTargetDataHandle TargetHandle, FGameplayAbilityTargetDataHandle HandleToAdd)
+{
+	TargetHandle.Append(&HandleToAdd);
+	return TargetHandle;
+}
+
 FGameplayAbilityTargetDataHandle UAbilitySystemBlueprintLibrary::AbilityTargetDataFromLocations(const FGameplayAbilityTargetingLocationInfo& SourceLocation, const FGameplayAbilityTargetingLocationInfo& TargetLocation)
 {
 	// Construct TargetData
@@ -37,7 +49,7 @@ FGameplayAbilityTargetDataHandle UAbilitySystemBlueprintLibrary::AbilityTargetDa
 
 	// Give it a handle and return
 	FGameplayAbilityTargetDataHandle	Handle;
-	Handle.Data = TSharedPtr<FGameplayAbilityTargetData_LocationInfo>(NewData);
+	Handle.Data.Add(TSharedPtr<FGameplayAbilityTargetData_LocationInfo>(NewData));
 	return Handle;
 }
 
@@ -48,7 +60,7 @@ FGameplayAbilityTargetDataHandle UAbilitySystemBlueprintLibrary::AbilityTargetDa
 	FGameplayAbilityTargetDataHandle	Handle;
 
 	// Give it a handle and return
-	Handle.Data = TSharedPtr<FGameplayAbilityTargetData_Mesh>(NewData);
+	Handle.Data.Add(TSharedPtr<FGameplayAbilityTargetData_Mesh>(NewData));
 	return Handle;
 }
 
@@ -59,14 +71,24 @@ FGameplayAbilityTargetDataHandle UAbilitySystemBlueprintLibrary::AbilityTargetDa
 
 	// Give it a handle and return
 	FGameplayAbilityTargetDataHandle	Handle;
-	Handle.Data = TSharedPtr<FGameplayAbilityTargetData>(TargetData);
+	Handle.Data.Add(TSharedPtr<FGameplayAbilityTargetData>(TargetData));
 
 	return Handle;
 }
 
-TArray<TWeakObjectPtr<AActor>> UAbilitySystemBlueprintLibrary::GetActorsFromTargetData(FGameplayAbilityTargetDataHandle TargetData)
+//TArray<TSharedPtr<FGameplayAbilityTargetData>> UAbilitySystemBlueprintLibrary::GetDataArrayFromTargetData(FGameplayAbilityTargetDataHandle TargetData)
+//{
+//	return TargetData.Data;
+//}
+
+int32 UAbilitySystemBlueprintLibrary::GetDataCountFromTargetData(FGameplayAbilityTargetDataHandle TargetData)
+{
+	return TargetData.Data.Num();
+}
+
+TArray<TWeakObjectPtr<AActor>> UAbilitySystemBlueprintLibrary::GetActorsFromTargetData(FGameplayAbilityTargetDataHandle TargetData, int32 Index)
 {
-	FGameplayAbilityTargetData* Data = TargetData.Data.Get();
+	FGameplayAbilityTargetData* Data = TargetData.Data[Index].Get();
 	if (Data)
 	{
 		return Data->GetActors();
@@ -74,9 +96,9 @@ TArray<TWeakObjectPtr<AActor>> UAbilitySystemBlueprintLibrary::GetActorsFromTarg
 	return TArray<TWeakObjectPtr<AActor>>();
 }
 
-bool UAbilitySystemBlueprintLibrary::TargetDataHasHitResult(FGameplayAbilityTargetDataHandle TargetData)
+bool UAbilitySystemBlueprintLibrary::TargetDataHasHitResult(FGameplayAbilityTargetDataHandle TargetData, int32 Index)
 {
-	FGameplayAbilityTargetData* Data = TargetData.Data.Get();
+	FGameplayAbilityTargetData* Data = TargetData.Data[Index].Get();
 	if (Data)
 	{
 		return Data->HasHitResult();
@@ -84,9 +106,9 @@ bool UAbilitySystemBlueprintLibrary::TargetDataHasHitResult(FGameplayAbilityTarg
 	return false;
 }
 
-FHitResult UAbilitySystemBlueprintLibrary::GetHitResultFromTargetData(FGameplayAbilityTargetDataHandle TargetData)
+FHitResult UAbilitySystemBlueprintLibrary::GetHitResultFromTargetData(FGameplayAbilityTargetDataHandle TargetData, int32 Index)
 {
-	FGameplayAbilityTargetData* Data = TargetData.Data.Get();
+	FGameplayAbilityTargetData* Data = TargetData.Data[Index].Get();
 	if (Data)
 	{
 		const FHitResult* HitResultPtr = Data->GetHitResult();
@@ -99,9 +121,9 @@ FHitResult UAbilitySystemBlueprintLibrary::GetHitResultFromTargetData(FGameplayA
 	return FHitResult();
 }
 
-bool UAbilitySystemBlueprintLibrary::TargetDataHasEndPoint(FGameplayAbilityTargetDataHandle TargetData)
+bool UAbilitySystemBlueprintLibrary::TargetDataHasEndPoint(FGameplayAbilityTargetDataHandle TargetData, int32 Index)
 {
-	FGameplayAbilityTargetData* Data = TargetData.Data.Get();
+	FGameplayAbilityTargetData* Data = TargetData.Data[Index].Get();
 	if (Data)
 	{
 		const FHitResult* HitResultPtr = Data->GetHitResult();
@@ -110,9 +132,9 @@ bool UAbilitySystemBlueprintLibrary::TargetDataHasEndPoint(FGameplayAbilityTarge
 	return false;
 }
 
-FVector UAbilitySystemBlueprintLibrary::GetTargetDataEndPoint(FGameplayAbilityTargetDataHandle TargetData)
+FVector UAbilitySystemBlueprintLibrary::GetTargetDataEndPoint(FGameplayAbilityTargetDataHandle TargetData, int32 Index)
 {
-	FGameplayAbilityTargetData* Data = TargetData.Data.Get();
+	FGameplayAbilityTargetData* Data = TargetData.Data[Index].Get();
 	if (Data)
 	{
 		const FHitResult* HitResultPtr = Data->GetHitResult();
@@ -158,4 +180,4 @@ FHitResult UAbilitySystemBlueprintLibrary::GetHitResult(FGameplayCueParameters P
 bool UAbilitySystemBlueprintLibrary::HasHitResult(FGameplayCueParameters Parameters)
 {
 	return Parameters.InstigatorContext.HitResult.IsValid();
-}
\ No newline at end of file
+}
diff --git a/Engine/Source/Runtime/GameplayAbilities/Private/GameplayAbilityTargetTypes.cpp b/Engine/Source/Runtime/GameplayAbilities/Private/GameplayAbilityTargetTypes.cpp
--- a/Engine/Source/Runtime/GameplayAbilities/Private/GameplayAbilityTargetTypes.cpp
+++ b/Engine/Source/Runtime/GameplayAbilities/Private/GameplayAbilityTargetTypes.cpp
@@ -85,46 +85,71 @@ FGameplayAbilityTargetDataHandle FGameplayAbilityTargetingLocationInfo::MakeTarg
 
 bool FGameplayAbilityTargetDataHandle::NetSerialize(FArchive& Ar, class UPackageMap* Map, bool& bOutSuccess)
 {
+	int32 DataNum;
+	if (Ar.IsSaving())
+	{
+		DataNum = Data.Num();
+	}
+	Ar << DataNum;
+	if (Ar.IsLoading())
+	{
+		Data.SetNumZeroed(DataNum);
+	}
 
-	UScriptStruct* ScriptStruct = Data.IsValid() ? Data->GetScriptStruct() : NULL;
-	Ar << ScriptStruct;
-
-	if (ScriptStruct)
+	for (int32 i = 0; i < DataNum; ++i)
 	{
-		if (Ar.IsLoading())
+		UScriptStruct* ScriptStruct = Data[i].IsValid() ? Data[i]->GetScriptStruct() : NULL;
+		Ar << ScriptStruct;
+
+		if (ScriptStruct)
 		{
-			// For now, just always reset/reallocate the data when loading.
-			// Longer term if we want to generalize this and use it for property replication, we should support
-			// only reallocating when necessary
-			check(!Data.IsValid());
+			if (Ar.IsLoading())
+			{
+				// For now, just always reset/reallocate the data when loading.
+				// Longer term if we want to generalize this and use it for property replication, we should support
+				// only reallocating when necessary
+				check(!Data[i].IsValid());
 
-			FGameplayAbilityTargetData* NewData = (FGameplayAbilityTargetData*)FMemory::Malloc(ScriptStruct->GetCppStructOps()->GetSize());
-			ScriptStruct->InitializeScriptStruct(NewData);
+				FGameplayAbilityTargetData * NewData = (FGameplayAbilityTargetData*)FMemory::Malloc(ScriptStruct->GetCppStructOps()->GetSize());
+				ScriptStruct->InitializeScriptStruct(NewData);
 
-			Data = TSharedPtr<FGameplayAbilityTargetData>(NewData);
-		}
+				Data[i] = TSharedPtr<FGameplayAbilityTargetData>(NewData);
+			}
 
-		void* ContainerPtr = Data.Get();
+			void* ContainerPtr = Data[i].Get();
 
-		if (ScriptStruct->StructFlags & STRUCT_NetSerializeNative)
-		{
-			ScriptStruct->GetCppStructOps()->NetSerialize(Ar, Map, bOutSuccess, Data.Get());
-		}
-		else
-		{
-			// This won't work since UStructProperty::NetSerializeItem is deprecrated.
-			//	1) we have to manually crawl through the topmost struct's fields since we don't have a UStructProperty for it (just the UScriptProperty)
-			//	2) if there are any UStructProperties in the topmost struct's fields, we will assert in UStructProperty::NetSerializeItem.
+			if (ScriptStruct->StructFlags & STRUCT_NetSerializeNative)
+			{
+				ScriptStruct->GetCppStructOps()->NetSerialize(Ar, Map, bOutSuccess, Data[i].Get());
+			}
+			else
+			{
+				// This won't work since UStructProperty::NetSerializeItem is deprecrated.
+				//	1) we have to manually crawl through the topmost struct's fields since we don't have a UStructProperty for it (just the UScriptProperty)
+				//	2) if there are any UStructProperties in the topmost struct's fields, we will assert in UStructProperty::NetSerializeItem.
+
+				ABILITY_LOG(Fatal, TEXT("FGameplayAbilityTargetDataHandle::NetSerialize called on data struct %s without a native NetSerialize"), *ScriptStruct->GetName());
 
-			ABILITY_LOG(Fatal, TEXT("FGameplayAbilityTargetDataHandle::NetSerialize called on data struct %s without a native NetSerialize"), *ScriptStruct->GetName());
+				for (TFieldIterator<UProperty> It(ScriptStruct); It; ++It)
+				{
+					if (It->PropertyFlags & CPF_RepSkip)
+					{
+						continue;
+					}
+
+					void * PropertyData = It->ContainerPtrToValuePtr<void*>(ContainerPtr);
+
+					It->NetSerializeItem(Ar, Map, PropertyData);
+				}
+			}
 		}
 	}
 
 	//ABILITY_LOG(Warning, TEXT("FGameplayAbilityTargetDataHandle Serialized: %s"), ScriptStruct ? *ScriptStruct->GetName() : TEXT("NULL") );
+
 	bOutSuccess = true;
 	return true;
 }
-
 bool FGameplayAbilityTargetingLocationInfo::NetSerialize(FArchive& Ar, class UPackageMap* Map, bool& bOutSuccess)
 {
 	Ar << LocationType;
diff --git a/Engine/Source/Runtime/GameplayAbilities/Public/Abilities/GameplayAbilityTargetTypes.h b/Engine/Source/Runtime/GameplayAbilities/Public/Abilities/GameplayAbilityTargetTypes.h
--- a/Engine/Source/Runtime/GameplayAbilities/Public/Abilities/GameplayAbilityTargetTypes.h
+++ b/Engine/Source/Runtime/GameplayAbilities/Public/Abilities/GameplayAbilityTargetTypes.h
@@ -153,21 +153,33 @@ struct GAMEPLAYABILITIES_API FGameplayAbilityTargetDataHandle
 
 	FGameplayAbilityTargetDataHandle() { }
 	FGameplayAbilityTargetDataHandle(struct FGameplayAbilityTargetData* DataPtr)
-		: Data(DataPtr)
 	{
-
+		Data.Add(TSharedPtr<FGameplayAbilityTargetData>(DataPtr));
 	}
 
-	TSharedPtr<FGameplayAbilityTargetData>	Data;
+	TArray<TSharedPtr<FGameplayAbilityTargetData>>	Data;
 
 	void Clear()
 	{
 		Data.Reset();
 	}
 
-	bool IsValid() const
+	bool IsValid(int32 Index) const
+	{
+		return (Index < Data.Num() && Data[Index].IsValid());
+	}
+
+	FGameplayAbilityTargetData* Get(int32 Index)
+	{
+		return IsValid(Index) ? Data[Index].Get() : NULL;
+	}
+
+	void Append(struct FGameplayAbilityTargetDataHandle* OtherHandle)
 	{
-		return Data.IsValid();
+		for (int32 i = 0; i < OtherHandle->Data.Num(); ++i)
+		{
+			Data.Add(OtherHandle->Data[i]);
+		}
 	}
 
 	bool NetSerialize(FArchive& Ar, class UPackageMap* Map, bool& bOutSuccess);
@@ -176,9 +188,22 @@ struct GAMEPLAYABILITIES_API FGameplayAbilityTargetDataHandle
 	bool operator==(FGameplayAbilityTargetDataHandle const& Other) const
 	{
 		// Both invalid structs or both valid and Pointer compare (???) // deep comparison equality
-		bool bBothValid = IsValid() && Other.IsValid();
-		bool bBothInvalid = !IsValid() && !Other.IsValid();
-		return (bBothInvalid || (bBothValid && (Data.Get() == Other.Data.Get())));
+		if (Data.Num() != Other.Data.Num())
+		{
+			return false;
+		}
+		for (int32 i = 0; i < Data.Num(); ++i)
+		{
+			if (Data[i].IsValid() != Other.Data[i].IsValid())
+			{
+				return false;
+			}
+			if (Data[i].Get() != Other.Data[i].Get())
+			{
+				return false;
+			}
+		}
+		return true;
 	}
 
 	/** Comparison operator */
@@ -346,7 +371,7 @@ struct GAMEPLAYABILITIES_API FGameplayAbilityTargetData_LocationInfo : public FG
 
 	/** Generic location data for target */
 	UPROPERTY(BlueprintReadWrite, EditAnywhere, Category = Targeting)
-		FGameplayAbilityTargetingLocationInfo TargetLocation;
+	FGameplayAbilityTargetingLocationInfo TargetLocation;
 
 	// -------------------------------------
 
@@ -590,7 +615,7 @@ struct TStructOpsTypeTraits<FGameplayAbilityTargetData_SingleTargetHit> : public
 {
 	enum
 	{
-		WithNetSerializer = true	// Fow now this is REQUIRED for FGameplayAbilityTargetDataHandle net serialization to work
+		WithNetSerializer = true	// For now this is REQUIRED for FGameplayAbilityTargetDataHandle net serialization to work
 	};
 };
 
diff --git a/Engine/Source/Runtime/GameplayAbilities/Public/AbilitySystemBlueprintLibrary.h b/Engine/Source/Runtime/GameplayAbilities/Public/AbilitySystemBlueprintLibrary.h
--- a/Engine/Source/Runtime/GameplayAbilities/Public/AbilitySystemBlueprintLibrary.h
+++ b/Engine/Source/Runtime/GameplayAbilities/Public/AbilitySystemBlueprintLibrary.h
@@ -28,6 +28,9 @@ class UAbilitySystemBlueprintLibrary : public UBlueprintFunctionLibrary
 	UFUNCTION(BlueprintCallable, Category = "Ability|TargetData")
 	static void ApplyGameplayEffectToTargetData(FGameplayAbilityTargetDataHandle Target, UGameplayEffect *GameplayEffect, const FGameplayAbilityActorInfo InstigatorInfo);
 
+	UFUNCTION(BlueprintCallable, Category = "Ability|TargetData")
+	static FGameplayAbilityTargetDataHandle AppendTargetDataHandle(FGameplayAbilityTargetDataHandle TargetHandle, FGameplayAbilityTargetDataHandle HandleToAdd);
+
 	UFUNCTION(BlueprintPure, Category = "Ability|TargetData")
 	static FGameplayAbilityTargetDataHandle	AbilityTargetDataFromLocations(const FGameplayAbilityTargetingLocationInfo& SourceLocation, const FGameplayAbilityTargetingLocationInfo& TargetLocation);
 
@@ -37,20 +40,27 @@ class UAbilitySystemBlueprintLibrary : public UBlueprintFunctionLibrary
 	UFUNCTION(BlueprintPure, Category = "Ability|TargetData")
 	static FGameplayAbilityTargetDataHandle	AbilityTargetDataFromHitResult(FHitResult HitResult);
 
+	//TODO Consider adding this functionality, or delete it soon.
+//	UFUNCTION(BlueprintPure, Category = "Ability|TargetData")
+//	static TArray<TSharedPtr<FGameplayAbilityTargetData>> GetDataArrayFromTargetData(FGameplayAbilityTargetDataHandle TargetData);
+
+	UFUNCTION(BlueprintPure, Category = "Ability|TargetData")
+	static int32 GetDataCountFromTargetData(FGameplayAbilityTargetDataHandle TargetData);
+
 	UFUNCTION(BlueprintPure, Category = "Ability|TargetData")
-	static TArray<TWeakObjectPtr<AActor>> GetActorsFromTargetData(FGameplayAbilityTargetDataHandle TargetData);
+	static TArray<TWeakObjectPtr<AActor>> GetActorsFromTargetData(FGameplayAbilityTargetDataHandle TargetData, int32 Index);
 
 	UFUNCTION(BlueprintPure, Category = "Ability|TargetData")
-	static bool TargetDataHasHitResult(FGameplayAbilityTargetDataHandle HitResult);
+	static bool TargetDataHasHitResult(FGameplayAbilityTargetDataHandle HitResult, int32 Index);
 
 	UFUNCTION(BlueprintPure, Category = "Ability|TargetData")
-	static FHitResult GetHitResultFromTargetData(FGameplayAbilityTargetDataHandle HitResult);
+	static FHitResult GetHitResultFromTargetData(FGameplayAbilityTargetDataHandle HitResult, int32 Index);
 
 	UFUNCTION(BlueprintPure, Category = "Ability|TargetData")
-	static bool TargetDataHasEndPoint(FGameplayAbilityTargetDataHandle TargetData);
+	static bool TargetDataHasEndPoint(FGameplayAbilityTargetDataHandle TargetData, int32 Index);
 
 	UFUNCTION(BlueprintPure, Category = "Ability|TargetData")
-	static FVector GetTargetDataEndPoint(FGameplayAbilityTargetDataHandle TargetData);
+	static FVector GetTargetDataEndPoint(FGameplayAbilityTargetDataHandle TargetData, int32 Index);
 
 	// -------------------------------------------------------------------------------
 	//		GameplayCue
@@ -68,4 +78,4 @@ class UAbilitySystemBlueprintLibrary : public UBlueprintFunctionLibrary
 
 	// -------------------------------------------------------------------------------
 
-};
\ No newline at end of file
+};
