#Orion

- Ability Widgets work! Woo!
- Turning on UMG HUD

diff --git a/Engine/Source/Runtime/GameplayAbilities/Private/AbilitySystemComponent_Abilities.cpp b/Engine/Source/Runtime/GameplayAbilities/Private/AbilitySystemComponent_Abilities.cpp
--- a/Engine/Source/Runtime/GameplayAbilities/Private/AbilitySystemComponent_Abilities.cpp
+++ b/Engine/Source/Runtime/GameplayAbilities/Private/AbilitySystemComponent_Abilities.cpp
@@ -174,6 +174,21 @@ FGameplayAbilitySpec* UAbilitySystemComponent::FindAbilitySpecFromHandle(FGamepl
 	return nullptr;
 }
 
+FGameplayAbilitySpec* UAbilitySystemComponent::FindAbilitySpecFromInputID(int32 InputID)
+{
+	if (InputID != INDEX_NONE)
+	{
+		for (FGameplayAbilitySpec& Spec : ActivatableAbilities)
+		{
+			if (Spec.InputID == InputID)
+			{
+				return &Spec;
+			}
+		}
+	}
+	return nullptr;
+}
+
 UGameplayAbility* UAbilitySystemComponent::CreateNewInstanceOfAbility(FGameplayAbilitySpec& Spec, UGameplayAbility* Ability)
 {
 	check(Ability);
diff --git a/Engine/Source/Runtime/GameplayAbilities/Public/AbilitySystemComponent.h b/Engine/Source/Runtime/GameplayAbilities/Public/AbilitySystemComponent.h
--- a/Engine/Source/Runtime/GameplayAbilities/Public/AbilitySystemComponent.h
+++ b/Engine/Source/Runtime/GameplayAbilities/Public/AbilitySystemComponent.h
@@ -412,6 +412,8 @@ class GAMEPLAYABILITIES_API UAbilitySystemComponent : public UActorComponent, pu
 	TArray<FGameplayAbilitySpec>	ActivatableAbilities;
 
 	FGameplayAbilitySpec* FindAbilitySpecFromHandle(FGameplayAbilitySpecHandle Handle);
+
+	FGameplayAbilitySpec* FindAbilitySpecFromInputID(int32 InputID);
 	
 	UFUNCTION()
 	void	OnRep_ActivateAbilities();
