#UE4

- Guarding against NULL

diff --git a/Engine/Source/Runtime/GameplayAbilities/Private/Abilities/GameplayAbility_CharacterJump.cpp b/Engine/Source/Runtime/GameplayAbilities/Private/Abilities/GameplayAbility_CharacterJump.cpp
--- a/Engine/Source/Runtime/GameplayAbilities/Private/Abilities/GameplayAbility_CharacterJump.cpp
+++ b/Engine/Source/Runtime/GameplayAbilities/Private/Abilities/GameplayAbility_CharacterJump.cpp
@@ -46,8 +46,8 @@ bool UGameplayAbility_CharacterJump::CanActivateAbility(const FGameplayAbilitySp
 		return false;
 	}
 
-	const UCharacterMovementComponent *CharMovement = CastChecked<UCharacterMovementComponent>(ActorInfo->MovementComponent.Get());
-	return (CharMovement->MovementMode == EMovementMode::MOVE_Walking);
+	const UCharacterMovementComponent *CharMovement = CastChecked<UCharacterMovementComponent>(ActorInfo->MovementComponent.Get(), ECastCheckedType::NullAllowed);
+	return (CharMovement && CharMovement->MovementMode == EMovementMode::MOVE_Walking);
 }
 
 /**
