FActiveGameplayEffectHandle is now globally resolvable. Fixed WaitGameplayEffectremoved task to work on externally applied gameplay effects. Bunch of small blueprint fixes for designers.

diff --git a/Engine/Source/Runtime/GameplayAbilities/Private/Abilities/Tasks/AbilityTask_WaitGameplayEffectRemoved.cpp b/Engine/Source/Runtime/GameplayAbilities/Private/Abilities/Tasks/AbilityTask_WaitGameplayEffectRemoved.cpp
--- a/Engine/Source/Runtime/GameplayAbilities/Private/Abilities/Tasks/AbilityTask_WaitGameplayEffectRemoved.cpp
+++ b/Engine/Source/Runtime/GameplayAbilities/Private/Abilities/Tasks/AbilityTask_WaitGameplayEffectRemoved.cpp
@@ -9,6 +9,7 @@
 UAbilityTask_WaitGameplayEffectRemoved::UAbilityTask_WaitGameplayEffectRemoved(const class FPostConstructInitializeProperties& PCIP)
 	: Super(PCIP)
 {
+	Registered = false;
 }
 
 UAbilityTask_WaitGameplayEffectRemoved* UAbilityTask_WaitGameplayEffectRemoved::WaitForGameplayEffectRemoved(UObject* WorldContextObject, FActiveGameplayEffectHandle InHandle)
@@ -21,27 +22,35 @@ UAbilityTask_WaitGameplayEffectRemoved* UAbilityTask_WaitGameplayEffectRemoved::
 
 void UAbilityTask_WaitGameplayEffectRemoved::Activate()
 {
-	if (AbilitySystemComponent.IsValid())	
-	{		
-		FOnActiveGameplayEffectRemoved* DelPtr = AbilitySystemComponent->OnGameplayEffectRemovedDelegate(Handle);
+	UAbilitySystemComponent* EffectOwningAbilitySystemComponent = Handle.GetOwningAbilitySystemComponent();
+
+	if (EffectOwningAbilitySystemComponent)
+	{
+		FOnActiveGameplayEffectRemoved* DelPtr = EffectOwningAbilitySystemComponent->OnGameplayEffectRemovedDelegate(Handle);
 		if (DelPtr)
 		{
 			DelPtr->AddUObject(this, &UAbilityTask_WaitGameplayEffectRemoved::OnGameplayEffectRemoved);
+			Registered = true;
 		}
-		else
-		{
-			// GameplayEffect was already removed, treat this as a warning? Could be cases of immunity or chained gameplay rules that would instant remove something
-			OnGameplayEffectRemoved();
-		}
+	}
+
+	if (!Registered)
+	{
+		// GameplayEffect was already removed, treat this as a warning? Could be cases of immunity or chained gameplay rules that would instant remove something
+		OnGameplayEffectRemoved();
 	}
 }
 
 void UAbilityTask_WaitGameplayEffectRemoved::OnDestroy(bool AbilityIsEnding)
 {
-	FOnActiveGameplayEffectRemoved* DelPtr = AbilitySystemComponent->OnGameplayEffectRemovedDelegate(Handle);
-	if (DelPtr)
+	UAbilitySystemComponent* EffectOwningAbilitySystemComponent = Handle.GetOwningAbilitySystemComponent();
+	if (EffectOwningAbilitySystemComponent)
 	{
-		DelPtr->RemoveUObject(this, &UAbilityTask_WaitGameplayEffectRemoved::OnGameplayEffectRemoved);
+		FOnActiveGameplayEffectRemoved* DelPtr = EffectOwningAbilitySystemComponent->OnGameplayEffectRemovedDelegate(Handle);
+		if (DelPtr)
+		{
+			DelPtr->RemoveUObject(this, &UAbilityTask_WaitGameplayEffectRemoved::OnGameplayEffectRemoved);
+		}
 	}
 
 	Super::OnDestroy(AbilityIsEnding);
diff --git a/Engine/Source/Runtime/GameplayAbilities/Private/AbilitySystemBlueprintLibrary.cpp b/Engine/Source/Runtime/GameplayAbilities/Private/AbilitySystemBlueprintLibrary.cpp
--- a/Engine/Source/Runtime/GameplayAbilities/Private/AbilitySystemBlueprintLibrary.cpp
+++ b/Engine/Source/Runtime/GameplayAbilities/Private/AbilitySystemBlueprintLibrary.cpp
@@ -20,18 +20,57 @@ UAbilitySystemComponent* UAbilitySystemBlueprintLibrary::GetAbilitySystemCompone
 	return UAbilitySystemGlobals::GetAbilitySystemComponentFromActor(Actor);
 }
 
-void UAbilitySystemBlueprintLibrary::ApplyGameplayEffectToTargetData(FGameplayAbilityTargetDataHandle Target, UGameplayEffect *GameplayEffect, const FGameplayAbilityActorInfo InstigatorInfo)
+TArray<FActiveGameplayEffectHandle> UAbilitySystemBlueprintLibrary::ApplyGameplayEffectToTargetData(FGameplayAbilityTargetDataHandle Target, UGameplayEffect* GameplayEffect, const FGameplayAbilityActorInfo InstigatorInfo, int32 Level)
 {
-	if (GameplayEffect != NULL)
+	TArray<FActiveGameplayEffectHandle> Handles;
+
+	if (GameplayEffect == nullptr)
+	{
+		return Handles;
+	}
+
+	for (auto Data : Target.Data)
+	{
+		if (Data.IsValid())
+		{
+			Handles.Append(Data->ApplyGameplayEffect(GameplayEffect, InstigatorInfo, (float)Level));
+		}
+	}
+
+	return Handles;
+}
+
+FActiveGameplayEffectHandle UAbilitySystemBlueprintLibrary::ApplyGameplayEffectToTargetData_Single(FGameplayAbilityTargetDataHandle Target, UGameplayEffect* GameplayEffect, const FGameplayAbilityActorInfo InstigatorInfo, int32 Level)
+{
+	FActiveGameplayEffectHandle Handle;
+
+	if (GameplayEffect == nullptr)
+	{
+		return Handle;
+	}
+
+	if (Target.Data.Num() != 1)
 	{
-		for (int32 i = 0; i < Target.Data.Num(); ++i)
+		ABILITY_LOG(Warning, TEXT("ApplyGameplayEffectToTargetData_Single called on invalid TargetDataHandle."));
+		return Handle;
+	}
+
+	
+	if (Target.Data[0].IsValid())
+	{
+		TArray<FActiveGameplayEffectHandle> Handles = Target.Data[0]->ApplyGameplayEffect(GameplayEffect, InstigatorInfo, (float)Level);
+		int32 Num = Handles.Num();
+		if (Num > 0)
 		{
-			if (Target.Data[i].IsValid())
+			Handle = Handles[0];
+			if (Num > 1)
 			{
-				Target.Data[i]->ApplyGameplayEffect(GameplayEffect, InstigatorInfo);
+				ABILITY_LOG(Warning, TEXT("ApplyGameplayEffectToTargetData_Single called on TargetData with multiple actor targets"));
 			}
 		}
 	}
+
+	return Handle;
 }
 
 FGameplayAbilityTargetDataHandle UAbilitySystemBlueprintLibrary::AppendTargetDataHandle(FGameplayAbilityTargetDataHandle TargetHandle, FGameplayAbilityTargetDataHandle HandleToAdd)
@@ -64,6 +103,17 @@ FGameplayAbilityTargetDataHandle UAbilitySystemBlueprintLibrary::AbilityTargetDa
 	return Handle;
 }
 
+
+FGameplayAbilityTargetDataHandle UAbilitySystemBlueprintLibrary::AbilityTargetDataFromActor(AActor* Actor)
+{
+	// Construct TargetData
+	FGameplayAbilityTargetData_ActorArray*	NewData = new FGameplayAbilityTargetData_ActorArray();
+	NewData->TargetActorArray.Add(Actor);
+
+	FGameplayAbilityTargetDataHandle		Handle(NewData);
+	return Handle;
+}
+
 FGameplayAbilityTargetDataHandle UAbilitySystemBlueprintLibrary::AbilityTargetDataFromHitResult(FHitResult HitResult)
 {
 	// Construct TargetData
@@ -76,24 +126,24 @@ FGameplayAbilityTargetDataHandle UAbilitySystemBlueprintLibrary::AbilityTargetDa
 	return Handle;
 }
 
-//TArray<TSharedPtr<FGameplayAbilityTargetData>> UAbilitySystemBlueprintLibrary::GetDataArrayFromTargetData(FGameplayAbilityTargetDataHandle TargetData)
-//{
-//	return TargetData.Data;
-//}
-
 int32 UAbilitySystemBlueprintLibrary::GetDataCountFromTargetData(FGameplayAbilityTargetDataHandle TargetData)
 {
 	return TargetData.Data.Num();
 }
 
-TArray<TWeakObjectPtr<AActor>> UAbilitySystemBlueprintLibrary::GetActorsFromTargetData(FGameplayAbilityTargetDataHandle TargetData, int32 Index)
+TArray<AActor*> UAbilitySystemBlueprintLibrary::GetActorsFromTargetData(FGameplayAbilityTargetDataHandle TargetData, int32 Index)
 {
 	FGameplayAbilityTargetData* Data = TargetData.Data[Index].Get();
+	TArray<AActor*>	ResolvedArray;
 	if (Data)
 	{
-		return Data->GetActors();
+		TArray<TWeakObjectPtr<AActor>> WeakArray = Data->GetActors();
+		for (TWeakObjectPtr<AActor> WeakPtr : WeakArray)
+		{
+			ResolvedArray.Add(WeakPtr.Get());
+		}
 	}
-	return TArray<TWeakObjectPtr<AActor>>();
+	return ResolvedArray;
 }
 
 bool UAbilitySystemBlueprintLibrary::TargetDataHasHitResult(FGameplayAbilityTargetDataHandle TargetData, int32 Index)
diff --git a/Engine/Source/Runtime/GameplayAbilities/Private/GameplayAbilityTargetTypes.cpp b/Engine/Source/Runtime/GameplayAbilities/Private/GameplayAbilityTargetTypes.cpp
--- a/Engine/Source/Runtime/GameplayAbilities/Private/GameplayAbilityTargetTypes.cpp
+++ b/Engine/Source/Runtime/GameplayAbilities/Private/GameplayAbilityTargetTypes.cpp
@@ -7,14 +7,14 @@
 #include "AbilitySystemBlueprintLibrary.h"
 
 
-void FGameplayAbilityTargetData::ApplyGameplayEffect(UGameplayEffect* GameplayEffect, const FGameplayAbilityActorInfo InstigatorInfo)
+TArray<FActiveGameplayEffectHandle> FGameplayAbilityTargetData::ApplyGameplayEffect(UGameplayEffect* GameplayEffect, const FGameplayAbilityActorInfo InstigatorInfo, float Level)
 {
 	// TODO: Improve relationship between InstigatorContext and FGameplayAbilityTargetData/FHitResult (or use something different between HitResult)
 
 
 	FGameplayEffectSpec	SpecToApply(GameplayEffect,					// The UGameplayEffect data asset
 		InstigatorInfo.Actor.Get(),		// The actor who instigated this
-		1.f,							// FIXME: Leveling
+		Level,							// FIXME: Leveling
 		NULL							// FIXME: CurveData override... should we just remove this?
 		);
 	if (HasHitResult())
@@ -27,20 +27,24 @@ void FGameplayAbilityTargetData::ApplyGameplayEffect(UGameplayEffect* GameplayEf
 		SpecToApply.InstigatorContext.AddOrigin(GetOrigin().GetLocation());
 	}
 
-	TArray<TWeakObjectPtr<AActor>> Actors = GetActors();
+	TArray<TWeakObjectPtr<AActor> > Actors = GetActors();
+	TArray<FActiveGameplayEffectHandle>	AppliedHandles;
+	AppliedHandles.Reserve(Actors.Num());
+
 	for (TWeakObjectPtr<AActor> TargetActor : Actors)
 	{
 		if (TargetActor.IsValid())
 		{
-			UAbilitySystemComponent * TargetComponent = UAbilitySystemBlueprintLibrary::GetAbilitySystemComponent(TargetActor.Get());
+			UAbilitySystemComponent* TargetComponent = UAbilitySystemBlueprintLibrary::GetAbilitySystemComponent(TargetActor.Get());
 			if (TargetComponent)
 			{
-				InstigatorInfo.AbilitySystemComponent->ApplyGameplayEffectSpecToTarget(SpecToApply, TargetComponent);
+				AppliedHandles.Add( InstigatorInfo.AbilitySystemComponent->ApplyGameplayEffectSpecToTarget(SpecToApply, TargetComponent) );
 			}
 		}
 	}
-}
 
+	return AppliedHandles;
+}
 
 FString FGameplayAbilityTargetData::ToString() const
 {
diff --git a/Engine/Source/Runtime/GameplayAbilities/Private/GameplayEffect.cpp b/Engine/Source/Runtime/GameplayAbilities/Private/GameplayEffect.cpp
--- a/Engine/Source/Runtime/GameplayAbilities/Private/GameplayEffect.cpp
+++ b/Engine/Source/Runtime/GameplayAbilities/Private/GameplayEffect.cpp
@@ -1573,8 +1573,8 @@ FActiveGameplayEffect & FActiveGameplayEffectsContainer::CreateNewActiveGameplay
 {
 	SCOPE_CYCLE_COUNTER(STAT_CreateNewActiveGameplayEffect);
 
-	LastAssignedHandle = LastAssignedHandle.GetNextHandle();
-	FActiveGameplayEffect & NewEffect = *new (GameplayEffects)FActiveGameplayEffect(LastAssignedHandle, Spec, GetWorldTime(), GetGameStateTime(), InPrevPredictionKey, InCurrPredictionKey);
+	FActiveGameplayEffectHandle NewHandle = FActiveGameplayEffectHandle::GenerateNewHandle(Owner);
+	FActiveGameplayEffect & NewEffect = *new (GameplayEffects)FActiveGameplayEffect(NewHandle, Spec, GetWorldTime(), GetGameStateTime(), InPrevPredictionKey, InCurrPredictionKey);
 
 	// register callbacks with the timer manager
 	if (Owner)
@@ -1582,14 +1582,14 @@ FActiveGameplayEffect & FActiveGameplayEffectsContainer::CreateNewActiveGameplay
 		if (Spec.GetDuration() > 0.f)
 		{
 			FTimerManager& TimerManager = Owner->GetWorld()->GetTimerManager();
-			FTimerDelegate Delegate = FTimerDelegate::CreateUObject(Owner, &UAbilitySystemComponent::CheckDurationExpired, LastAssignedHandle);
+			FTimerDelegate Delegate = FTimerDelegate::CreateUObject(Owner, &UAbilitySystemComponent::CheckDurationExpired, NewHandle);
 			TimerManager.SetTimer(NewEffect.DurationHandle, Delegate, Spec.GetDuration(), false, Spec.GetDuration());
 		}
 		// The timer manager moves things from the pending list to the active list after checking the active list on the first tick so we need to execute here
 		if (Spec.GetPeriod() != UGameplayEffect::NO_PERIOD)
 		{
 			FTimerManager& TimerManager = Owner->GetWorld()->GetTimerManager();
-			FTimerDelegate Delegate = FTimerDelegate::CreateUObject(Owner, &UAbilitySystemComponent::ExecutePeriodicEffect, LastAssignedHandle);
+			FTimerDelegate Delegate = FTimerDelegate::CreateUObject(Owner, &UAbilitySystemComponent::ExecutePeriodicEffect, NewHandle);
 
 			// If this is a periodic stacking effect, make sure that it's in sync with the others.
 			float FirstDelay = -1.f;
@@ -1615,7 +1615,7 @@ FActiveGameplayEffect & FActiveGameplayEffectsContainer::CreateNewActiveGameplay
 			TimerManager.SetTimer(NewEffect.PeriodHandle, Delegate, Spec.GetPeriod(), true, FirstDelay); // this is going to be off by a frame for stacking because of the pending list
 		}
 
-		NewEffect.Spec.Duration.Get()->OnDirty = FAggregator::FOnDirty::CreateRaw(this, &FActiveGameplayEffectsContainer::OnDurationAggregatorDirty, Owner, LastAssignedHandle);
+		NewEffect.Spec.Duration.Get()->OnDirty = FAggregator::FOnDirty::CreateRaw(this, &FActiveGameplayEffectsContainer::OnDurationAggregatorDirty, Owner, NewHandle);
 	}
 	
 	if (InCurrPredictionKey == 0 || IsNetAuthority())	// Clients predicting a GameplayEffect must not call MarkItemDirty
@@ -2149,3 +2149,30 @@ FString EGameplayEffectStackingPolicyToString(int32 Type)
 	return e->GetEnum(Type).ToString();
 }
 
+namespace GlobalActiveGameplayEffectHandles
+{
+	static TMap<FActiveGameplayEffectHandle, TWeakObjectPtr<UAbilitySystemComponent>>	Map;
+}
+
+FActiveGameplayEffectHandle FActiveGameplayEffectHandle::GenerateNewHandle(UAbilitySystemComponent* OwningComponent)
+{
+	static int32 GHandleID=0;
+	FActiveGameplayEffectHandle NewHandle(GHandleID++);
+
+	TWeakObjectPtr<UAbilitySystemComponent> WeakPtr(OwningComponent);
+
+	GlobalActiveGameplayEffectHandles::Map.Add(NewHandle, WeakPtr);
+
+	return NewHandle;
+}
+
+UAbilitySystemComponent* FActiveGameplayEffectHandle::GetOwningAbilitySystemComponent()
+{
+	TWeakObjectPtr<UAbilitySystemComponent>* Ptr = GlobalActiveGameplayEffectHandles::Map.Find(*this);
+	if (Ptr)
+	{
+		return Ptr->Get();
+	}
+
+	return nullptr;	
+}
\ No newline at end of file
diff --git a/Engine/Source/Runtime/GameplayAbilities/Public/Abilities/GameplayAbilityTargetTypes.h b/Engine/Source/Runtime/GameplayAbilities/Public/Abilities/GameplayAbilityTargetTypes.h
--- a/Engine/Source/Runtime/GameplayAbilities/Public/Abilities/GameplayAbilityTargetTypes.h
+++ b/Engine/Source/Runtime/GameplayAbilities/Public/Abilities/GameplayAbilityTargetTypes.h
@@ -2,6 +2,7 @@
 
 #pragma once
 
+#include "GameplayEffectTypes.h"
 #include "GameplayAbilityTargetTypes.generated.h"
 
 class UGameplayEffect;
@@ -26,34 +27,34 @@ namespace EGameplayTargetingConfirmation
 
 
 /**
-*	A generic structure for targeting data. We want generic functions to produce this data and other generic
-*	functions to consume this data.
-*
-*	We expect this to be able to hold specific actors/object reference and also generic location/direction/origin
-*	information.
-*
-*	Some example producers:
-*		-Overlap/Hit collision event generates TargetData about who was hit in a melee attack
-*		-A mouse input causes a hit trace and the actor infront of the crosshair is turned into TargetData
-*		-A mouse input causes TargetData to be generated from the owner's crosshair view origin/direction
-*		-An AOE/aura pulses and all actors in a radius around the instigator are added to TargetData
-*		-Panzer Dragoon style 'painting' targeting mode
-*		-MMORPG style ground AOE targeting style (potentially both a location on the ground and actors that were targeted)
-*
-*	Some example consumers:
-*		-Apply a GameplayEffect to all actors in TargetData
-*		-Find closest actor from all in TargetData
-*		-Call some function on all actors in TargetData
-*		-Filter or merge TargetDatas
-*		-Spawn a new actor at a TargetData location
-*
-*
-*
-*	Maybe it is better to distinguish between actor list targeting vs positional targeting data?
-*		-AOE/aura type of targeting data blurs the line
-*
-*
-*/
+ *	A generic structure for targeting data. We want generic functions to produce this data and other generic
+ *	functions to consume this data.
+ *
+ *	We expect this to be able to hold specific actors/object reference and also generic location/direction/origin
+ *	information.
+ *
+ *	Some example producers:
+ *		-Overlap/Hit collision event generates TargetData about who was hit in a melee attack
+ *		-A mouse input causes a hit trace and the actor infront of the crosshair is turned into TargetData
+ *		-A mouse input causes TargetData to be generated from the owner's crosshair view origin/direction
+ *		-An AOE/aura pulses and all actors in a radius around the instigator are added to TargetData
+ *		-Panzer Dragoon style 'painting' targeting mode
+ *		-MMORPG style ground AOE targeting style (potentially both a location on the ground and actors that were targeted)
+ *
+ *	Some example consumers:
+ *		-Apply a GameplayEffect to all actors in TargetData
+ *		-Find closest actor from all in TargetData
+ *		-Call some function on all actors in TargetData
+ * 		-Filter or merge TargetDatas
+ *		-Spawn a new actor at a TargetData location
+ *
+ *
+ *
+ *	Maybe it is better to distinguish between actor list targeting vs positional targeting data?
+ *		-AOE/aura type of targeting data blurs the line
+ *
+ *
+ */
 
 USTRUCT()
 struct GAMEPLAYABILITIES_API FGameplayAbilityTargetData
@@ -62,9 +63,9 @@ struct GAMEPLAYABILITIES_API FGameplayAbilityTargetData
 
 	virtual ~FGameplayAbilityTargetData() { }
 
-	void ApplyGameplayEffect(UGameplayEffect *GameplayEffect, const struct FGameplayAbilityActorInfo InstigatorInfo);
+	TArray<FActiveGameplayEffectHandle> ApplyGameplayEffect(UGameplayEffect *GameplayEffect, const struct FGameplayAbilityActorInfo InstigatorInfo, float Level);
 
-	virtual TArray<TWeakObjectPtr<AActor>>	GetActors() const
+	virtual TArray<TWeakObjectPtr<AActor> >	GetActors() const
 	{
 		return TArray<TWeakObjectPtr<AActor>>();
 	}
diff --git a/Engine/Source/Runtime/GameplayAbilities/Public/Abilities/Tasks/AbilityTask_WaitGameplayEffectRemoved.h b/Engine/Source/Runtime/GameplayAbilities/Public/Abilities/Tasks/AbilityTask_WaitGameplayEffectRemoved.h
--- a/Engine/Source/Runtime/GameplayAbilities/Public/Abilities/Tasks/AbilityTask_WaitGameplayEffectRemoved.h
+++ b/Engine/Source/Runtime/GameplayAbilities/Public/Abilities/Tasks/AbilityTask_WaitGameplayEffectRemoved.h
@@ -36,4 +36,5 @@ class UAbilityTask_WaitGameplayEffectRemoved : public UAbilityTask
 protected:
 
 	virtual void OnDestroy(bool AbilityIsEnding) override;
+	bool Registered;
 };
\ No newline at end of file
diff --git a/Engine/Source/Runtime/GameplayAbilities/Public/AbilitySystemBlueprintLibrary.h b/Engine/Source/Runtime/GameplayAbilities/Public/AbilitySystemBlueprintLibrary.h
--- a/Engine/Source/Runtime/GameplayAbilities/Public/AbilitySystemBlueprintLibrary.h
+++ b/Engine/Source/Runtime/GameplayAbilities/Public/AbilitySystemBlueprintLibrary.h
@@ -26,7 +26,11 @@ class UAbilitySystemBlueprintLibrary : public UBlueprintFunctionLibrary
 	// -------------------------------------------------------------------------------
 
 	UFUNCTION(BlueprintCallable, Category = "Ability|TargetData")
-	static void ApplyGameplayEffectToTargetData(FGameplayAbilityTargetDataHandle Target, UGameplayEffect *GameplayEffect, const FGameplayAbilityActorInfo InstigatorInfo);
+	static TArray<FActiveGameplayEffectHandle> ApplyGameplayEffectToTargetData(FGameplayAbilityTargetDataHandle Target, UGameplayEffect *GameplayEffect, const FGameplayAbilityActorInfo InstigatorInfo, int32 Level=1);
+
+	/** Applies GameplayEffect to TargetData and returns the handle of the  */
+	UFUNCTION(BlueprintCallable, Category = "Ability|TargetData")
+	static FActiveGameplayEffectHandle ApplyGameplayEffectToTargetData_Single(FGameplayAbilityTargetDataHandle Target, UGameplayEffect *GameplayEffect, const FGameplayAbilityActorInfo InstigatorInfo, int32 Level=1);
 
 	UFUNCTION(BlueprintCallable, Category = "Ability|TargetData")
 	static FGameplayAbilityTargetDataHandle AppendTargetDataHandle(FGameplayAbilityTargetDataHandle TargetHandle, FGameplayAbilityTargetDataHandle HandleToAdd);
@@ -40,15 +44,14 @@ class UAbilitySystemBlueprintLibrary : public UBlueprintFunctionLibrary
 	UFUNCTION(BlueprintPure, Category = "Ability|TargetData")
 	static FGameplayAbilityTargetDataHandle	AbilityTargetDataFromHitResult(FHitResult HitResult);
 
-	//TODO Consider adding this functionality, or delete it soon.
-//	UFUNCTION(BlueprintPure, Category = "Ability|TargetData")
-//	static TArray<TSharedPtr<FGameplayAbilityTargetData>> GetDataArrayFromTargetData(FGameplayAbilityTargetDataHandle TargetData);
-
 	UFUNCTION(BlueprintPure, Category = "Ability|TargetData")
 	static int32 GetDataCountFromTargetData(FGameplayAbilityTargetDataHandle TargetData);
 
 	UFUNCTION(BlueprintPure, Category = "Ability|TargetData")
-	static TArray<TWeakObjectPtr<AActor>> GetActorsFromTargetData(FGameplayAbilityTargetDataHandle TargetData, int32 Index);
+	static FGameplayAbilityTargetDataHandle	AbilityTargetDataFromActor(AActor* Actor);
+
+	UFUNCTION(BlueprintPure, Category = "Ability|TargetData")
+	static TArray<AActor*> GetActorsFromTargetData(FGameplayAbilityTargetDataHandle TargetData, int32 Index);
 
 	UFUNCTION(BlueprintPure, Category = "Ability|TargetData")
 	static bool TargetDataHasHitResult(FGameplayAbilityTargetDataHandle HitResult, int32 Index);
diff --git a/Engine/Source/Runtime/GameplayAbilities/Public/GameplayEffect.h b/Engine/Source/Runtime/GameplayAbilities/Public/GameplayEffect.h
--- a/Engine/Source/Runtime/GameplayAbilities/Public/GameplayEffect.h
+++ b/Engine/Source/Runtime/GameplayAbilities/Public/GameplayEffect.h
@@ -373,59 +373,6 @@ protected:
 
 };
 
-/**
-* This handle is required for things outside of FActiveGameplayEffectsContainer to refer to a specific active GameplayEffect
-*	For example if a skill needs to create an active effect and then destroy that specific effect that it created, it has to do so
-*	through a handle. a pointer or index into the active list is not sufficient.
-*/
-USTRUCT(BlueprintType)
-struct FActiveGameplayEffectHandle
-{
-	GENERATED_USTRUCT_BODY()
-
-	FActiveGameplayEffectHandle()
-		: Handle(INDEX_NONE)
-	{
-
-	}
-
-	FActiveGameplayEffectHandle(int32 InHandle)
-		: Handle(InHandle)
-	{
-
-	}
-
-	bool IsValid() const
-	{
-		return Handle != INDEX_NONE;
-	}
-
-	FActiveGameplayEffectHandle GetNextHandle()
-	{
-		return FActiveGameplayEffectHandle(Handle + 1);
-	}
-
-	bool operator==(const FActiveGameplayEffectHandle& Other) const
-	{
-		return Handle == Other.Handle;
-	}
-
-	bool operator!=(const FActiveGameplayEffectHandle& Other) const
-	{
-		return Handle != Other.Handle;
-	}
-
-	FString ToString() const
-	{
-		return FString::Printf(TEXT("%d"), Handle);
-	}
-
-private:
-
-	UPROPERTY()
-	int32 Handle;
-};
-
 /**
  * FGameplayEffectLevelSpec
  *	Level Specification. This can be a static, constant level specified on creation or it can be dynamically tied to a source's attribute value.
@@ -1302,8 +1249,6 @@ private:
 
 	FAggregatorRef& FindOrCreateAttributeAggregator(FGameplayAttribute Attribute);
 
-	FActiveGameplayEffectHandle LastAssignedHandle;
-
 	TMap<FGameplayAttribute, FAggregatorRef> OngoingPropertyEffects;
 
 	TMap<FGameplayAttribute, FOnGameplayAttributeChange> AttributeChangeDelegates;
@@ -1336,6 +1281,7 @@ struct TStructOpsTypeTraits< FActiveGameplayEffectsContainer > : public TStructO
 };
 
 
+/** Allows blueprints to generate a GameplayEffectSpec once and then reference it by handle, to apply it multiple times/multiple targets. */
 USTRUCT(BlueprintType)
 struct GAMEPLAYABILITIES_API FGameplayEffectSpecHandle
 {
diff --git a/Engine/Source/Runtime/GameplayAbilities/Public/GameplayEffectTypes.h b/Engine/Source/Runtime/GameplayAbilities/Public/GameplayEffectTypes.h
--- a/Engine/Source/Runtime/GameplayAbilities/Public/GameplayEffectTypes.h
+++ b/Engine/Source/Runtime/GameplayAbilities/Public/GameplayEffectTypes.h
@@ -131,6 +131,64 @@ namespace EGameplayEffectStackingPolicy
 }
 
 
+/**
+ * This handle is required for things outside of FActiveGameplayEffectsContainer to refer to a specific active GameplayEffect
+ *	For example if a skill needs to create an active effect and then destroy that specific effect that it created, it has to do so
+ *	through a handle. a pointer or index into the active list is not sufficient.
+ */
+USTRUCT(BlueprintType)
+struct FActiveGameplayEffectHandle
+{
+	GENERATED_USTRUCT_BODY()
+
+	FActiveGameplayEffectHandle()
+		: Handle(INDEX_NONE)
+	{
+
+	}
+
+	FActiveGameplayEffectHandle(int32 InHandle)
+		: Handle(InHandle)
+	{
+
+	}
+
+	bool IsValid() const
+	{
+		return Handle != INDEX_NONE;
+	}
+
+	static FActiveGameplayEffectHandle GenerateNewHandle(UAbilitySystemComponent* OwningComponent);
+
+	UAbilitySystemComponent* GetOwningAbilitySystemComponent();
+
+	bool operator==(const FActiveGameplayEffectHandle& Other) const
+	{
+		return Handle == Other.Handle;
+	}
+
+	bool operator!=(const FActiveGameplayEffectHandle& Other) const
+	{
+		return Handle != Other.Handle;
+	}
+
+	friend uint32 GetTypeHash(const FActiveGameplayEffectHandle& InHandle)
+	{
+		return InHandle.Handle;
+	}
+
+	FString ToString() const
+	{
+		return FString::Printf(TEXT("%d"), Handle);
+	}
+
+private:
+
+	UPROPERTY()
+	int32 Handle;
+};
+
+
 /**
  * FGameplayEffectInstigatorContext
  *	Data struct for an instigator. This is still being fleshed out. We will want to track actors but also be able to provide some level of tracking for actors that are destroyed.
