- putting targeting and spawning transform logic to the projectile
- making the FGameplayAbilityTargetingLocationInfo take in a UMeshComponent so that we can support skeletal and static meshes
- AOrionAbilityTargetActor_TowerSelector setting up SourceLocation right
- fixing the TArray<<>> to TArray<<> >

diff --git a/Engine/Source/Runtime/GameplayAbilities/Public/Abilities/GameplayAbilityTargetTypes.h b/Engine/Source/Runtime/GameplayAbilities/Public/Abilities/GameplayAbilityTargetTypes.h
--- a/Engine/Source/Runtime/GameplayAbilities/Public/Abilities/GameplayAbilityTargetTypes.h
+++ b/Engine/Source/Runtime/GameplayAbilities/Public/Abilities/GameplayAbilityTargetTypes.h
@@ -67,7 +67,7 @@ struct GAMEPLAYABILITIES_API FGameplayAbilityTargetData
 
 	virtual TArray<TWeakObjectPtr<AActor> >	GetActors() const
 	{
-		return TArray<TWeakObjectPtr<AActor>>();
+		return TArray<TWeakObjectPtr<AActor> >();
 	}
 
 	// -------------------------------------
@@ -233,10 +233,10 @@ struct GAMEPLAYABILITIES_API FGameplayAbilityTargetData_Radius : public FGamepla
 	FGameplayAbilityTargetData_Radius()
 	: Origin(0.f) { }
 
-	FGameplayAbilityTargetData_Radius(const TArray<TWeakObjectPtr<AActor>> InActors, const FVector& InOrigin)
+	FGameplayAbilityTargetData_Radius(const TArray<TWeakObjectPtr<AActor> > InActors, const FVector& InOrigin)
 		: Actors(InActors), Origin(InOrigin) { }
 
-	virtual TArray<TWeakObjectPtr<AActor>>	GetActors() const { return Actors; }
+	virtual TArray<TWeakObjectPtr<AActor> >	GetActors() const { return Actors; }
 
 	virtual bool HasOrigin() const { return true; }
 
@@ -251,7 +251,7 @@ struct GAMEPLAYABILITIES_API FGameplayAbilityTargetData_Radius : public FGamepla
 
 private:
 
-	TArray<TWeakObjectPtr<AActor>> Actors;
+	TArray<TWeakObjectPtr<AActor> > Actors;
 	FVector Origin;
 };
 
@@ -315,7 +315,7 @@ public:
 
 	FGameplayAbilityTargetDataHandle MakeTargetDataHandleFromHitResult(TWeakObjectPtr<UGameplayAbility> Ability, FHitResult HitResult) const;
 
-	FGameplayAbilityTargetDataHandle MakeTargetDataHandleFromActors(TArray<TWeakObjectPtr<AActor>> TargetActors) const;
+	FGameplayAbilityTargetDataHandle MakeTargetDataHandleFromActors(TArray<TWeakObjectPtr<AActor> > TargetActors) const;
 
 	/** Type of location used - will determine what data is transmitted over the network and what fields are used when calculating position. */
 	UPROPERTY(BlueprintReadWrite, meta = (ExposeOnSpawn = true), Category = Targeting)
@@ -331,7 +331,7 @@ public:
 
 	/** Socket-based targeting requires a skeletal mesh component to check for the named socket. */
 	UPROPERTY(BlueprintReadWrite, meta = (ExposeOnSpawn = true), Category = Targeting)
-	USkeletalMeshComponent* SourceComponent;
+	UMeshComponent* SourceComponent;
 
 	/** If SourceComponent is valid, this is the name of the socket transform that will be used. If no Socket is provided, SourceComponent's transform will be used. */
 	UPROPERTY(BlueprintReadWrite, meta = (ExposeOnSpawn = true), Category = Targeting)
@@ -433,9 +433,9 @@ struct GAMEPLAYABILITIES_API FGameplayAbilityTargetData_ActorArray : public FGam
 
 	/** Rather than targeting a single point, this type of targeting selects multiple actors. */
 	UPROPERTY(BlueprintReadWrite, EditAnywhere, Category = Targeting)
-	TArray<TWeakObjectPtr<AActor>> TargetActorArray;
+	TArray<TWeakObjectPtr<AActor> > TargetActorArray;
 
-	virtual TArray<TWeakObjectPtr<AActor>>	GetActors() const override
+	virtual TArray<TWeakObjectPtr<AActor> >	GetActors() const override
 	{
 		return TargetActorArray;
 	}
@@ -576,9 +576,9 @@ struct GAMEPLAYABILITIES_API FGameplayAbilityTargetData_SingleTargetHit : public
 
 	// -------------------------------------
 
-	virtual TArray<TWeakObjectPtr<AActor>>	GetActors() const override
+	virtual TArray<TWeakObjectPtr<AActor> >	GetActors() const override
 	{
-		TArray<TWeakObjectPtr<AActor>>	Actors;
+		TArray<TWeakObjectPtr<AActor> >	Actors;
 		if (HitResult.Actor.IsValid())
 		{
 			Actors.Push(HitResult.Actor.Get());
