restore use of delay nodes/world context objects in gameplay abilities

diff --git a/Engine/Source/Runtime/GameplayAbilities/Public/Abilities/GameplayAbility.h b/Engine/Source/Runtime/GameplayAbilities/Public/Abilities/GameplayAbility.h
--- a/Engine/Source/Runtime/GameplayAbilities/Public/Abilities/GameplayAbility.h
+++ b/Engine/Source/Runtime/GameplayAbilities/Public/Abilities/GameplayAbility.h
@@ -163,6 +163,11 @@ public:
 
 	virtual UWorld* GetWorld() const override
 	{
+		if (!IsInstantiated())
+		{
+			// If we are a CDO, we must return nullptr instead of calling Outer->GetWorld() to fool UObject::ImplementsGetWorld.
+			return nullptr;
+		}
 		return GetOuter()->GetWorld();
 	}
 
