#UE4 - Added FGameplayTagCountContainer struct, which contains a counted set of GameplayTags
  - Added support for an active GameplayAbility blocking other abilities from activating based on ability tags

diff --git a/Engine/Source/Runtime/GameplayAbilities/Private/Abilities/GameplayAbility.cpp b/Engine/Source/Runtime/GameplayAbilities/Private/Abilities/GameplayAbility.cpp
--- a/Engine/Source/Runtime/GameplayAbilities/Private/Abilities/GameplayAbility.cpp
+++ b/Engine/Source/Runtime/GameplayAbilities/Private/Abilities/GameplayAbility.cpp
@@ -271,6 +271,12 @@ void UGameplayAbility::EndAbility(const FGameplayAbilitySpecHandle Handle, const
 	}
 	ActiveTasks.Reset();	// Empty the array but dont resize memory, since this object is probably going to be destroyed very soon anyways.
 
+	// Unblock ability tags we previously blocked
+	if (ActorInfo && ActorInfo->AbilitySystemComponent.IsValid())
+	{
+		ActorInfo->AbilitySystemComponent->UnBlockAbilitiesWithTags(BlockAbilitiesWithTag);
+	}
+
 	// Tell owning AbilitySystemComponent that we ended so it can do stuff (including MarkPendingKill us)
 	ActorInfo->AbilitySystemComponent->NotifyAbilityEnded(Handle, this);
 }
@@ -302,6 +308,7 @@ void UGameplayAbility::PreActivate(const FGameplayAbilitySpecHandle Handle, cons
 	UAbilitySystemComponent* Comp = ActorInfo->AbilitySystemComponent.Get();
 
 	Comp->CancelAbilitiesWithTags(CancelAbilitiesWithTag, Handle, ActorInfo, ActivationInfo, this);
+	Comp->BlockAbilitiesWithTags(BlockAbilitiesWithTag);
 
 	bIsActive = true;
 	Comp->NotifyAbilityActivated(Handle, this);
diff --git a/Engine/Source/Runtime/GameplayAbilities/Private/AbilitySystemComponent_Abilities.cpp b/Engine/Source/Runtime/GameplayAbilities/Private/AbilitySystemComponent_Abilities.cpp
--- a/Engine/Source/Runtime/GameplayAbilities/Private/AbilitySystemComponent_Abilities.cpp
+++ b/Engine/Source/Runtime/GameplayAbilities/Private/AbilitySystemComponent_Abilities.cpp
@@ -277,6 +277,16 @@ void UAbilitySystemComponent::CancelAbilitiesWithTags(const FGameplayTagContaine
 	}
 }
 
+void UAbilitySystemComponent::BlockAbilitiesWithTags(const FGameplayTagContainer Tags)
+{
+	BlockedAbilityTags.UpdateTagMap(Tags, 1);
+}
+
+void UAbilitySystemComponent::UnBlockAbilitiesWithTags(const FGameplayTagContainer Tags)
+{
+	BlockedAbilityTags.UpdateTagMap(Tags, -1);
+}
+
 #if !(UE_BUILD_SHIPPING || UE_BUILD_TEST)
 int32 DenyClientActivation = 0;
 static FAutoConsoleVariableRef CVarDenyClientActivation(
@@ -325,6 +335,12 @@ bool UAbilitySystemComponent::TryActivateAbility(FGameplayAbilitySpecHandle Hand
 
 	UGameplayAbility* Ability = Spec->Ability;
 
+	// Check if any of this ability's tags are currently blocked
+	if (BlockedAbilityTags.HasAnyMatchingGameplayTags(Ability->AbilityTags, EGameplayTagMatchType::IncludeParentTags))
+	{
+		return false;
+	}
+
 	// Always do a non instanced CanActivate check
 	if (!Ability->CanActivateAbility(Handle, ActorInfo))
 	{
diff --git a/Engine/Source/Runtime/GameplayAbilities/Private/GameplayAbilityTypes.cpp b/Engine/Source/Runtime/GameplayAbilities/Private/GameplayAbilityTypes.cpp
--- a/Engine/Source/Runtime/GameplayAbilities/Private/GameplayAbilityTypes.cpp
+++ b/Engine/Source/Runtime/GameplayAbilities/Private/GameplayAbilityTypes.cpp
@@ -2,6 +2,9 @@
 
 #include "AbilitySystemPrivatePCH.h"
 #include "AbilitySystemComponent.h"
+#include "GameplayTagContainer.h"
+#include "GameplayTagsModule.h"
+#include "GameplayEffectTypes.h"
 
 //----------------------------------------------------------------------
 
@@ -92,4 +95,92 @@ void FGameplayAbilityActivationInfo::SetActivationConfirmed()
 
 // -----------------------------------------------------------------
 
+bool FGameplayTagCountContainer::HasMatchingGameplayTag(FGameplayTag TagToCheck, EGameplayTagMatchType::Type TagMatchType) const
+{
+	if (TagMatchType == EGameplayTagMatchType::Explicit)
+	{
+		// Search for TagToCheck
+		const int32* Count = GameplayTagCountMap.Find(TagToCheck);
+		if (Count && *Count > 0)
+		{
+			return true;
+		}
+	}
+	else if (TagMatchType == EGameplayTagMatchType::IncludeParentTags)
+	{
+		// Search for TagToCheck or any of its parent tags
+		FGameplayTagContainer TagAndParentsContainer = IGameplayTagsModule::Get().GetGameplayTagsManager().RequestGameplayTagParents(TagToCheck);
+		for (auto TagIt = TagAndParentsContainer.CreateConstIterator(); TagIt; ++TagIt)
+		{
+			const int32* Count = GameplayTagCountMap.Find(*TagIt);
+			if (Count && *Count > 0)
+			{
+				return true;
+			}
+		}
+	}
+
+	return false;
+}
+
+bool FGameplayTagCountContainer::HasAnyMatchingGameplayTags(const FGameplayTagContainer& TagContainer, EGameplayTagMatchType::Type TagMatchType, bool bCountEmptyAsMatch) const
+{
+	SCOPE_CYCLE_COUNTER(STAT_GameplayEffectsHasAnyTag);
+	if (TagContainer.Num() == 0)
+		return bCountEmptyAsMatch;
+
+	for (auto It = TagContainer.CreateConstIterator(); It; ++It)
+	{
+		if (HasMatchingGameplayTag(*It, TagMatchType))
+		{
+			return true;
+		}
+	}
+
+	return false;
+}
+
+void FGameplayTagCountContainer::UpdateTagMap(const FGameplayTag& Tag, int32 CountDelta)
+{
+	// Update count of Tag
+	int32& Count = GameplayTagCountMap.FindOrAdd(Tag);
+
+	bool WasZero = Count == 0;
+	Count = FMath::Max(Count + CountDelta, 0);
+		
+	// If we went from 0->1 or 1->0
+	if (WasZero || Count == 0)
+	{
+		FOnGameplayEffectTagCountChanged *Delegate = GameplayTagEventMap.Find(Tag);
+		if (Delegate)
+		{
+			Delegate->Broadcast(Tag, Count);
+		}
+	}
+}
+
+void FGameplayTagCountContainer::UpdateTagMap(const FGameplayTagContainer& Container, int32 CountDelta)
+{
+	for (auto TagIt = Container.CreateConstIterator(); TagIt; ++TagIt)
+	{
+		const FGameplayTag& BaseTag = *TagIt;
+		if (TagContainerType == EGameplayTagMatchType::Explicit)
+		{
+			// Update count of BaseTag
+			UpdateTagMap(BaseTag, CountDelta);
+		}
+		else if (TagContainerType == EGameplayTagMatchType::IncludeParentTags)
+		{
+			// Update count of BaseTag and all of its parent tags
+			FGameplayTagContainer TagAndParentsContainer = IGameplayTagsModule::Get().GetGameplayTagsManager().RequestGameplayTagParents(BaseTag);
+			for (auto ParentTagIt = TagAndParentsContainer.CreateConstIterator(); ParentTagIt; ++ParentTagIt)
+			{
+				UpdateTagMap(*ParentTagIt, CountDelta);
+			}
+		}
+	}
+}
+
+
+
 
diff --git a/Engine/Source/Runtime/GameplayAbilities/Public/Abilities/GameplayAbility.h b/Engine/Source/Runtime/GameplayAbilities/Public/Abilities/GameplayAbility.h
--- a/Engine/Source/Runtime/GameplayAbilities/Public/Abilities/GameplayAbility.h
+++ b/Engine/Source/Runtime/GameplayAbilities/Public/Abilities/GameplayAbility.h
@@ -408,6 +408,10 @@ protected:
 	UPROPERTY(EditDefaultsOnly, Category = Tags)
 	FGameplayTagContainer CancelAbilitiesWithTag;
 
+	/** Abilities with these tags are blocked while this ability is active */
+	UPROPERTY(EditDefaultsOnly, Category = Tags)
+	FGameplayTagContainer BlockAbilitiesWithTag;
+
 	// ----------------------------------------------------------------------------------------------------------------
 	//
 	//	Ability Tasks
diff --git a/Engine/Source/Runtime/GameplayAbilities/Public/Abilities/GameplayAbilityTypes.h b/Engine/Source/Runtime/GameplayAbilities/Public/Abilities/GameplayAbilityTypes.h
--- a/Engine/Source/Runtime/GameplayAbilities/Public/Abilities/GameplayAbilityTypes.h
+++ b/Engine/Source/Runtime/GameplayAbilities/Public/Abilities/GameplayAbilityTypes.h
@@ -2,6 +2,7 @@
 
 #pragma once
 
+#include "GameplayEffect.h"
 #include "GameplayAbilityTypes.generated.h"
 
 class UGameplayEffect;
@@ -325,6 +326,27 @@ struct FGameplayAbiliyInputBinds
 	UEnum* GetBindEnum() { return FindObject<UEnum>(ANY_PACKAGE, *EnumName); }
 };
 
+/** 
+ *	Structure that contains a counted set of GameplayTags. Can optionally include parent tags
+ *	
+ */
+struct FGameplayTagCountContainer
+{
+	FGameplayTagCountContainer()
+	: TagContainerType(EGameplayTagMatchType::Explicit)
+	{ }
+
+	bool HasMatchingGameplayTag(FGameplayTag TagToCheck, EGameplayTagMatchType::Type TagMatchType) const;
+	bool HasAnyMatchingGameplayTags(const FGameplayTagContainer& TagContainer, EGameplayTagMatchType::Type TagMatchType, bool bCountEmptyAsMatch = true) const;
+	void UpdateTagMap(const struct FGameplayTagContainer& Container, int32 CountDelta);
+	void UpdateTagMap(const struct FGameplayTag& Tag, int32 CountDelta);
+
+	TMap<struct FGameplayTag, FOnGameplayEffectTagCountChanged> GameplayTagEventMap;
+	TMap<struct FGameplayTag, int32> GameplayTagCountMap;
+
+	EGameplayTagMatchType::Type TagContainerType;
+};
+
 /** Used for cleaning up predicted data on network clients */
 DECLARE_MULTICAST_DELEGATE(FAbilitySystemComponentPredictionKeyClear);
 
diff --git a/Engine/Source/Runtime/GameplayAbilities/Public/AbilitySystemComponent.h b/Engine/Source/Runtime/GameplayAbilities/Public/AbilitySystemComponent.h
--- a/Engine/Source/Runtime/GameplayAbilities/Public/AbilitySystemComponent.h
+++ b/Engine/Source/Runtime/GameplayAbilities/Public/AbilitySystemComponent.h
@@ -391,6 +391,9 @@ class GAMEPLAYABILITIES_API UAbilitySystemComponent : public UActorComponent, pu
 	UGameplayAbility* CreateNewInstanceOfAbility(FGameplayAbilitySpec& Spec, UGameplayAbility* Ability);
 
 	void CancelAbilitiesWithTags(const FGameplayTagContainer Tags, const FGameplayAbilitySpecHandle Handle, const FGameplayAbilityActorInfo* ActorInfo, const FGameplayAbilityActivationInfo ActivationInfo, UGameplayAbility* Ignore);
+
+	void BlockAbilitiesWithTags(const FGameplayTagContainer Tags);
+	void UnBlockAbilitiesWithTags(const FGameplayTagContainer Tags);
 	
 	/** FUll list of all instance-per-execution gameplay abilities associated with this component */
 	UPROPERTY()
@@ -597,6 +600,8 @@ private:
 	UPROPERTY(ReplicatedUsing=OnRep_GameplayEffects)
 	FActiveGameplayCueContainer	ActiveGameplayCues;
 
+	FGameplayTagCountContainer BlockedAbilityTags;
+
 	UFUNCTION()
 	void OnRep_GameplayEffects();
 
