#UE4 - Modified FActiveGameplayEffectsContainer to use FGameplayTagCountContainer

diff --git a/Engine/Source/Runtime/GameplayAbilities/Private/AbilitySystemComponent_Abilities.cpp b/Engine/Source/Runtime/GameplayAbilities/Private/AbilitySystemComponent_Abilities.cpp
--- a/Engine/Source/Runtime/GameplayAbilities/Private/AbilitySystemComponent_Abilities.cpp
+++ b/Engine/Source/Runtime/GameplayAbilities/Private/AbilitySystemComponent_Abilities.cpp
@@ -336,7 +336,7 @@ bool UAbilitySystemComponent::TryActivateAbility(FGameplayAbilitySpecHandle Hand
 	UGameplayAbility* Ability = Spec->Ability;
 
 	// Check if any of this ability's tags are currently blocked
-	if (BlockedAbilityTags.HasAnyMatchingGameplayTags(Ability->AbilityTags, EGameplayTagMatchType::IncludeParentTags))
+	if (BlockedAbilityTags.HasAnyMatchingGameplayTags(Ability->AbilityTags, EGameplayTagMatchType::IncludeParentTags, false))
 	{
 		return false;
 	}
diff --git a/Engine/Source/Runtime/GameplayAbilities/Private/GameplayAbilityTypes.cpp b/Engine/Source/Runtime/GameplayAbilities/Private/GameplayAbilityTypes.cpp
--- a/Engine/Source/Runtime/GameplayAbilities/Private/GameplayAbilityTypes.cpp
+++ b/Engine/Source/Runtime/GameplayAbilities/Private/GameplayAbilityTypes.cpp
@@ -93,94 +93,4 @@ void FGameplayAbilityActivationInfo::SetActivationConfirmed()
 }
 
 
-// -----------------------------------------------------------------
-
-bool FGameplayTagCountContainer::HasMatchingGameplayTag(FGameplayTag TagToCheck, EGameplayTagMatchType::Type TagMatchType) const
-{
-	if (TagMatchType == EGameplayTagMatchType::Explicit)
-	{
-		// Search for TagToCheck
-		const int32* Count = GameplayTagCountMap.Find(TagToCheck);
-		if (Count && *Count > 0)
-		{
-			return true;
-		}
-	}
-	else if (TagMatchType == EGameplayTagMatchType::IncludeParentTags)
-	{
-		// Search for TagToCheck or any of its parent tags
-		FGameplayTagContainer TagAndParentsContainer = IGameplayTagsModule::Get().GetGameplayTagsManager().RequestGameplayTagParents(TagToCheck);
-		for (auto TagIt = TagAndParentsContainer.CreateConstIterator(); TagIt; ++TagIt)
-		{
-			const int32* Count = GameplayTagCountMap.Find(*TagIt);
-			if (Count && *Count > 0)
-			{
-				return true;
-			}
-		}
-	}
-
-	return false;
-}
-
-bool FGameplayTagCountContainer::HasAnyMatchingGameplayTags(const FGameplayTagContainer& TagContainer, EGameplayTagMatchType::Type TagMatchType, bool bCountEmptyAsMatch) const
-{
-	SCOPE_CYCLE_COUNTER(STAT_GameplayEffectsHasAnyTag);
-	if (TagContainer.Num() == 0)
-		return bCountEmptyAsMatch;
-
-	for (auto It = TagContainer.CreateConstIterator(); It; ++It)
-	{
-		if (HasMatchingGameplayTag(*It, TagMatchType))
-		{
-			return true;
-		}
-	}
-
-	return false;
-}
-
-void FGameplayTagCountContainer::UpdateTagMap(const FGameplayTag& Tag, int32 CountDelta)
-{
-	// Update count of Tag
-	int32& Count = GameplayTagCountMap.FindOrAdd(Tag);
-
-	bool WasZero = Count == 0;
-	Count = FMath::Max(Count + CountDelta, 0);
-		
-	// If we went from 0->1 or 1->0
-	if (WasZero || Count == 0)
-	{
-		FOnGameplayEffectTagCountChanged *Delegate = GameplayTagEventMap.Find(Tag);
-		if (Delegate)
-		{
-			Delegate->Broadcast(Tag, Count);
-		}
-	}
-}
-
-void FGameplayTagCountContainer::UpdateTagMap(const FGameplayTagContainer& Container, int32 CountDelta)
-{
-	for (auto TagIt = Container.CreateConstIterator(); TagIt; ++TagIt)
-	{
-		const FGameplayTag& BaseTag = *TagIt;
-		if (TagContainerType == EGameplayTagMatchType::Explicit)
-		{
-			// Update count of BaseTag
-			UpdateTagMap(BaseTag, CountDelta);
-		}
-		else if (TagContainerType == EGameplayTagMatchType::IncludeParentTags)
-		{
-			// Update count of BaseTag and all of its parent tags
-			FGameplayTagContainer TagAndParentsContainer = IGameplayTagsModule::Get().GetGameplayTagsManager().RequestGameplayTagParents(BaseTag);
-			for (auto ParentTagIt = TagAndParentsContainer.CreateConstIterator(); ParentTagIt; ++ParentTagIt)
-			{
-				UpdateTagMap(*ParentTagIt, CountDelta);
-			}
-		}
-	}
-}
-
-
-
 
diff --git a/Engine/Source/Runtime/GameplayAbilities/Private/GameplayEffect.cpp b/Engine/Source/Runtime/GameplayAbilities/Private/GameplayEffect.cpp
--- a/Engine/Source/Runtime/GameplayAbilities/Private/GameplayEffect.cpp
+++ b/Engine/Source/Runtime/GameplayAbilities/Private/GameplayEffect.cpp
@@ -3,6 +3,7 @@
 #include "AbilitySystemPrivatePCH.h"
 #include "GameplayEffect.h"
 #include "GameplayEffectExtension.h"
+#include "GameplayEffectTypes.h"
 #include "GameplayEffectStackingExtension.h"
 #include "GameplayTagsModule.h"
 #include "AbilitySystemComponent.h"
@@ -1641,46 +1642,23 @@ FActiveGameplayEffect & FActiveGameplayEffectsContainer::CreateNewActiveGameplay
 	return NewEffect;
 }
 
-void FActiveGameplayEffectsContainer::UpdateTagMap(const FGameplayTag& BaseTag, int32 CountDelta, IGameplayTagsModule& GameplayTagsModule)
+void FActiveGameplayEffectsContainer::UpdateTagMap(const FGameplayTag& BaseTag, int32 CountDelta)
 {
-	FGameplayTagContainer TagAndParentsContainer = GameplayTagsModule.GetGameplayTagsManager().RequestGameplayTagParents(BaseTag);
-	for (auto TagIt = TagAndParentsContainer.CreateConstIterator(); TagIt; ++TagIt)
-	{
-		const FGameplayTag& Tag = *TagIt;
-		int32& Count = GameplayTagCountMap.FindOrAdd(Tag);
-
-		bool WasZero = Count == 0;
-		Count = FMath::Max(Count + CountDelta, 0);
-		
-		// If we went from 0->1 or 1->0
-		if (WasZero || Count == 0)
-		{
-			FOnGameplayEffectTagCountChanged *Delegate = GameplayTagEventMap.Find(Tag);
-			if (Delegate)
-			{
-				Delegate->Broadcast(Tag, Count);
-			}
-		}
-	}
+	GameplayTagCountContainer.UpdateTagMap(BaseTag, CountDelta);
 }
 
-void FActiveGameplayEffectsContainer::UpdateTagMap(const FGameplayTagContainer& Container, int32 CountDelta, IGameplayTagsModule& GameplayTagsModule)
+void FActiveGameplayEffectsContainer::UpdateTagMap(const FGameplayTagContainer& Container, int32 CountDelta)
 {
-	for (auto TagIt = Container.CreateConstIterator(); TagIt; ++TagIt)
-	{
-		const FGameplayTag& Tag = *TagIt;
-		UpdateTagMap(Tag, CountDelta, GameplayTagsModule);
-	}
+	GameplayTagCountContainer.UpdateTagMap(Container, CountDelta);
 }
 
 void FActiveGameplayEffectsContainer::InternalOnActiveGameplayEffectAdded(const FActiveGameplayEffect& Effect)
 {
 	// Update gameplaytag count and broadcast delegate if we just added this tag (count=0, prior to increment)
-	IGameplayTagsModule& GameplayTagsModule = IGameplayTagsModule::Get();
-	UpdateTagMap(Effect.Spec.Def->OwnedTagsContainer, 1, GameplayTagsModule);
+	UpdateTagMap(Effect.Spec.Def->OwnedTagsContainer, 1);
 	for (const FGameplayEffectCue& Cue : Effect.Spec.Def->GameplayCues)
 	{
-		UpdateTagMap(Cue.GameplayCueTags, 1, GameplayTagsModule);
+		UpdateTagMap(Cue.GameplayCueTags, 1);
 	}	
 }
 
@@ -1744,11 +1722,11 @@ void FActiveGameplayEffectsContainer::InternalOnActiveGameplayEffectRemoved(cons
 
 	// Update gameplaytag count and broadcast delegate if we are at 0
 	IGameplayTagsModule& GameplayTagsModule = IGameplayTagsModule::Get();
-	UpdateTagMap(Effect.Spec.Def->OwnedTagsContainer, -1, GameplayTagsModule);
+	UpdateTagMap(Effect.Spec.Def->OwnedTagsContainer, -1);
 
 	for (const FGameplayEffectCue& Cue : Effect.Spec.Def->GameplayCues)
 	{
-		UpdateTagMap(Cue.GameplayCueTags, -1, GameplayTagsModule);
+		UpdateTagMap(Cue.GameplayCueTags, -1);
 	}
 }
 
@@ -1941,7 +1919,7 @@ void FActiveGameplayEffectsContainer::RecalculateStacking()
 void FActiveGameplayEffectsContainer::GetOwnedGameplayTags(FGameplayTagContainer& TagContainer) const
 {
 	SCOPE_CYCLE_COUNTER(STAT_GameplayEffectsGetOwnedTags);
-	for (auto It = GameplayTagCountMap.CreateConstIterator(); It; ++It)
+	for (auto It = GameplayTagCountContainer.GameplayTagCountMap.CreateConstIterator(); It; ++It)
 	{
 		if(It.Value() > 0)
 		{
@@ -1952,52 +1930,19 @@ void FActiveGameplayEffectsContainer::GetOwnedGameplayTags(FGameplayTagContainer
 
 bool FActiveGameplayEffectsContainer::HasMatchingGameplayTag(FGameplayTag TagToCheck) const
 {
-	/**
-	 *	The desired behavior here is:
-	 *		"Do we have tag a.b.c and the GameplayEffect has a.b.c.d -> match!"
-	 *		"Do we have tag a.b.c.d. and the GameplayEffect has a.b.c -> no match!"
-	 */
-
-	const int32* Count = GameplayTagCountMap.Find(TagToCheck);
-	if (Count)
-	{
-		return *Count > 0;
-	}
-	return false;
+	return GameplayTagCountContainer.HasMatchingGameplayTag(TagToCheck, EGameplayTagMatchType::Explicit);
 }
 
 bool FActiveGameplayEffectsContainer::HasAllMatchingGameplayTags(const FGameplayTagContainer& TagContainer, bool bCountEmptyAsMatch) const
 {
 	SCOPE_CYCLE_COUNTER(STAT_GameplayEffectsHasAllTags);
-	if (TagContainer.Num() == 0)
-		return bCountEmptyAsMatch;
-
-	for (auto It = TagContainer.CreateConstIterator(); It; ++It)
-	{
-		if (!HasMatchingGameplayTag(*It))
-		{
-			return false;
-		}
-	}
-
-	return true;
+	return GameplayTagCountContainer.HasAllMatchingGameplayTags(TagContainer, EGameplayTagMatchType::Explicit, bCountEmptyAsMatch);
 }
 
 bool FActiveGameplayEffectsContainer::HasAnyMatchingGameplayTags(const FGameplayTagContainer& TagContainer, bool bCountEmptyAsMatch) const
 {
 	SCOPE_CYCLE_COUNTER(STAT_GameplayEffectsHasAnyTag);
-	if (TagContainer.Num() == 0)
-		return bCountEmptyAsMatch;
-
-	for (auto It = TagContainer.CreateConstIterator(); It; ++It)
-	{
-		if (HasMatchingGameplayTag(*It))
-		{
-			return true;
-		}
-	}
-
-	return false;
+	return GameplayTagCountContainer.HasAnyMatchingGameplayTags(TagContainer, EGameplayTagMatchType::Explicit, bCountEmptyAsMatch);
 }
 
 bool FActiveGameplayEffectsContainer::CanApplyAttributeModifiers(const UGameplayEffect *GameplayEffect, float Level, AActor *Instigator)
@@ -2056,7 +2001,7 @@ TArray<float> FActiveGameplayEffectsContainer::GetActiveEffectsTimeRemaining(con
 
 FOnGameplayEffectTagCountChanged& FActiveGameplayEffectsContainer::RegisterGameplayTagEvent(FGameplayTag Tag)
 {
-	return GameplayTagEventMap.FindOrAdd(Tag);
+	return GameplayTagCountContainer.GameplayTagEventMap.FindOrAdd(Tag);
 }
 
 FOnGameplayAttributeChange& FActiveGameplayEffectsContainer::RegisterGameplayAttributeEvent(FGameplayAttribute Attribute)
@@ -2175,4 +2120,111 @@ UAbilitySystemComponent* FActiveGameplayEffectHandle::GetOwningAbilitySystemComp
 	}
 
 	return nullptr;	
-}
\ No newline at end of file
+}
+
+// -----------------------------------------------------------------
+
+bool FGameplayTagCountContainer::HasMatchingGameplayTag(FGameplayTag TagToCheck, EGameplayTagMatchType::Type TagMatchType) const
+{
+	if (TagMatchType == EGameplayTagMatchType::Explicit)
+	{
+		// Search for TagToCheck
+		const int32* Count = GameplayTagCountMap.Find(TagToCheck);
+		if (Count && *Count > 0)
+		{
+			return true;
+		}
+	}
+	else if (TagMatchType == EGameplayTagMatchType::IncludeParentTags)
+	{
+		// Search for TagToCheck or any of its parent tags
+		FGameplayTagContainer TagAndParentsContainer = IGameplayTagsModule::Get().GetGameplayTagsManager().RequestGameplayTagParents(TagToCheck);
+		for (auto TagIt = TagAndParentsContainer.CreateConstIterator(); TagIt; ++TagIt)
+		{
+			const int32* Count = GameplayTagCountMap.Find(*TagIt);
+			if (Count && *Count > 0)
+			{
+				return true;
+			}
+		}
+	}
+
+	return false;
+}
+
+bool FGameplayTagCountContainer::HasAllMatchingGameplayTags(const FGameplayTagContainer& TagContainer, EGameplayTagMatchType::Type TagMatchType, bool bCountEmptyAsMatch) const
+{
+	if (TagContainer.Num() == 0)
+		return bCountEmptyAsMatch;
+
+	for (auto It = TagContainer.CreateConstIterator(); It; ++It)
+	{
+		if (!HasMatchingGameplayTag(*It, TagMatchType))
+		{
+			return false;
+		}
+	}
+
+	return true;
+}
+
+
+bool FGameplayTagCountContainer::HasAnyMatchingGameplayTags(const FGameplayTagContainer& TagContainer, EGameplayTagMatchType::Type TagMatchType, bool bCountEmptyAsMatch) const
+{
+	if (TagContainer.Num() == 0)
+		return bCountEmptyAsMatch;
+
+	for (auto It = TagContainer.CreateConstIterator(); It; ++It)
+	{
+		if (HasMatchingGameplayTag(*It, TagMatchType))
+		{
+			return true;
+		}
+	}
+
+	return false;
+}
+
+void FGameplayTagCountContainer::UpdateTagMap(const FGameplayTag& Tag, int32 CountDelta)
+{
+	// Update count of Tag
+	int32& Count = GameplayTagCountMap.FindOrAdd(Tag);
+
+	bool WasZero = Count == 0;
+	Count = FMath::Max(Count + CountDelta, 0);
+		
+	// If we went from 0->1 or 1->0
+	if (WasZero || Count == 0)
+	{
+		FOnGameplayEffectTagCountChanged *Delegate = GameplayTagEventMap.Find(Tag);
+		if (Delegate)
+		{
+			Delegate->Broadcast(Tag, Count);
+		}
+	}
+}
+
+void FGameplayTagCountContainer::UpdateTagMap(const FGameplayTagContainer& Container, int32 CountDelta)
+{
+	for (auto TagIt = Container.CreateConstIterator(); TagIt; ++TagIt)
+	{
+		const FGameplayTag& BaseTag = *TagIt;
+		if (TagContainerType == EGameplayTagMatchType::Explicit)
+		{
+			// Update count of BaseTag
+			UpdateTagMap(BaseTag, CountDelta);
+		}
+		else if (TagContainerType == EGameplayTagMatchType::IncludeParentTags)
+		{
+			// Update count of BaseTag and all of its parent tags
+			FGameplayTagContainer TagAndParentsContainer = IGameplayTagsModule::Get().GetGameplayTagsManager().RequestGameplayTagParents(BaseTag);
+			for (auto ParentTagIt = TagAndParentsContainer.CreateConstIterator(); ParentTagIt; ++ParentTagIt)
+			{
+				UpdateTagMap(*ParentTagIt, CountDelta);
+			}
+		}
+	}
+}
+
+
+
diff --git a/Engine/Source/Runtime/GameplayAbilities/Public/Abilities/GameplayAbilityTypes.h b/Engine/Source/Runtime/GameplayAbilities/Public/Abilities/GameplayAbilityTypes.h
--- a/Engine/Source/Runtime/GameplayAbilities/Public/Abilities/GameplayAbilityTypes.h
+++ b/Engine/Source/Runtime/GameplayAbilities/Public/Abilities/GameplayAbilityTypes.h
@@ -326,27 +326,6 @@ struct FGameplayAbiliyInputBinds
 	UEnum* GetBindEnum() { return FindObject<UEnum>(ANY_PACKAGE, *EnumName); }
 };
 
-/** 
- *	Structure that contains a counted set of GameplayTags. Can optionally include parent tags
- *	
- */
-struct FGameplayTagCountContainer
-{
-	FGameplayTagCountContainer()
-	: TagContainerType(EGameplayTagMatchType::Explicit)
-	{ }
-
-	bool HasMatchingGameplayTag(FGameplayTag TagToCheck, EGameplayTagMatchType::Type TagMatchType) const;
-	bool HasAnyMatchingGameplayTags(const FGameplayTagContainer& TagContainer, EGameplayTagMatchType::Type TagMatchType, bool bCountEmptyAsMatch = true) const;
-	void UpdateTagMap(const struct FGameplayTagContainer& Container, int32 CountDelta);
-	void UpdateTagMap(const struct FGameplayTag& Tag, int32 CountDelta);
-
-	TMap<struct FGameplayTag, FOnGameplayEffectTagCountChanged> GameplayTagEventMap;
-	TMap<struct FGameplayTag, int32> GameplayTagCountMap;
-
-	EGameplayTagMatchType::Type TagContainerType;
-};
-
 /** Used for cleaning up predicted data on network clients */
 DECLARE_MULTICAST_DELEGATE(FAbilitySystemComponentPredictionKeyClear);
 
diff --git a/Engine/Source/Runtime/GameplayAbilities/Public/GameplayEffect.h b/Engine/Source/Runtime/GameplayAbilities/Public/GameplayEffect.h
--- a/Engine/Source/Runtime/GameplayAbilities/Public/GameplayEffect.h
+++ b/Engine/Source/Runtime/GameplayAbilities/Public/GameplayEffect.h
@@ -1156,7 +1156,7 @@ struct FActiveGameplayEffectsContainer : public FFastArraySerializer
 
 	friend struct FActiveGameplayEffect;
 
-	FActiveGameplayEffectsContainer() : bNeedToRecalculateStacks(false) {};
+	FActiveGameplayEffectsContainer() : bNeedToRecalculateStacks(false), GameplayTagCountContainer(EGameplayTagMatchType::IncludeParentTags) {};
 
 	UPROPERTY()
 	TArray< FActiveGameplayEffect >	GameplayEffects;
@@ -1253,8 +1253,7 @@ private:
 
 	TMap<FGameplayAttribute, FOnGameplayAttributeChange> AttributeChangeDelegates;
 
-	TMap<FGameplayTag, FOnGameplayEffectTagCountChanged> GameplayTagEventMap;
-	TMap<FGameplayTag, int32> GameplayTagCountMap;
+	FGameplayTagCountContainer GameplayTagCountContainer;
 
 	void InternalUpdateNumericalAttribute(FGameplayAttribute Attribute, float NewValue, const FGameplayEffectModCallbackData* ModData);
 
@@ -1267,8 +1266,8 @@ private:
 	void InternalOnActiveGameplayEffectAdded(const FActiveGameplayEffect& Effect);
 	void InternalOnActiveGameplayEffectRemoved(const FActiveGameplayEffect& Effect);
 
-	void UpdateTagMap(const FGameplayTagContainer& Container, int32 CountDelta, class IGameplayTagsModule& GameplayTagsModule);
-	void UpdateTagMap(const FGameplayTag& Tag, int32 CountDelta, class IGameplayTagsModule& GameplayTagsModule);
+	void UpdateTagMap(const FGameplayTagContainer& Container, int32 CountDelta);
+	void UpdateTagMap(const FGameplayTag& Tag, int32 CountDelta);
 };
 
 template<>
diff --git a/Engine/Source/Runtime/GameplayAbilities/Public/GameplayEffectTypes.h b/Engine/Source/Runtime/GameplayAbilities/Public/GameplayEffectTypes.h
--- a/Engine/Source/Runtime/GameplayAbilities/Public/GameplayEffectTypes.h
+++ b/Engine/Source/Runtime/GameplayAbilities/Public/GameplayEffectTypes.h
@@ -314,4 +314,32 @@ DECLARE_MULTICAST_DELEGATE_TwoParams(FOnGameplayEffectTagCountChanged, const FGa
 
 DECLARE_MULTICAST_DELEGATE(FOnActiveGameplayEffectRemoved);
 
-DECLARE_MULTICAST_DELEGATE_TwoParams(FOnGameplayAttributeChange, float ,const FGameplayEffectModCallbackData*);
\ No newline at end of file
+DECLARE_MULTICAST_DELEGATE_TwoParams(FOnGameplayAttributeChange, float ,const FGameplayEffectModCallbackData*);
+
+// -----------------------------------------------------------
+
+/** 
+ *	Structure that contains a counted set of GameplayTags. Can optionally include parent tags
+ *	
+ */
+struct FGameplayTagCountContainer
+{
+	FGameplayTagCountContainer()
+	: TagContainerType(EGameplayTagMatchType::Explicit)
+	{ }
+
+	FGameplayTagCountContainer(EGameplayTagMatchType::Type InTagContainerType)
+	: TagContainerType(InTagContainerType)
+	{ }
+
+	bool HasMatchingGameplayTag(FGameplayTag TagToCheck, EGameplayTagMatchType::Type TagMatchType) const;
+	bool HasAllMatchingGameplayTags(const FGameplayTagContainer& TagContainer, EGameplayTagMatchType::Type TagMatchType, bool bCountEmptyAsMatch = true) const;
+	bool HasAnyMatchingGameplayTags(const FGameplayTagContainer& TagContainer, EGameplayTagMatchType::Type TagMatchType, bool bCountEmptyAsMatch = true) const;
+	void UpdateTagMap(const struct FGameplayTagContainer& Container, int32 CountDelta);
+	void UpdateTagMap(const struct FGameplayTag& Tag, int32 CountDelta);
+
+	TMap<struct FGameplayTag, FOnGameplayEffectTagCountChanged> GameplayTagEventMap;
+	TMap<struct FGameplayTag, int32> GameplayTagCountMap;
+
+	EGameplayTagMatchType::Type TagContainerType;
+};
