minor cleanup. Call CanJump in default CharacterJump ability.

diff --git a/Engine/Source/Runtime/GameplayAbilities/Private/Abilities/GameplayAbility_CharacterJump.cpp b/Engine/Source/Runtime/GameplayAbilities/Private/Abilities/GameplayAbility_CharacterJump.cpp
--- a/Engine/Source/Runtime/GameplayAbilities/Private/Abilities/GameplayAbility_CharacterJump.cpp
+++ b/Engine/Source/Runtime/GameplayAbilities/Private/Abilities/GameplayAbility_CharacterJump.cpp
@@ -46,8 +46,8 @@ bool UGameplayAbility_CharacterJump::CanActivateAbility(const FGameplayAbilitySp
 		return false;
 	}
 
-	const UCharacterMovementComponent *CharMovement = CastChecked<UCharacterMovementComponent>(ActorInfo->MovementComponent.Get(), ECastCheckedType::NullAllowed);
-	return (CharMovement && CharMovement->MovementMode == EMovementMode::MOVE_Walking);
+	const ACharacter* Character = CastChecked<ACharacter>(ActorInfo->Actor.Get(), ECastCheckedType::NullAllowed);
+	return (Character && Character->CanJump());
 }
 
 /**
diff --git a/Engine/Source/Runtime/GameplayAbilities/Public/Abilities/GameplayAbilityTypes.h b/Engine/Source/Runtime/GameplayAbilities/Public/Abilities/GameplayAbilityTypes.h
--- a/Engine/Source/Runtime/GameplayAbilities/Public/Abilities/GameplayAbilityTypes.h
+++ b/Engine/Source/Runtime/GameplayAbilities/Public/Abilities/GameplayAbilityTypes.h
@@ -13,6 +13,7 @@ class AGameplayAbilityTargetActor;
 class UAbilityTask;
 class UAttributeSet;
 
+GAMEPLAYABILITIES_API DECLARE_LOG_CATEGORY_EXTERN(LogAbilitySystemComponent, Log, All);
 
 UENUM(BlueprintType)
 namespace EGameplayAbilityInstancingPolicy
@@ -326,6 +327,19 @@ struct FGameplayAbiliyInputBinds
 	UEnum* GetBindEnum() { return FindObject<UEnum>(ANY_PACKAGE, *EnumName); }
 };
 
+USTRUCT()
+struct GAMEPLAYABILITIES_API FAttributeDefaults
+{
+	GENERATED_USTRUCT_BODY()
+
+	UPROPERTY(EditAnywhere, Category = "AttributeTest")
+	TSubclassOf<UAttributeSet> Attributes;
+
+	UPROPERTY(EditAnywhere, BlueprintReadWrite, Category = "AttributeTest")
+	class UDataTable*	DefaultStartingTable;
+};
+
+
 /** Used for cleaning up predicted data on network clients */
 DECLARE_MULTICAST_DELEGATE(FAbilitySystemComponentPredictionKeyClear);
 
diff --git a/Engine/Source/Runtime/GameplayAbilities/Public/AbilitySystemComponent.h b/Engine/Source/Runtime/GameplayAbilities/Public/AbilitySystemComponent.h
--- a/Engine/Source/Runtime/GameplayAbilities/Public/AbilitySystemComponent.h
+++ b/Engine/Source/Runtime/GameplayAbilities/Public/AbilitySystemComponent.h
@@ -7,24 +7,6 @@
 #include "Abilities/GameplayAbilityTargetTypes.h"
 #include "AbilitySystemComponent.generated.h"
 
-GAMEPLAYABILITIES_API DECLARE_LOG_CATEGORY_EXTERN(LogAbilitySystemComponent, Log, All);
-
-/** Used to register callbacks to confirm/cancel input */
-DECLARE_DYNAMIC_MULTICAST_DELEGATE(FAbilityConfirmOrCancel);
-
-
-USTRUCT()
-struct GAMEPLAYABILITIES_API FAttributeDefaults
-{
-	GENERATED_USTRUCT_BODY()
-
-	UPROPERTY(EditAnywhere, Category="AttributeTest")
-	TSubclassOf<UAttributeSet> Attributes;
-
-	UPROPERTY(EditAnywhere, BlueprintReadWrite, Category="AttributeTest")
-	class UDataTable*	DefaultStartingTable;
-};
-
 /** 
  *	UAbilitySystemComponent	
  *
@@ -64,6 +46,9 @@ class GAMEPLAYABILITIES_API UAbilitySystemComponent : public UActorComponent, pu
 {
 	GENERATED_UCLASS_BODY()
 
+	/** Used to register callbacks to confirm/cancel input */
+	DECLARE_DYNAMIC_MULTICAST_DELEGATE(FAbilityConfirmOrCancel);
+
 	friend FGameplayEffectSpec;
 	friend class AAbilitySystemDebugHUD;
 
