Added TraceStart field to transmitted hit data for abilities. Added HasOrigin and GetOrigin BP nodes to ability targeting toolset.

diff --git a/Engine/Source/Runtime/GameplayAbilities/Private/AbilitySystemBlueprintLibrary.cpp b/Engine/Source/Runtime/GameplayAbilities/Private/AbilitySystemBlueprintLibrary.cpp
--- a/Engine/Source/Runtime/GameplayAbilities/Private/AbilitySystemBlueprintLibrary.cpp
+++ b/Engine/Source/Runtime/GameplayAbilities/Private/AbilitySystemBlueprintLibrary.cpp
@@ -171,12 +171,43 @@ FHitResult UAbilitySystemBlueprintLibrary::GetHitResultFromTargetData(FGameplayA
 	return FHitResult();
 }
 
+bool UAbilitySystemBlueprintLibrary::TargetDataHasOrigin(FGameplayAbilityTargetDataHandle TargetData, int32 Index)
+{
+	FGameplayAbilityTargetData* Data = TargetData.Data[Index].Get();
+	if (Data)
+	{
+		return (Data->HasHitResult() || Data->HasOrigin());
+	}
+	return false;
+}
+
+FTransform UAbilitySystemBlueprintLibrary::GetTargetDataOrigin(FGameplayAbilityTargetDataHandle TargetData, int32 Index)
+{
+	FGameplayAbilityTargetData* Data = TargetData.Data[Index].Get();
+	if (Data)
+	{
+		if (Data->HasOrigin())
+		{
+			return Data->GetOrigin();
+		}
+		if (Data->HasHitResult())
+		{
+			const FHitResult* HitResultPtr = Data->GetHitResult();
+			FTransform ReturnTransform;
+			ReturnTransform.SetLocation(HitResultPtr->TraceStart);
+			ReturnTransform.SetRotation((HitResultPtr->Location - HitResultPtr->TraceStart).SafeNormal().Rotation().Quaternion());
+			return ReturnTransform;
+		}
+	}
+
+	return FTransform::Identity;
+}
+
 bool UAbilitySystemBlueprintLibrary::TargetDataHasEndPoint(FGameplayAbilityTargetDataHandle TargetData, int32 Index)
 {
 	FGameplayAbilityTargetData* Data = TargetData.Data[Index].Get();
 	if (Data)
 	{
-		const FHitResult* HitResultPtr = Data->GetHitResult();
 		return (Data->HasHitResult() || Data->HasEndPoint());
 	}
 	return false;
@@ -205,7 +236,7 @@ FVector UAbilitySystemBlueprintLibrary::GetTargetDataEndPoint(FGameplayAbilityTa
 		}
 	}
 
-	return FVector();
+	return FVector::ZeroVector;
 }
 
 // -------------------------------------------------------------------------------------
diff --git a/Engine/Source/Runtime/GameplayAbilities/Private/GameplayAbilityTargetTypes.cpp b/Engine/Source/Runtime/GameplayAbilities/Private/GameplayAbilityTargetTypes.cpp
--- a/Engine/Source/Runtime/GameplayAbilities/Private/GameplayAbilityTargetTypes.cpp
+++ b/Engine/Source/Runtime/GameplayAbilities/Private/GameplayAbilityTargetTypes.cpp
@@ -213,6 +213,7 @@ bool FGameplayAbilityTargetData_SingleTargetHit::NetSerialize(FArchive& Ar, clas
 {
 	Ar << HitResult.Actor;
 
+	HitResult.TraceStart.NetSerialize(Ar, Map, bOutSuccess);
 	HitResult.Location.NetSerialize(Ar, Map, bOutSuccess);
 	HitResult.Normal.NetSerialize(Ar, Map, bOutSuccess);
 
diff --git a/Engine/Source/Runtime/GameplayAbilities/Public/AbilitySystemBlueprintLibrary.h b/Engine/Source/Runtime/GameplayAbilities/Public/AbilitySystemBlueprintLibrary.h
--- a/Engine/Source/Runtime/GameplayAbilities/Public/AbilitySystemBlueprintLibrary.h
+++ b/Engine/Source/Runtime/GameplayAbilities/Public/AbilitySystemBlueprintLibrary.h
@@ -59,6 +59,12 @@ class UAbilitySystemBlueprintLibrary : public UBlueprintFunctionLibrary
 	UFUNCTION(BlueprintPure, Category = "Ability|TargetData")
 	static FHitResult GetHitResultFromTargetData(FGameplayAbilityTargetDataHandle HitResult, int32 Index);
 
+	UFUNCTION(BlueprintPure, Category = "Ability|TargetData")
+	static bool TargetDataHasOrigin(FGameplayAbilityTargetDataHandle TargetData, int32 Index);
+
+	UFUNCTION(BlueprintPure, Category = "Ability|TargetData")
+	static FTransform GetTargetDataOrigin(FGameplayAbilityTargetDataHandle TargetData, int32 Index);
+
 	UFUNCTION(BlueprintPure, Category = "Ability|TargetData")
 	static bool TargetDataHasEndPoint(FGameplayAbilityTargetDataHandle TargetData, int32 Index);
 
