Updated handle generation for TargetActorArray to permit creating a series of single-actor arrays instead of a single multi-actor array. Refactored a test ability to show how this makes BPs smoother for designers. Updated SelectMultipleTargets targeting code to permit this mode. There may be an argument for this being the only mode in the future.

diff --git a/Engine/Source/Runtime/GameplayAbilities/Private/GameplayAbilityTargetTypes.cpp b/Engine/Source/Runtime/GameplayAbilities/Private/GameplayAbilityTargetTypes.cpp
--- a/Engine/Source/Runtime/GameplayAbilities/Private/GameplayAbilityTargetTypes.cpp
+++ b/Engine/Source/Runtime/GameplayAbilities/Private/GameplayAbilityTargetTypes.cpp
@@ -78,13 +78,36 @@ FGameplayAbilityTargetDataHandle FGameplayAbilityTargetingLocationInfo::MakeTarg
 	return FGameplayAbilityTargetDataHandle(ReturnData);
 }
 
-FGameplayAbilityTargetDataHandle FGameplayAbilityTargetingLocationInfo::MakeTargetDataHandleFromActors(TArray<TWeakObjectPtr<AActor> > TargetActors) const
+FGameplayAbilityTargetDataHandle FGameplayAbilityTargetingLocationInfo::MakeTargetDataHandleFromActors(TArray<TWeakObjectPtr<AActor> > TargetActors, bool OneActorPerHandle) const
 {
 	/** Note: This is cleaned up by the FGameplayAbilityTargetDataHandle (via an internal TSharedPtr) */
 	FGameplayAbilityTargetData_ActorArray* ReturnData = new FGameplayAbilityTargetData_ActorArray();
-	ReturnData->TargetActorArray = TargetActors;
+	FGameplayAbilityTargetDataHandle ReturnDataHandle = FGameplayAbilityTargetDataHandle(ReturnData);
 	ReturnData->SourceLocation = *this;
-	return FGameplayAbilityTargetDataHandle(ReturnData);
+	if (OneActorPerHandle)
+	{
+		if (TargetActors[0].IsValid())
+		{
+			ReturnData->TargetActorArray.Add(TargetActors[0].Get());
+		}
+
+		for (int32 i = 1; i < TargetActors.Num(); ++i)
+		{
+			if (TargetActors[i].IsValid())
+			{
+				FGameplayAbilityTargetData_ActorArray* CurrentData = new FGameplayAbilityTargetData_ActorArray();
+				FGameplayAbilityTargetDataHandle CurrentDataHandle = FGameplayAbilityTargetDataHandle(CurrentData);
+				CurrentData->SourceLocation = *this;
+				CurrentData->TargetActorArray.Add(TargetActors[i].Get());
+				ReturnDataHandle.Append(&CurrentDataHandle);
+			}
+		}
+	}
+	else
+	{
+		ReturnData->TargetActorArray = TargetActors;
+	}
+	return ReturnDataHandle;
 }
 
 bool FGameplayAbilityTargetDataHandle::NetSerialize(FArchive& Ar, class UPackageMap* Map, bool& bOutSuccess)
diff --git a/Engine/Source/Runtime/GameplayAbilities/Public/Abilities/GameplayAbilityTargetTypes.h b/Engine/Source/Runtime/GameplayAbilities/Public/Abilities/GameplayAbilityTargetTypes.h
--- a/Engine/Source/Runtime/GameplayAbilities/Public/Abilities/GameplayAbilityTargetTypes.h
+++ b/Engine/Source/Runtime/GameplayAbilities/Public/Abilities/GameplayAbilityTargetTypes.h
@@ -315,7 +315,7 @@ public:
 
 	FGameplayAbilityTargetDataHandle MakeTargetDataHandleFromHitResult(TWeakObjectPtr<UGameplayAbility> Ability, FHitResult HitResult) const;
 
-	FGameplayAbilityTargetDataHandle MakeTargetDataHandleFromActors(TArray<TWeakObjectPtr<AActor> > TargetActors) const;
+	FGameplayAbilityTargetDataHandle MakeTargetDataHandleFromActors(TArray<TWeakObjectPtr<AActor> > TargetActors, bool OneActorPerHandle = false) const;
 
 	/** Type of location used - will determine what data is transmitted over the network and what fields are used when calculating position. */
 	UPROPERTY(BlueprintReadWrite, meta = (ExposeOnSpawn = true), Category = Targeting)
@@ -449,7 +449,45 @@ struct GAMEPLAYABILITIES_API FGameplayAbilityTargetData_ActorArray : public FGam
 
 	virtual FTransform GetOrigin() const override
 	{
-		return SourceLocation.GetTargetingTransform();		//No single aim-at target, since this is aimed at multiple characters.
+		FTransform ReturnTransform = SourceLocation.GetTargetingTransform();
+
+		//Aim at first valid target, if we have one. Duplicating GetEndPoint() code here so we don't iterate through the target array twice.
+		for (int32 i = 0; i < TargetActorArray.Num(); ++i)
+		{
+			if (TargetActorArray[i].IsValid())
+			{
+				ReturnTransform.SetRotation((TargetActorArray[i].Get()->GetActorLocation()).SafeNormal().Rotation().Quaternion());
+				break;
+			}
+		}
+		return ReturnTransform;
+	}
+
+	// -------------------------------------
+
+	virtual bool HasEndPoint() const override
+	{
+		//We have an endpoint if we have at least one valid actor in our target array
+		for (int32 i = 0; i < TargetActorArray.Num(); ++i)
+		{
+			if (TargetActorArray[i].IsValid())
+			{
+				return true;
+			}
+		}
+		return false;
+	}
+
+	virtual FVector GetEndPoint() const override
+	{
+		for (int32 i = 0; i < TargetActorArray.Num(); ++i)
+		{
+			if (TargetActorArray[i].IsValid())
+			{
+				return TargetActorArray[i].Get()->GetActorLocation();
+			}
+		}
+		return FVector::ZeroVector;
 	}
 
 	// -------------------------------------
