Added convenience function to abilities in order to reduce node count for designer BPs.

diff --git a/Engine/Source/Runtime/GameplayAbilities/Private/Abilities/GameplayAbility.cpp b/Engine/Source/Runtime/GameplayAbilities/Private/Abilities/GameplayAbility.cpp
--- a/Engine/Source/Runtime/GameplayAbilities/Private/Abilities/GameplayAbility.cpp
+++ b/Engine/Source/Runtime/GameplayAbilities/Private/Abilities/GameplayAbility.cpp
@@ -409,6 +409,13 @@ FGameplayAbilityActorInfo UGameplayAbility::GetActorInfo() const
 	return *CurrentActorInfo;
 }
 
+/** Convenience method for abilities to get owning actor - useful for aiming abilities */
+AActor* UGameplayAbility::GetOwningActorFromActorInfo() const
+{
+	check(CurrentActorInfo);
+	return CurrentActorInfo->Actor.Get();
+}
+
 /** Convenience method for abilities to get skeletal mesh component - useful for aiming abilities */
 USkeletalMeshComponent* UGameplayAbility::GetOwningComponentFromActorInfo() const
 {
diff --git a/Engine/Source/Runtime/GameplayAbilities/Public/Abilities/GameplayAbility.h b/Engine/Source/Runtime/GameplayAbilities/Public/Abilities/GameplayAbility.h
--- a/Engine/Source/Runtime/GameplayAbilities/Public/Abilities/GameplayAbility.h
+++ b/Engine/Source/Runtime/GameplayAbilities/Public/Abilities/GameplayAbility.h
@@ -353,6 +353,9 @@ protected:
 	UFUNCTION(BlueprintCallable, Category=Ability)
 	FGameplayAbilityActorInfo GetActorInfo() const;
 
+	UFUNCTION(BlueprintCallable, Category = Ability)
+	AActor* GetOwningActorFromActorInfo() const;
+
 	UFUNCTION(BlueprintCallable, Category = Ability)
 	USkeletalMeshComponent* GetOwningComponentFromActorInfo() const;
 
