Added missing check for zero-length array.

diff --git a/Engine/Source/Runtime/GameplayAbilities/Private/GameplayAbilityTargetTypes.cpp b/Engine/Source/Runtime/GameplayAbilities/Private/GameplayAbilityTargetTypes.cpp
--- a/Engine/Source/Runtime/GameplayAbilities/Private/GameplayAbilityTargetTypes.cpp
+++ b/Engine/Source/Runtime/GameplayAbilities/Private/GameplayAbilityTargetTypes.cpp
@@ -86,20 +86,23 @@ FGameplayAbilityTargetDataHandle FGameplayAbilityTargetingLocationInfo::MakeTarg
 	ReturnData->SourceLocation = *this;
 	if (OneActorPerHandle)
 	{
-		if (TargetActors[0].IsValid())
+		if (TargetActors.Num() > 0)
 		{
-			ReturnData->TargetActorArray.Add(TargetActors[0].Get());
-		}
+			if (TargetActors[0].IsValid())
+			{
+				ReturnData->TargetActorArray.Add(TargetActors[0].Get());
+			}
 
-		for (int32 i = 1; i < TargetActors.Num(); ++i)
-		{
-			if (TargetActors[i].IsValid())
+			for (int32 i = 1; i < TargetActors.Num(); ++i)
 			{
-				FGameplayAbilityTargetData_ActorArray* CurrentData = new FGameplayAbilityTargetData_ActorArray();
-				FGameplayAbilityTargetDataHandle CurrentDataHandle = FGameplayAbilityTargetDataHandle(CurrentData);
-				CurrentData->SourceLocation = *this;
-				CurrentData->TargetActorArray.Add(TargetActors[i].Get());
-				ReturnDataHandle.Append(&CurrentDataHandle);
+				if (TargetActors[i].IsValid())
+				{
+					FGameplayAbilityTargetData_ActorArray* CurrentData = new FGameplayAbilityTargetData_ActorArray();
+					FGameplayAbilityTargetDataHandle CurrentDataHandle = FGameplayAbilityTargetDataHandle(CurrentData);
+					CurrentData->SourceLocation = *this;
+					CurrentData->TargetActorArray.Add(TargetActors[i].Get());
+					ReturnDataHandle.Append(&CurrentDataHandle);
+				}
 			}
 		}
 	}
