- fixing crash in tower
- removing homing sub class moving data to base projectile
- adding in safety checks to ability system BP library

diff --git a/Engine/Source/Runtime/GameplayAbilities/Private/AbilitySystemBlueprintLibrary.cpp b/Engine/Source/Runtime/GameplayAbilities/Private/AbilitySystemBlueprintLibrary.cpp
--- a/Engine/Source/Runtime/GameplayAbilities/Private/AbilitySystemBlueprintLibrary.cpp
+++ b/Engine/Source/Runtime/GameplayAbilities/Private/AbilitySystemBlueprintLibrary.cpp
@@ -133,38 +133,48 @@ int32 UAbilitySystemBlueprintLibrary::GetDataCountFromTargetData(FGameplayAbilit
 
 TArray<AActor*> UAbilitySystemBlueprintLibrary::GetActorsFromTargetData(FGameplayAbilityTargetDataHandle TargetData, int32 Index)
 {
-	FGameplayAbilityTargetData* Data = TargetData.Data[Index].Get();
-	TArray<AActor*>	ResolvedArray;
-	if (Data)
+	if (TargetData.Data.IsValidIndex(Index))
 	{
-		TArray<TWeakObjectPtr<AActor> > WeakArray = Data->GetActors();
-		for (TWeakObjectPtr<AActor> WeakPtr : WeakArray)
+		FGameplayAbilityTargetData* Data = TargetData.Data[Index].Get();
+		TArray<AActor*>	ResolvedArray;
+		if (Data)
 		{
-			ResolvedArray.Add(WeakPtr.Get());
+			TArray<TWeakObjectPtr<AActor> > WeakArray = Data->GetActors();
+			for (TWeakObjectPtr<AActor> WeakPtr : WeakArray)
+			{
+				ResolvedArray.Add(WeakPtr.Get());
+			}
 		}
+		return ResolvedArray;
 	}
-	return ResolvedArray;
+	return TArray<AActor*>();
 }
 
 bool UAbilitySystemBlueprintLibrary::TargetDataHasHitResult(FGameplayAbilityTargetDataHandle TargetData, int32 Index)
 {
-	FGameplayAbilityTargetData* Data = TargetData.Data[Index].Get();
-	if (Data)
+	if (TargetData.Data.IsValidIndex(Index))
 	{
-		return Data->HasHitResult();
+		FGameplayAbilityTargetData* Data = TargetData.Data[Index].Get();
+		if (Data)
+		{
+			return Data->HasHitResult();
+		}
 	}
 	return false;
 }
 
 FHitResult UAbilitySystemBlueprintLibrary::GetHitResultFromTargetData(FGameplayAbilityTargetDataHandle TargetData, int32 Index)
 {
-	FGameplayAbilityTargetData* Data = TargetData.Data[Index].Get();
-	if (Data)
+	if (TargetData.Data.IsValidIndex(Index))
 	{
-		const FHitResult* HitResultPtr = Data->GetHitResult();
-		if (HitResultPtr)
+		FGameplayAbilityTargetData* Data = TargetData.Data[Index].Get();
+		if (Data)
 		{
-			return *HitResultPtr;
+			const FHitResult* HitResultPtr = Data->GetHitResult();
+			if (HitResultPtr)
+			{
+				return *HitResultPtr;
+			}
 		}
 	}
 
@@ -205,35 +215,41 @@ FTransform UAbilitySystemBlueprintLibrary::GetTargetDataOrigin(FGameplayAbilityT
 
 bool UAbilitySystemBlueprintLibrary::TargetDataHasEndPoint(FGameplayAbilityTargetDataHandle TargetData, int32 Index)
 {
-	FGameplayAbilityTargetData* Data = TargetData.Data[Index].Get();
-	if (Data)
+	if (TargetData.Data.IsValidIndex(Index))
 	{
-		return (Data->HasHitResult() || Data->HasEndPoint());
+		FGameplayAbilityTargetData* Data = TargetData.Data[Index].Get();
+		if (Data)
+		{
+			return (Data->HasHitResult() || Data->HasEndPoint());
+		}
 	}
 	return false;
 }
 
 FVector UAbilitySystemBlueprintLibrary::GetTargetDataEndPoint(FGameplayAbilityTargetDataHandle TargetData, int32 Index)
 {
-	FGameplayAbilityTargetData* Data = TargetData.Data[Index].Get();
-	if (Data)
+	if (TargetData.Data.IsValidIndex(Index))
 	{
-		const FHitResult* HitResultPtr = Data->GetHitResult();
-		if (HitResultPtr)
+		FGameplayAbilityTargetData* Data = TargetData.Data[Index].Get();
+		if (Data)
 		{
-			if (HitResultPtr->Component.IsValid())
+			const FHitResult* HitResultPtr = Data->GetHitResult();
+			if (HitResultPtr)
 			{
-				return HitResultPtr->ImpactPoint;
+				if (HitResultPtr->Component.IsValid())
+				{
+					return HitResultPtr->ImpactPoint;
+				}
+				else
+				{
+					return HitResultPtr->TraceEnd;
+				}
 			}
-			else
+			else if (Data->HasEndPoint())
 			{
-				return HitResultPtr->TraceEnd;
+				return Data->GetEndPoint();
 			}
 		}
-		else if (Data->HasEndPoint())
-		{
-			return Data->GetEndPoint();
-		}
 	}
 
 	return FVector::ZeroVector;
