More operator== constness fixes.

diff --git a/Engine/Source/Runtime/GameplayAbilities/Public/AbilitySystemComponent.h b/Engine/Source/Runtime/GameplayAbilities/Public/AbilitySystemComponent.h
--- a/Engine/Source/Runtime/GameplayAbilities/Public/AbilitySystemComponent.h
+++ b/Engine/Source/Runtime/GameplayAbilities/Public/AbilitySystemComponent.h
@@ -160,7 +160,7 @@ class GAMEPLAYABILITIES_API UAbilitySystemComponent : public UActorComponent, pu
 
 	struct FPendingAbilityInfo
 	{
-		bool operator==(const FPendingAbilityInfo& Other)
+		bool operator==(const FPendingAbilityInfo& Other) const
 		{
 			return PrevPredictionKey == Other.PrevPredictionKey
 				&& CurrPredictionKey == Other.CurrPredictionKey
@@ -187,7 +187,7 @@ class GAMEPLAYABILITIES_API UAbilitySystemComponent : public UActorComponent, pu
 	{
 		FExecutingAbilityInfo() : State(EAbilityExecutionState::Executing) {};
 
-		bool operator==(const FExecutingAbilityInfo& Other)
+		bool operator==(const FExecutingAbilityInfo& Other) const
 		{
 			return CurrPredictionKey == Other.CurrPredictionKey
 				&& State == Other.State;
