GameplayAbilities prediction pass
-Formalize FPredictionKey, add NetSerialize to only replicate key values to the owning client.
-Added concept of staleness to PredictionKeys. When PredictionKey is stale, new predictive actions cannot happen since they are no longer atomic with the creation of the key.
-Fix issue with triggered events: changed concept of 'previous' key to 'base key'.
-Still several issues with triggered events to rework

diff --git a/Engine/Source/Runtime/GameplayAbilities/Private/Abilities/GameplayAbility.cpp b/Engine/Source/Runtime/GameplayAbilities/Private/Abilities/GameplayAbility.cpp
--- a/Engine/Source/Runtime/GameplayAbilities/Private/Abilities/GameplayAbility.cpp
+++ b/Engine/Source/Runtime/GameplayAbilities/Private/Abilities/GameplayAbility.cpp
@@ -80,9 +80,7 @@ void UGameplayAbility::SendGameplayEvent(FGameplayTag EventTag, FGameplayEventDa
 {
 	if (NetExecutionPolicy == EGameplayAbilityNetExecutionPolicy::Predictive)
 	{
-		check(CurrentActivationInfo.CurrPredictionKey != 0);
-		Payload.PrevPredictionKey = CurrentActivationInfo.PrevPredictionKey;
-		Payload.CurrPredictionKey = CurrentActivationInfo.CurrPredictionKey;
+		Payload.PredictionKey = CurrentActivationInfo.GetPredictionKeyForNewAction();
 	}
 
 	AActor *OwnerActor = Cast<AActor>(GetOuter());
@@ -208,7 +206,7 @@ bool UGameplayAbility::CommitAbility(const FGameplayAbilitySpecHandle Handle, co
 		return false;
 	}
 
-	CommitExecute(ActorInfo, ActivationInfo);
+	CommitExecute(Handle, ActorInfo, ActivationInfo);
 
 	// Fixme: Should we always call this or only if it is implemented? A noop may not hurt but could be bad for perf (storing a HasBlueprintCommit per instance isn't good either)
 	K2_CommitExecute();
@@ -242,14 +240,11 @@ bool UGameplayAbility::CommitCheck(const FGameplayAbilityActorInfo* ActorInfo, c
 	return true;
 }
 
-void UGameplayAbility::CommitExecute(const FGameplayAbilityActorInfo* ActorInfo, const FGameplayAbilityActivationInfo ActivationInfo)
+void UGameplayAbility::CommitExecute(const FGameplayAbilitySpecHandle Handle, const FGameplayAbilityActorInfo* ActorInfo, const FGameplayAbilityActivationInfo ActivationInfo)
 {
-	if (!ActorInfo->AbilitySystemComponent->IsNetSimulating() || ActivationInfo.CurrPredictionKey != 0)
-	{
-		ApplyCooldown(ActorInfo, ActivationInfo);
+	ApplyCooldown(Handle, ActorInfo, ActivationInfo);
 
-		ApplyCost(ActorInfo, ActivationInfo);
-	}
+	ApplyCost(Handle, ActorInfo, ActivationInfo);
 }
 
 void UGameplayAbility::CancelAbility(const FGameplayAbilitySpecHandle Handle, const FGameplayAbilityActorInfo* ActorInfo, const FGameplayAbilityActivationInfo ActivationInfo)
@@ -357,12 +352,11 @@ bool UGameplayAbility::CheckCooldown(const FGameplayAbilityActorInfo* ActorInfo)
 	return true;
 }
 
-void UGameplayAbility::ApplyCooldown(const FGameplayAbilityActorInfo* ActorInfo, const FGameplayAbilityActivationInfo ActivationInfo)
+void UGameplayAbility::ApplyCooldown(const FGameplayAbilitySpecHandle Handle, const FGameplayAbilityActorInfo* ActorInfo, const FGameplayAbilityActivationInfo ActivationInfo)
 {
 	if (CooldownGameplayEffect)
 	{
-		check(ActorInfo->AbilitySystemComponent.IsValid());
-		ActorInfo->AbilitySystemComponent->ApplyGameplayEffectToSelf(CooldownGameplayEffect, 1.f, ActorInfo->Actor.Get(), FModifierQualifier().PredictionKeys(ActivationInfo.PrevPredictionKey, ActivationInfo.CurrPredictionKey));
+		ApplyGameplayEffectToOwner(Handle, ActorInfo, ActivationInfo, CooldownGameplayEffect, 1.f);
 	}
 }
 
@@ -376,12 +370,11 @@ bool UGameplayAbility::CheckCost(const FGameplayAbilityActorInfo* ActorInfo) con
 	return true;
 }
 
-void UGameplayAbility::ApplyCost(const FGameplayAbilityActorInfo* ActorInfo, const FGameplayAbilityActivationInfo ActivationInfo)
+void UGameplayAbility::ApplyCost(const FGameplayAbilitySpecHandle Handle, const FGameplayAbilityActorInfo* ActorInfo, const FGameplayAbilityActivationInfo ActivationInfo)
 {
-	check(ActorInfo->AbilitySystemComponent.IsValid());
 	if (CostGameplayEffect)
 	{
-		ActorInfo->AbilitySystemComponent->ApplyGameplayEffectToSelf(CostGameplayEffect, 1.f, ActorInfo->Actor.Get(), FModifierQualifier().PredictionKeys(ActivationInfo.PrevPredictionKey, ActivationInfo.CurrPredictionKey));
+		ApplyGameplayEffectToOwner(Handle, ActorInfo, ActivationInfo, CostGameplayEffect, 1.f);
 	}
 }
 
@@ -558,28 +551,22 @@ void UGameplayAbility::TaskEnded(UAbilityTask* Task)
  *	work with the PredictionKey system.
  */
 
-void UGameplayAbility::ExecuteGameplayCue(FGameplayTag GameplayCueTag)
+void UGameplayAbility::K2_ExecuteGameplayCue(FGameplayTag GameplayCueTag)
 {
-	if (ensure(CurrentActorInfo) && CurrentActorInfo->IsNetAuthority())
-	{
-		CurrentActorInfo->AbilitySystemComponent->ExecuteGameplayCue(GameplayCueTag);
-	}
+	check(CurrentActorInfo);
+	CurrentActorInfo->AbilitySystemComponent->ExecuteGameplayCue(GameplayCueTag, CurrentActivationInfo.GetPredictionKeyForNewAction());
 }
 
-void UGameplayAbility::AddGameplayCue(FGameplayTag GameplayCueTag)
+void UGameplayAbility::K2_AddGameplayCue(FGameplayTag GameplayCueTag)
 {
-	if (ensure(CurrentActorInfo) && CurrentActorInfo->IsNetAuthority())
-	{
-		CurrentActorInfo->AbilitySystemComponent->AddGameplayCue(GameplayCueTag);
-	}
+	check(CurrentActorInfo);
+	CurrentActorInfo->AbilitySystemComponent->AddGameplayCue(GameplayCueTag, CurrentActivationInfo.GetPredictionKeyForNewAction());
 }
 
-void UGameplayAbility::RemoveGameplayCue(FGameplayTag GameplayCueTag)
+void UGameplayAbility::K2_RemoveGameplayCue(FGameplayTag GameplayCueTag)
 {
-	if (ensure(CurrentActorInfo) && CurrentActorInfo->IsNetAuthority())
-	{
-		CurrentActorInfo->AbilitySystemComponent->RemoveGameplayCue(GameplayCueTag);
-	}
+	check(CurrentActorInfo);
+	CurrentActorInfo->AbilitySystemComponent->RemoveGameplayCue(GameplayCueTag, CurrentActivationInfo.GetPredictionKeyForNewAction());
 }
 
 int32 UGameplayAbility::GetAbilityLevel() const
@@ -610,4 +597,50 @@ bool UGameplayAbility::IsTriggered() const
 	// Assume that if there is triggered data, then we are triggered. 
 	// If we need to support abilities that can be both, this will need to be expanded.
 	return AbilityTriggers.Num() > 0;
+}
+
+// -------------------------------------------------------
+
+FActiveGameplayEffectHandle UGameplayAbility::K2_ApplyGameplayEffectToOwner(const UGameplayEffect* GameplayEffect, int32 GameplayEffectLevel)
+{
+	check(CurrentActorInfo);
+	check(CurrentSpecHandle.IsValid());
+
+	//CurrentActivationInfo = GetCurrentAbilitySpec()->ActivationInfo;
+
+	return ApplyGameplayEffectToOwner(CurrentSpecHandle, CurrentActorInfo, CurrentActivationInfo, GameplayEffect, GameplayEffectLevel);
+}
+
+FActiveGameplayEffectHandle UGameplayAbility::ApplyGameplayEffectToOwner(const FGameplayAbilitySpecHandle Handle, const FGameplayAbilityActorInfo* ActorInfo, const FGameplayAbilityActivationInfo ActivationInfo, const UGameplayEffect* GameplayEffect, int32 GameplayEffectLevel)
+{
+	if (ActivationInfo.ActivationMode == EGameplayAbilityActivationMode::Authority || ActivationInfo.ActivationMode == EGameplayAbilityActivationMode::Predicting)
+	{
+		return ActorInfo->AbilitySystemComponent->ApplyGameplayEffectToSelf(GameplayEffect, 1.f, ActorInfo->Actor.Get(), FModifierQualifier().PredictionKey(ActivationInfo.GetPredictionKeyForNewAction()));
+	}
+
+	// We cannot apply GameplayEffects in this context. Return an empty handle.
+	return FActiveGameplayEffectHandle();
+}
+
+FActiveGameplayEffectHandle UGameplayAbility::K2_ApplyGameplayEffectToTarget(FGameplayAbilityTargetDataHandle Target, const UGameplayEffect* GameplayEffect, int32 GameplayEffectLevel)
+{
+	return ApplyGameplayEffectToTarget(CurrentSpecHandle, CurrentActorInfo, CurrentActivationInfo, Target, GameplayEffect, GameplayEffectLevel);
+}
+
+FActiveGameplayEffectHandle UGameplayAbility::ApplyGameplayEffectToTarget(const FGameplayAbilitySpecHandle Handle, const FGameplayAbilityActorInfo* ActorInfo, const FGameplayAbilityActivationInfo ActivationInfo, FGameplayAbilityTargetDataHandle Target, const UGameplayEffect* GameplayEffect, int32 GameplayEffectLevel)
+{
+	FActiveGameplayEffectHandle EffectHandle;
+	for (auto Data : Target.Data)
+	{
+		TArray<FActiveGameplayEffectHandle> EffectHandles = Data->ApplyGameplayEffect(GameplayEffect, ActorInfo, (float)GameplayEffectLevel);
+		if (EffectHandles.Num() > 0)
+		{
+			EffectHandle = EffectHandles[0];
+			if (EffectHandles.Num() > 1)
+			{
+				ABILITY_LOG(Warning, TEXT("ApplyGameplayEffectToTargetData_Single called on TargetData with multiple actor targets. %s"), *GetName());
+			}
+		}
+	}
+	return EffectHandle;
 }
\ No newline at end of file
diff --git a/Engine/Source/Runtime/GameplayAbilities/Private/Abilities/GameplayAbility_CharacterJump.cpp b/Engine/Source/Runtime/GameplayAbilities/Private/Abilities/GameplayAbility_CharacterJump.cpp
--- a/Engine/Source/Runtime/GameplayAbilities/Private/Abilities/GameplayAbility_CharacterJump.cpp
+++ b/Engine/Source/Runtime/GameplayAbilities/Private/Abilities/GameplayAbility_CharacterJump.cpp
@@ -59,11 +59,7 @@ bool UGameplayAbility_CharacterJump::CanActivateAbility(const FGameplayAbilitySp
 void UGameplayAbility_CharacterJump::CancelAbility(const FGameplayAbilitySpecHandle Handle, const FGameplayAbilityActorInfo* ActorInfo, const FGameplayAbilityActivationInfo ActivationInfo)
 {
 	Super::CancelAbility(Handle, ActorInfo, ActivationInfo);
-
-	if (ActivationInfo.ActivationMode == EGameplayAbilityActivationMode::Authority ||
-		ActivationInfo.ActivationMode == EGameplayAbilityActivationMode::Predicting)
-	{
-		ACharacter * Character = CastChecked<ACharacter>(ActorInfo->Actor.Get());
-		Character->StopJumping();
-	}
+	
+	ACharacter * Character = CastChecked<ACharacter>(ActorInfo->Actor.Get());
+	Character->StopJumping();
 }
\ No newline at end of file
diff --git a/Engine/Source/Runtime/GameplayAbilities/Private/Abilities/Tasks/AbilityTask_WaitGameplayEffectRemoved.cpp b/Engine/Source/Runtime/GameplayAbilities/Private/Abilities/Tasks/AbilityTask_WaitGameplayEffectRemoved.cpp
--- a/Engine/Source/Runtime/GameplayAbilities/Private/Abilities/Tasks/AbilityTask_WaitGameplayEffectRemoved.cpp
+++ b/Engine/Source/Runtime/GameplayAbilities/Private/Abilities/Tasks/AbilityTask_WaitGameplayEffectRemoved.cpp
@@ -22,6 +22,13 @@ UAbilityTask_WaitGameplayEffectRemoved* UAbilityTask_WaitGameplayEffectRemoved::
 
 void UAbilityTask_WaitGameplayEffectRemoved::Activate()
 {
+	if (Handle.IsValid() == false)
+	{
+		InvalidHandle.Broadcast();
+		EndTask();
+		return;;
+	}
+
 	UAbilitySystemComponent* EffectOwningAbilitySystemComponent = Handle.GetOwningAbilitySystemComponent();
 
 	if (EffectOwningAbilitySystemComponent)
diff --git a/Engine/Source/Runtime/GameplayAbilities/Private/AbilitySystemBlueprintLibrary.cpp b/Engine/Source/Runtime/GameplayAbilities/Private/AbilitySystemBlueprintLibrary.cpp
--- a/Engine/Source/Runtime/GameplayAbilities/Private/AbilitySystemBlueprintLibrary.cpp
+++ b/Engine/Source/Runtime/GameplayAbilities/Private/AbilitySystemBlueprintLibrary.cpp
@@ -20,59 +20,6 @@ UAbilitySystemComponent* UAbilitySystemBlueprintLibrary::GetAbilitySystemCompone
 	return UAbilitySystemGlobals::GetAbilitySystemComponentFromActor(Actor);
 }
 
-TArray<FActiveGameplayEffectHandle> UAbilitySystemBlueprintLibrary::ApplyGameplayEffectToTargetData(FGameplayAbilityTargetDataHandle Target, UGameplayEffect* GameplayEffect, const FGameplayAbilityActorInfo InstigatorInfo, int32 Level)
-{
-	TArray<FActiveGameplayEffectHandle> Handles;
-
-	if (GameplayEffect == nullptr)
-	{
-		return Handles;
-	}
-
-	for (auto Data : Target.Data)
-	{
-		if (Data.IsValid())
-		{
-			Handles.Append(Data->ApplyGameplayEffect(GameplayEffect, InstigatorInfo, (float)Level));
-		}
-	}
-
-	return Handles;
-}
-
-FActiveGameplayEffectHandle UAbilitySystemBlueprintLibrary::ApplyGameplayEffectToTargetData_Single(FGameplayAbilityTargetDataHandle Target, UGameplayEffect* GameplayEffect, const FGameplayAbilityActorInfo InstigatorInfo, int32 Level)
-{
-	FActiveGameplayEffectHandle Handle;
-
-	if (GameplayEffect == nullptr)
-	{
-		return Handle;
-	}
-
-	if (Target.Data.Num() != 1)
-	{
-		ABILITY_LOG(Warning, TEXT("ApplyGameplayEffectToTargetData_Single called on invalid TargetDataHandle."));
-		return Handle;
-	}
-
-	
-	if (Target.Data[0].IsValid())
-	{
-		TArray<FActiveGameplayEffectHandle> Handles = Target.Data[0]->ApplyGameplayEffect(GameplayEffect, InstigatorInfo, (float)Level);
-		int32 Num = Handles.Num();
-		if (Num > 0)
-		{
-			Handle = Handles[0];
-			if (Num > 1)
-			{
-				ABILITY_LOG(Warning, TEXT("ApplyGameplayEffectToTargetData_Single called on TargetData with multiple actor targets"));
-			}
-		}
-	}
-
-	return Handle;
-}
-
 FGameplayAbilityTargetDataHandle UAbilitySystemBlueprintLibrary::AppendTargetDataHandle(FGameplayAbilityTargetDataHandle TargetHandle, FGameplayAbilityTargetDataHandle HandleToAdd)
 {
 	TargetHandle.Append(&HandleToAdd);
diff --git a/Engine/Source/Runtime/GameplayAbilities/Private/AbilitySystemComponent.cpp b/Engine/Source/Runtime/GameplayAbilities/Private/AbilitySystemComponent.cpp
--- a/Engine/Source/Runtime/GameplayAbilities/Private/AbilitySystemComponent.cpp
+++ b/Engine/Source/Runtime/GameplayAbilities/Private/AbilitySystemComponent.cpp
@@ -42,8 +42,7 @@ UAbilitySystemComponent::UAbilitySystemComponent(const class FPostConstructIniti
 	ActiveGameplayCues.Owner = this;
 
 	SetIsReplicated(true);
-	ReplicatedPredictionKey = 0;
-	LocalPredictionKey = 0;
+
 	UserAbilityActivationInhibited = false;
 }
 
@@ -155,6 +154,11 @@ bool UAbilitySystemComponent::IsOwnerActorAuthoritative() const
 	return !IsNetSimulating();
 }
 
+bool UAbilitySystemComponent::HasNetworkAuthorityToApplyGameplayEffect(const FModifierQualifier QualifierContext) const
+{
+	return (IsOwnerActorAuthoritative() || QualifierContext.PredictionKey().IsStale()==false);
+}
+
 void UAbilitySystemComponent::SetNumericAttribute(const FGameplayAttribute &Attribute, float NewFloatValue)
 {
 	const UAttributeSet* AttributeSet = GetAttributeSubobjectChecked(Attribute.GetAttributeSetClass());
@@ -193,10 +197,13 @@ FGameplayEffectInstigatorContext UAbilitySystemComponent::GetInstigatorContext()
 FActiveGameplayEffectHandle UAbilitySystemComponent::ApplyGameplayEffectToTarget(UGameplayEffect *GameplayEffect, UAbilitySystemComponent *Target, float Level, FModifierQualifier BaseQualifier)
 {
 	check(GameplayEffect);
+	if (HasNetworkAuthorityToApplyGameplayEffect(BaseQualifier))
+	{
+		FGameplayEffectSpec	Spec(GameplayEffect, AbilityActorInfo->Actor.Get(), Level, GetCurveDataOverride());
+		return ApplyGameplayEffectSpecToTarget(Spec, Target, BaseQualifier);
+	}
 
-	FGameplayEffectSpec	Spec(GameplayEffect, AbilityActorInfo->Actor.Get(), Level, GetCurveDataOverride());
-	
-	return ApplyGameplayEffectSpecToTarget(Spec, Target, BaseQualifier);
+	return FActiveGameplayEffectHandle();
 }
 
 /** Helper function since we can't have default/optional values for FModifierQualifier in K2 function */
@@ -210,9 +217,13 @@ FActiveGameplayEffectHandle UAbilitySystemComponent::ApplyGameplayEffectToSelf(c
 {
 	check(GameplayEffect);
 
-	FGameplayEffectSpec	Spec(GameplayEffect, Instigator, Level, GetCurveDataOverride());
+	if (HasNetworkAuthorityToApplyGameplayEffect(BaseQualifier))
+	{
+		FGameplayEffectSpec	Spec(GameplayEffect, Instigator, Level, GetCurveDataOverride());
+		return ApplyGameplayEffectSpecToSelf(Spec, BaseQualifier);
+	}
 
-	return ApplyGameplayEffectSpecToSelf(Spec, BaseQualifier);
+	return FActiveGameplayEffectHandle();
 }
 
 float UAbilitySystemComponent::GetGameplayEffectMagnitudeByTag(FActiveGameplayEffectHandle InHandle, const FGameplayTag& InTag) const
@@ -289,11 +300,14 @@ void UAbilitySystemComponent::TEMP_ApplyActiveGameplayEffects()
 
 FActiveGameplayEffectHandle UAbilitySystemComponent::ApplyGameplayEffectSpecToTarget(OUT FGameplayEffectSpec &Spec, UAbilitySystemComponent *Target, FModifierQualifier BaseQualifier)
 {
-	// Apply outgoing Effects to the Spec.
-	// Outgoing immunity may stop the outgoing effect from being applied to the target
-	if (ActiveGameplayEffects.ApplyActiveEffectsTo(Spec, FModifierQualifier(BaseQualifier).Type(EGameplayMod::OutgoingGE)))
+	if (HasNetworkAuthorityToApplyGameplayEffect(BaseQualifier))
 	{
-		return Target->ApplyGameplayEffectSpecToSelf(Spec, BaseQualifier);
+		// Apply outgoing Effects to the Spec.
+		// Outgoing immunity may stop the outgoing effect from being applied to the target
+		if (ActiveGameplayEffects.ApplyActiveEffectsTo(Spec, FModifierQualifier(BaseQualifier).Type(EGameplayMod::OutgoingGE)))
+		{
+			return Target->ApplyGameplayEffectSpecToSelf(Spec, BaseQualifier);
+		}
 	}
 
 	return FActiveGameplayEffectHandle();
@@ -303,10 +317,15 @@ FActiveGameplayEffectHandle UAbilitySystemComponent::ApplyGameplayEffectSpecToSe
 {
 	// Temp, only non instant, non periodic GEs can be predictive
 	// Effects with other effects may be a mix so go with non-predictive
-	check((BaseQualifier.CurrentPredictionKey() == 0) || (Spec.GetPeriod() == UGameplayEffect::NO_PERIOD));
+	check((BaseQualifier.PredictionKey().IsValidKey() == false) || (Spec.GetPeriod() == UGameplayEffect::NO_PERIOD));
+
+	if (!HasNetworkAuthorityToApplyGameplayEffect(BaseQualifier))
+	{
+		return FActiveGameplayEffectHandle();
+	}
 
 	// Clients should treat predicted instant effects as if they have infinite duration. The effects will be cleaned up later.
-	bool bTreatAsInfiniteDuration = GetOwnerRole() != ROLE_Authority && BaseQualifier.CurrentPredictionKey() != 0 && Spec.GetDuration() == UGameplayEffect::INSTANT_APPLICATION;
+	bool bTreatAsInfiniteDuration = GetOwnerRole() != ROLE_Authority && BaseQualifier.PredictionKey().IsValidKey() && Spec.GetDuration() == UGameplayEffect::INSTANT_APPLICATION;
 
 	// check if the effect being applied actually succeeds
 	float ChanceToApply = Spec.GetChanceToApplyToTarget();
@@ -317,7 +336,7 @@ FActiveGameplayEffectHandle UAbilitySystemComponent::ApplyGameplayEffectSpecToSe
 
 	// Make sure we create our copy of the spec in the right place first...
 	FActiveGameplayEffectHandle	MyHandle;
-	bool bInvokeGameplayCueApplied = false;
+	bool bInvokeGameplayCueApplied = UGameplayEffect::INFINITE_DURATION == Spec.GetDuration();
 
 	FGameplayEffectSpec* OurCopyOfSpec = NULL;
 	TSharedPtr<FGameplayEffectSpec> StackSpec;
@@ -330,11 +349,9 @@ FActiveGameplayEffectHandle UAbilitySystemComponent::ApplyGameplayEffectSpecToSe
 			{
 				ActiveGameplayEffects.StacksNeedToRecalculate();
 			}
-			FActiveGameplayEffect &NewActiveEffect = ActiveGameplayEffects.CreateNewActiveGameplayEffect(Spec, BaseQualifier.PreviousPredictionKey(), BaseQualifier.CurrentPredictionKey());
+			FActiveGameplayEffect &NewActiveEffect = ActiveGameplayEffects.CreateNewActiveGameplayEffect(Spec, BaseQualifier.PredictionKey());
 			MyHandle = NewActiveEffect.Handle;
 			OurCopyOfSpec = &NewActiveEffect.Spec;
-
-			bInvokeGameplayCueApplied = true;
 		}
 		
 		if (!OurCopyOfSpec)
@@ -404,6 +421,13 @@ FActiveGameplayEffectHandle UAbilitySystemComponent::ApplyGameplayEffectSpecToSe
 	{
 		ExecuteGameplayEffect(*OurCopyOfSpec, FModifierQualifier(BaseQualifier).IgnoreHandle(MyHandle));
 	}
+	else if (bTreatAsInfiniteDuration)
+	{
+		// This is an instant application but we are treating it as an infinite duration for prediction. We should still predict the execute GameplayCUE.
+		// (in non predictive case, this will happen inside ::ExecuteGameplayEffect)
+		InvokeGameplayCueEvent(*OurCopyOfSpec, EGameplayCueEvent::Executed);
+	}
+
 
 	if (Spec.GetPeriod() != UGameplayEffect::NO_PERIOD && Spec.TargetEffectSpecs.Num() > 0)
 	{
@@ -486,26 +510,26 @@ void UAbilitySystemComponent::InvokeGameplayCueEvent(const FGameplayEffectSpec &
 	}
 }
 
-void UAbilitySystemComponent::ExecuteGameplayCue(const FGameplayTag GameplayCueTag)
+void UAbilitySystemComponent::ExecuteGameplayCue(const FGameplayTag GameplayCueTag, FPredictionKey PredictionKey)
 {
 	if (IsOwnerActorAuthoritative())
 	{
-		NetMulticast_InvokeGameplayCueExecuted(GameplayCueTag);
+		NetMulticast_InvokeGameplayCueExecuted(GameplayCueTag, PredictionKey);
 	}
-	else
+	else if (PredictionKey.IsValidKey())
 	{
 		InvokeGameplayCueEvent(GameplayCueTag, EGameplayCueEvent::Executed);
 	}
 }
 
-void UAbilitySystemComponent::AddGameplayCue(const FGameplayTag GameplayCueTag)
+void UAbilitySystemComponent::AddGameplayCue(const FGameplayTag GameplayCueTag, FPredictionKey PredictionKey)
 {
 	if (IsOwnerActorAuthoritative())
 	{
 		ActiveGameplayCues.AddCue(GameplayCueTag);
-		NetMulticast_InvokeGameplayCueAdded(GameplayCueTag);
+		NetMulticast_InvokeGameplayCueAdded(GameplayCueTag, PredictionKey);
 	}
-	else
+	else if (PredictionKey.IsValidKey())
 	{
 		// Allow for predictive gameplaycue events? Needs more thought
 		InvokeGameplayCueEvent(GameplayCueTag, EGameplayCueEvent::OnActive);
@@ -513,13 +537,16 @@ void UAbilitySystemComponent::AddGameplayCue(const FGameplayTag GameplayCueTag)
 	}
 }
 
-void UAbilitySystemComponent::RemoveGameplayCue(const FGameplayTag GameplayCueTag)
+void UAbilitySystemComponent::RemoveGameplayCue(const FGameplayTag GameplayCueTag, FPredictionKey PredictionKey)
 {
-	InvokeGameplayCueEvent(GameplayCueTag, EGameplayCueEvent::Removed);
 	if (IsOwnerActorAuthoritative())
 	{
 		ActiveGameplayCues.RemoveCue(GameplayCueTag);
 	}
+	else if (PredictionKey.IsValidKey())
+	{
+		InvokeGameplayCueEvent(GameplayCueTag, EGameplayCueEvent::Removed);
+	}
 }
 
 void UAbilitySystemComponent::InvokeGameplayCueEvent(const FGameplayTag GameplayCueTag, EGameplayCueEvent::Type EventType)
@@ -539,19 +566,28 @@ void UAbilitySystemComponent::InvokeGameplayCueEvent(const FGameplayTag Gameplay
 	GameplayCueInterface->HandleGameplayCue(ActorOwner, GameplayCueTag, EventType, CueParameters);
 }
 
-void UAbilitySystemComponent::NetMulticast_InvokeGameplayCueExecuted_FromSpec_Implementation(const FGameplayEffectSpec Spec)
+void UAbilitySystemComponent::NetMulticast_InvokeGameplayCueExecuted_FromSpec_Implementation(const FGameplayEffectSpec Spec, FPredictionKey PredictionKey)
 {
-	InvokeGameplayCueEvent(Spec, EGameplayCueEvent::Executed);
+	if (PredictionKey.IsValidKey() == false)
+	{
+		InvokeGameplayCueEvent(Spec, EGameplayCueEvent::Executed);
+	}
 }
 
-void UAbilitySystemComponent::NetMulticast_InvokeGameplayCueExecuted_Implementation(const FGameplayTag GameplayCueTag)
+void UAbilitySystemComponent::NetMulticast_InvokeGameplayCueExecuted_Implementation(const FGameplayTag GameplayCueTag, FPredictionKey PredictionKey)
 {
-	InvokeGameplayCueEvent(GameplayCueTag, EGameplayCueEvent::Executed);
+	if (PredictionKey.IsValidKey() == false)
+	{
+		InvokeGameplayCueEvent(GameplayCueTag, EGameplayCueEvent::Executed);
+	}
 }
 
-void UAbilitySystemComponent::NetMulticast_InvokeGameplayCueAdded_Implementation(const FGameplayTag GameplayCueTag)
+void UAbilitySystemComponent::NetMulticast_InvokeGameplayCueAdded_Implementation(const FGameplayTag GameplayCueTag, FPredictionKey PredictionKey)
 {
-	InvokeGameplayCueEvent(GameplayCueTag, EGameplayCueEvent::OnActive);
+	if (PredictionKey.IsValidKey() == false)
+	{
+		InvokeGameplayCueEvent(GameplayCueTag, EGameplayCueEvent::OnActive);
+	}
 }
 
 bool UAbilitySystemComponent::IsGameplayCueActive(const FGameplayTag GameplayCueTag) const
@@ -665,7 +701,7 @@ void UAbilitySystemComponent::OnRep_PredictionKey()
 	int32 idx = 0;
 	for ( ; idx < PredictionDelegates.Num(); ++idx)
 	{
-		if (PredictionDelegates[idx].Key <= ReplicatedPredictionKey)
+		if (PredictionDelegates[idx].Key <= ReplicatedPredictionKey.Current)
 		{
 			PredictionDelegates[idx].Value.PredictionKeyClearDelegate.Broadcast();
 		}
@@ -681,7 +717,7 @@ void UAbilitySystemComponent::OnRep_PredictionKey()
 	}
 }
 
-UAbilitySystemComponent::FPredictionInfo& UAbilitySystemComponent::GetPredictionKeyDelegate(uint32 PredictionKey)
+UAbilitySystemComponent::FPredictionInfo& UAbilitySystemComponent::GetPredictionKeyDelegate(FPredictionKey::KeyType PredictionKey)
 {
 	// See if we already have one for this key
 	for (auto & Item : PredictionDelegates)
@@ -693,19 +729,13 @@ UAbilitySystemComponent::FPredictionInfo& UAbilitySystemComponent::GetPrediction
 	}
 
 	// Create a new one
-	TPair<uint32, FPredictionInfo> NewPair;
+	TPair<FPredictionKey::KeyType, FPredictionInfo> NewPair;
 	NewPair.Key = PredictionKey;
 	
 	PredictionDelegates.Add(NewPair);
 	return PredictionDelegates.Last().Value;
 }
 
-uint32 UAbilitySystemComponent::GetNextPredictionKey()
-{
-	LocalPredictionKey = FMath::Max(LocalPredictionKey, ReplicatedPredictionKey);
-	return ++LocalPredictionKey;
-}
-
 // ---------------------------------------------------------------------------------------
 
 void UAbilitySystemComponent::PrintAllGameplayEffects() const
diff --git a/Engine/Source/Runtime/GameplayAbilities/Private/AbilitySystemComponent_Abilities.cpp b/Engine/Source/Runtime/GameplayAbilities/Private/AbilitySystemComponent_Abilities.cpp
--- a/Engine/Source/Runtime/GameplayAbilities/Private/AbilitySystemComponent_Abilities.cpp
+++ b/Engine/Source/Runtime/GameplayAbilities/Private/AbilitySystemComponent_Abilities.cpp
@@ -312,7 +312,7 @@ void UAbilitySystemComponent::OnRep_ActivateAbilities()
  *	-This function handles networking and prediction
  *	-If all goes well, CallActivateAbility is called next.
  */
-bool UAbilitySystemComponent::TryActivateAbility(FGameplayAbilitySpecHandle Handle, uint32 PrevPredictionKey, uint32 CurrPredictionKey, UGameplayAbility** OutInstancedAbility )
+bool UAbilitySystemComponent::TryActivateAbility(FGameplayAbilitySpecHandle Handle, FPredictionKey InPredictionKey, UGameplayAbility** OutInstancedAbility)
 {
 	FGameplayAbilitySpec* Spec = FindAbilitySpecFromHandle(Handle);
 	if (!Spec)
@@ -347,26 +347,27 @@ bool UAbilitySystemComponent::TryActivateAbility(FGameplayAbilitySpecHandle Hand
 		return false;
 	}
 
-	if (Spec->IsActive)
+	if (Spec->IsActive && Ability->GetInstancingPolicy() != EGameplayAbilityInstancingPolicy::NonInstanced)
 	{
+		// Don't have a great solution for ending non instanced abilities on the server. So only warn if this was instanced.
 		ABILITY_LOG(Warning, TEXT("TryActivateAbility called when ability was already active. NetMode: %d. Ability: %s"), (int32)NetMode, *Ability->GetName());
 	}
-
-	//check(!Spec->IsActive);	// Shouldn't be here if the ability is already active.
+	
 	Spec->IsActive = true;
 
 	// Setup a fresh ActivationInfo for this AbilitySpec.
-	Spec->ActivationInfo = FGameplayAbilityActivationInfo(ActorInfo->Actor.Get(), PrevPredictionKey, CurrPredictionKey);
+	Spec->ActivationInfo = FGameplayAbilityActivationInfo(ActorInfo->Actor.Get(), InPredictionKey);
 	FGameplayAbilityActivationInfo &ActivationInfo = Spec->ActivationInfo;
 
+	UGameplayAbility* InstancedAbility = Ability;
+
 	// If we are the server or this is a client authorative 
 	if (Ability->GetNetExecutionPolicy() == EGameplayAbilityNetExecutionPolicy::Client || (NetMode == ROLE_Authority))
 	{
 		// Create instance of this ability if necessary
-
 		if (Ability->GetInstancingPolicy() == EGameplayAbilityInstancingPolicy::InstancedPerExecution)
 		{
-			UGameplayAbility* InstancedAbility = CreateNewInstanceOfAbility(*Spec, Ability);
+			InstancedAbility = CreateNewInstanceOfAbility(*Spec, Ability);
 			InstancedAbility->CallActivateAbility(Handle, ActorInfo, ActivationInfo);
 			if (OutInstancedAbility)
 			{
@@ -380,15 +381,15 @@ bool UAbilitySystemComponent::TryActivateAbility(FGameplayAbilitySpecHandle Hand
 	}
 	else if (Ability->GetNetExecutionPolicy() == EGameplayAbilityNetExecutionPolicy::Server)
 	{
-		ServerTryActivateAbility(Handle, ActivationInfo.PrevPredictionKey, ActivationInfo.CurrPredictionKey);
+		ServerTryActivateAbility(Handle, FPredictionKey());
 	}
 	else if (Ability->GetNetExecutionPolicy() == EGameplayAbilityNetExecutionPolicy::Predictive)
 	{
 		// This execution is now officially EGameplayAbilityActivationMode:Predicting and has a PredictionKey
-		ActivationInfo.GeneratePredictionKey(this);
+		ActivationInfo.GenerateNewPredictionKey();
 
 		// This must be called immediately after GeneratePredictionKey to prevent problems with recursively activating abilities
-		ServerTryActivateAbility(Handle, ActivationInfo.PrevPredictionKey, ActivationInfo.CurrPredictionKey);
+		ServerTryActivateAbility(Handle, ActivationInfo.GetPredictionKeyForNewAction());
 
 		if (Ability->GetInstancingPolicy() == EGameplayAbilityInstancingPolicy::InstancedPerExecution)
 		{
@@ -398,7 +399,7 @@ bool UAbilitySystemComponent::TryActivateAbility(FGameplayAbilitySpecHandle Hand
 
 			if (Ability->GetReplicationPolicy() == EGameplayAbilityReplicationPolicy::ReplicateNone)
 			{
-				UGameplayAbility* InstancedAbility = CreateNewInstanceOfAbility(*Spec, Ability);
+				InstancedAbility = CreateNewInstanceOfAbility(*Spec, Ability);
 				InstancedAbility->CallActivateAbility(Handle, ActorInfo, ActivationInfo);
 				if (OutInstancedAbility)
 				{
@@ -416,16 +417,22 @@ bool UAbilitySystemComponent::TryActivateAbility(FGameplayAbilitySpecHandle Hand
 		}
 	}
 
+	// Anything after this can no longer be predicted with the prediction key we just generated.
+	// The key we generated and sent in ServerTryActivateAbility will be confirmed/denied after the server run CallActivateAbility.
+	// If game code tries to do something predictively after this (latent call, response to input, etc) it will not be able to use this key.
+	ActivationInfo.SetPredictionStale();
+	InstancedAbility->SetCurrentActivationInfo(ActivationInfo);	// Need to push this to the ability if it was instanced.
+
 	return true;
 }
 
-void UAbilitySystemComponent::ServerTryActivateAbility_Implementation(FGameplayAbilitySpecHandle Handle, uint32 PrevPredictionKey, uint32 CurrPredictionKey)
+void UAbilitySystemComponent::ServerTryActivateAbility_Implementation(FGameplayAbilitySpecHandle Handle, FPredictionKey PredictionKey)
 {
 #if !(UE_BUILD_SHIPPING || UE_BUILD_TEST)
 	if (DenyClientActivation > 0)
 	{
 		DenyClientActivation--;
-		ClientActivateAbilityFailed(Handle, CurrPredictionKey);
+		ClientActivateAbilityFailed(Handle, PredictionKey.Current);
 		return;
 	}
 #endif
@@ -443,12 +450,13 @@ void UAbilitySystemComponent::ServerTryActivateAbility_Implementation(FGameplayA
 	ensure(AbilityActorInfo.IsValid());
 
 	// if this was triggered by a predicted ability the server triggered copy should be the canonical one
-	if (PrevPredictionKey != 0)
+	if (PredictionKey.Base != 0)
 	{
 		// first check if the server has already started this ability
-		for (auto ExecutingAbilityInfo : ExecutingServerAbilities)
+		for (int32 ExecutedIdx = 0; ExecutedIdx < ExecutingServerAbilities.Num(); ++ExecutedIdx)
 		{
-			if (ExecutingAbilityInfo.CurrPredictionKey == PrevPredictionKey && ExecutingAbilityInfo.Handle == Handle)
+			FExecutingAbilityInfo& ExecutingAbilityInfo = ExecutingServerAbilities[ExecutedIdx];
+			if (ExecutingAbilityInfo.PredictionKey.Current == PredictionKey.Base && ExecutingAbilityInfo.Handle == Handle)
 			{
 				switch(ExecutingAbilityInfo.State)
 				{
@@ -457,27 +465,26 @@ void UAbilitySystemComponent::ServerTryActivateAbility_Implementation(FGameplayA
 						FGameplayAbilitySpec* Spec = FindAbilitySpecFromHandle(ExecutingAbilityInfo.Handle);
 						if (Spec)
 						{
-							Spec->ActivationInfo.PrevPredictionKey = PrevPredictionKey;
-							Spec->ActivationInfo.CurrPredictionKey = CurrPredictionKey;
+							// This feels wrong? We are changing the PredictionKey state on another ability?
+							Spec->ActivationInfo.SetPredictionKey(PredictionKey);
 						}
 						break;
 					}
 					case EAbilityExecutionState::Failed:
-						ClientActivateAbilityFailed(Handle, CurrPredictionKey);
+						ClientActivateAbilityFailed(Handle, PredictionKey.Current);
 						break;
 					case EAbilityExecutionState::Succeeded:
-						ClientActivateAbilitySucceed(Handle, CurrPredictionKey);
+						ClientActivateAbilitySucceed(Handle, PredictionKey.Current);
 						break;
 				}
 
-				ExecutingServerAbilities.RemoveSingleSwap(ExecutingAbilityInfo);
+				ExecutingServerAbilities.RemoveAtSwap(ExecutedIdx);
 				return;
 			}
 		}
 
 		FPendingAbilityInfo AbilityInfo;
-		AbilityInfo.PrevPredictionKey = PrevPredictionKey;
-		AbilityInfo.CurrPredictionKey = CurrPredictionKey;
+		AbilityInfo.PredictionKey = PredictionKey;
 		AbilityInfo.Handle = Handle;
 
 		PendingClientAbilities.Add(AbilityInfo);
@@ -488,38 +495,38 @@ void UAbilitySystemComponent::ServerTryActivateAbility_Implementation(FGameplayA
 	Spec->InputPressed = true; // Pretend input was pressed. Allows UAbilityTask_WaitInputRelease to work.
 
 	// Attempt to activate the ability (server side) and tell the client if it succeeded or failed.
-	if (TryActivateAbility(Handle, PrevPredictionKey, CurrPredictionKey, &InstancedAbility))
+	if (TryActivateAbility(Handle, PredictionKey, &InstancedAbility))
 	{
-		ClientActivateAbilitySucceed(Handle, CurrPredictionKey);
+		ClientActivateAbilitySucceed(Handle, PredictionKey.Current);
 
 		// Update our ReplicatedPredictionKey. When the client gets value, he will know his state (actor+all components/subobjects) are up to do date and he can
 		// remove any necessary predictive work.
 
-		if (CurrPredictionKey > 0)
+		if (PredictionKey.IsValidKey())
 		{
-			ensure(CurrPredictionKey > ReplicatedPredictionKey);
-			ReplicatedPredictionKey = CurrPredictionKey;
+			ensure(PredictionKey.Current > ReplicatedPredictionKey.Current);
+			ReplicatedPredictionKey = PredictionKey;
 		}
 	}
 	else
 	{
 		Spec->InputPressed = false;
-		ClientActivateAbilityFailed(Handle, CurrPredictionKey);
+		ClientActivateAbilityFailed(Handle, PredictionKey.Current);
 	}
 }
 
-bool UAbilitySystemComponent::ServerTryActivateAbility_Validate(FGameplayAbilitySpecHandle AbilityToActivate, uint32 PrevPredictionKey, uint32 CurrPredictionKey)
+bool UAbilitySystemComponent::ServerTryActivateAbility_Validate(FGameplayAbilitySpecHandle AbilityToActivate, FPredictionKey PredictionKey)
 {
 	return true;
 }
 
-void UAbilitySystemComponent::ClientActivateAbilityFailed_Implementation(FGameplayAbilitySpecHandle Handle, uint32 PredictionKey)
+void UAbilitySystemComponent::ClientActivateAbilityFailed_Implementation(FGameplayAbilitySpecHandle Handle, int16 PredictionKey)
 {
 	if (PredictionKey > 0)
 	{
 		for (int32 idx = 0; idx < PredictionDelegates.Num(); ++idx)
 		{
-			uint32 Key = PredictionDelegates[idx].Key;
+			FPredictionKey::KeyType Key = PredictionDelegates[idx].Key;
 			if (Key == PredictionKey)
 			{
 				ABILITY_LOG(Warning, TEXT("Failed ActivateAbility, clearing prediction data %d"), PredictionKey);
@@ -530,7 +537,7 @@ void UAbilitySystemComponent::ClientActivateAbilityFailed_Implementation(FGamepl
 				if (Spec)
 				{
 					UGameplayAbility* AbilityToActivate = Spec->Ability;
-					for (uint32 DependentKey : PredictionDelegates[idx].Value.DependentPredictionKeys)
+					for (FPredictionKey::KeyType DependentKey : PredictionDelegates[idx].Value.DependentPredictionKeys)
 					{
 						ClientActivateAbilityFailed(Handle, DependentKey);
 					}
@@ -546,7 +553,7 @@ void UAbilitySystemComponent::ClientActivateAbilityFailed_Implementation(FGamepl
 	}
 }
 
-void UAbilitySystemComponent::ClientActivateAbilitySucceed_Implementation(FGameplayAbilitySpecHandle Handle, uint32 PredictionKey)
+void UAbilitySystemComponent::ClientActivateAbilitySucceed_Implementation(FGameplayAbilitySpecHandle Handle, int16 PredictionKey)
 {
 	FGameplayAbilitySpec* Spec = FindAbilitySpecFromHandle(Handle);
 	if (!Spec)
@@ -577,7 +584,7 @@ void UAbilitySystemComponent::ClientActivateAbilitySucceed_Implementation(FGamep
 			bool found = false;
 			for (UGameplayAbility* LocalAbility : Spec->NonReplicatedInstances)				// Fixme: this has to be updated once predictive abilities can replicate
 			{
-				if (LocalAbility->GetCurrentActivationInfo().CurrPredictionKey == PredictionKey)
+				if (LocalAbility->GetCurrentActivationInfo().GetPredictionKey().Current == PredictionKey)
 				{
 					LocalAbility->ConfirmActivateSucceed();
 					found = true;
@@ -621,24 +628,22 @@ void UAbilitySystemComponent::TriggerAbilityFromGameplayEvent(FGameplayAbilitySp
 	{
 		return;
 	}
-	
 
 	if (Ability->ShouldAbilityRespondToEvent(EventTag, Payload))
 	{
 		int32 ExecutingAbilityIndex = -1;
 
 		// if we're the server and this is coming from a predicted event we should check if the client has already predicted it
-		if (Payload->CurrPredictionKey
+		if (Payload->PredictionKey.Current > 0
 			&& Ability->GetNetExecutionPolicy() == EGameplayAbilityNetExecutionPolicy::Predictive
 			&& ActorInfo->Actor->Role == ROLE_Authority)
 		{
 			bool bPendingClientAbilityFound = false;
 			for (auto PendingAbilityInfo : Component.PendingClientAbilities)
 			{
-				if (Payload->CurrPredictionKey == PendingAbilityInfo.PrevPredictionKey && Handle == PendingAbilityInfo.Handle) // found a match
+				if (Payload->PredictionKey.Current == PendingAbilityInfo.PredictionKey.Base && Handle == PendingAbilityInfo.Handle) // found a match
 				{
-					Payload->PrevPredictionKey = PendingAbilityInfo.PrevPredictionKey;
-					Payload->CurrPredictionKey = PendingAbilityInfo.CurrPredictionKey;
+					Payload->PredictionKey = PendingAbilityInfo.PredictionKey;
 
 					Component.PendingClientAbilities.RemoveSingleSwap(PendingAbilityInfo);
 					bPendingClientAbilityFound = true;
@@ -651,14 +656,14 @@ void UAbilitySystemComponent::TriggerAbilityFromGameplayEvent(FGameplayAbilitySp
 			if (bPendingClientAbilityFound == false)
 			{
 				UAbilitySystemComponent::FExecutingAbilityInfo Info;
-				Info.CurrPredictionKey = Payload->CurrPredictionKey;
+				Info.PredictionKey = Payload->PredictionKey;
 				Info.Handle = Handle;
 
 				ExecutingAbilityIndex = Component.ExecutingServerAbilities.Add(Info);
 			}
 		}
 
-		if (TryActivateAbility(Handle, Payload->PrevPredictionKey, Payload->CurrPredictionKey))
+		if (TryActivateAbility(Handle, Payload->PredictionKey))
 		{
 			if (ExecutingAbilityIndex >= 0)
 			{
diff --git a/Engine/Source/Runtime/GameplayAbilities/Private/AbilitySystemDebugHUD.cpp b/Engine/Source/Runtime/GameplayAbilities/Private/AbilitySystemDebugHUD.cpp
--- a/Engine/Source/Runtime/GameplayAbilities/Private/AbilitySystemDebugHUD.cpp
+++ b/Engine/Source/Runtime/GameplayAbilities/Private/AbilitySystemDebugHUD.cpp
@@ -114,7 +114,7 @@ void AAbilitySystemDebugHUD::DrawDebugAbilitySystemComponent(UAbilitySystemCompo
 		// Draw Active GameplayEffect
 		for (FActiveGameplayEffect &Effect : Component->ActiveGameplayEffects.GameplayEffects)
 		{
-			String = FString::Printf(TEXT("%s. [%d, %d] %.2f"), *Effect.Spec.ToSimpleString(), Effect.PrevPredictionKey, Effect.CurrPredictionKey, Effect.GetTimeRemaining(GameWorldTime));
+			String = FString::Printf(TEXT("%s. [%d, %d] %.2f"), *Effect.Spec.ToSimpleString(), Effect.PredictionKey.Current, Effect.PredictionKey.Base, Effect.GetTimeRemaining(GameWorldTime));
 			DrawWithBackground(Font, String, Color, EAlignHorizontal::Left, X, EAlignVertical::Top, Y);	
 		}
 	}
diff --git a/Engine/Source/Runtime/GameplayAbilities/Private/GameplayAbilityTargetTypes.cpp b/Engine/Source/Runtime/GameplayAbilities/Private/GameplayAbilityTargetTypes.cpp
--- a/Engine/Source/Runtime/GameplayAbilities/Private/GameplayAbilityTargetTypes.cpp
+++ b/Engine/Source/Runtime/GameplayAbilities/Private/GameplayAbilityTargetTypes.cpp
@@ -7,13 +7,13 @@
 #include "AbilitySystemBlueprintLibrary.h"
 
 
-TArray<FActiveGameplayEffectHandle> FGameplayAbilityTargetData::ApplyGameplayEffect(UGameplayEffect* GameplayEffect, const FGameplayAbilityActorInfo InstigatorInfo, float Level)
+TArray<FActiveGameplayEffectHandle> FGameplayAbilityTargetData::ApplyGameplayEffect(const UGameplayEffect* GameplayEffect, const FGameplayAbilityActorInfo* InstigatorInfo, float Level)
 {
 	// TODO: Improve relationship between InstigatorContext and FGameplayAbilityTargetData/FHitResult (or use something different between HitResult)
 
 
 	FGameplayEffectSpec	SpecToApply(GameplayEffect,					// The UGameplayEffect data asset
-		InstigatorInfo.Actor.Get(),		// The actor who instigated this
+		InstigatorInfo->Actor.Get(),		// The actor who instigated this
 		Level,							// FIXME: Leveling
 		NULL							// FIXME: CurveData override... should we just remove this?
 		);
@@ -38,7 +38,7 @@ TArray<FActiveGameplayEffectHandle> FGameplayAbilityTargetData::ApplyGameplayEff
 			UAbilitySystemComponent* TargetComponent = UAbilitySystemBlueprintLibrary::GetAbilitySystemComponent(TargetActor.Get());
 			if (TargetComponent)
 			{
-				AppliedHandles.Add( InstigatorInfo.AbilitySystemComponent->ApplyGameplayEffectSpecToTarget(SpecToApply, TargetComponent) );
+				AppliedHandles.Add( InstigatorInfo->AbilitySystemComponent->ApplyGameplayEffectSpecToTarget(SpecToApply, TargetComponent) );
 			}
 		}
 	}
@@ -57,18 +57,21 @@ FGameplayAbilityTargetDataHandle FGameplayAbilityTargetingLocationInfo::MakeTarg
 	if (LocationType == EGameplayAbilityTargetingLocationType::Type::SocketTransform)
 	{
 		const FGameplayAbilityActorInfo* ActorInfo = Ability.IsValid() ? Ability.Get()->GetCurrentActorInfo() : NULL;
-		AActor* AISourceActor = ActorInfo ? (ActorInfo->Actor.IsValid() ? ActorInfo->Actor.Get() : NULL) : NULL;
-		UAnimInstance* AnimInstance = ActorInfo ? ActorInfo->AnimInstance.Get() : NULL;
-		USkeletalMeshComponent* AISourceComponent = AnimInstance ? AnimInstance->GetOwningComponent() : NULL;
-
-		if (AISourceActor && AISourceComponent)
+		if (ActorInfo)
 		{
-			FGameplayAbilityTargetData_Mesh* ReturnData = new FGameplayAbilityTargetData_Mesh();
-			ReturnData->SourceActor = AISourceActor;
-			ReturnData->SourceComponent = AISourceComponent;
-			ReturnData->SourceSocketName = SourceSocketName;
-			ReturnData->TargetPoint = HitResult.Location;
-			return FGameplayAbilityTargetDataHandle(ReturnData);
+			AActor* AISourceActor = ActorInfo->Actor.Get();
+			UAnimInstance* AnimInstance = ActorInfo->AnimInstance.Get();
+			USkeletalMeshComponent* AISourceComponent = AnimInstance ? AnimInstance->GetOwningComponent() : NULL;
+
+			if (AISourceActor && AISourceComponent)
+			{
+				FGameplayAbilityTargetData_Mesh* ReturnData = new FGameplayAbilityTargetData_Mesh();
+				ReturnData->SourceActor = AISourceActor;
+				ReturnData->SourceComponent = AISourceComponent;
+				ReturnData->SourceSocketName = SourceSocketName;
+				ReturnData->TargetPoint = HitResult.Location;
+				return FGameplayAbilityTargetDataHandle(ReturnData);
+			}
 		}
 	}
 
diff --git a/Engine/Source/Runtime/GameplayAbilities/Private/GameplayAbilityTypes.cpp b/Engine/Source/Runtime/GameplayAbilities/Private/GameplayAbilityTypes.cpp
--- a/Engine/Source/Runtime/GameplayAbilities/Private/GameplayAbilityTypes.cpp
+++ b/Engine/Source/Runtime/GameplayAbilities/Private/GameplayAbilityTypes.cpp
@@ -77,13 +77,9 @@ bool FGameplayAbilityActorInfo::IsNetAuthority() const
 	return false;
 }
 
-void FGameplayAbilityActivationInfo::GeneratePredictionKey(UAbilitySystemComponent * Component) const
-{
-	check(Component);
-	check(ActivationMode == EGameplayAbilityActivationMode::NonAuthority);
-
-	PrevPredictionKey = CurrPredictionKey;
-	CurrPredictionKey = Component->GetNextPredictionKey();
+void FGameplayAbilityActivationInfo::GenerateNewPredictionKey() const
+{	
+	PredictionKey.GenerateDependantPredictionKey();
 	ActivationMode = EGameplayAbilityActivationMode::Predicting;
 }
 
@@ -92,5 +88,7 @@ void FGameplayAbilityActivationInfo::SetActivationConfirmed()
 	ActivationMode = EGameplayAbilityActivationMode::Confirmed;
 }
 
-
-
+void FGameplayAbilityActivationInfo::SetPredictionStale()
+{
+	PredictionKey.bIsStale = true;
+}
\ No newline at end of file
diff --git a/Engine/Source/Runtime/GameplayAbilities/Private/GameplayEffect.cpp b/Engine/Source/Runtime/GameplayAbilities/Private/GameplayEffect.cpp
--- a/Engine/Source/Runtime/GameplayAbilities/Private/GameplayEffect.cpp
+++ b/Engine/Source/Runtime/GameplayAbilities/Private/GameplayEffect.cpp
@@ -1146,24 +1146,36 @@ void FActiveGameplayEffect::PreReplicatedRemove(const struct FActiveGameplayEffe
 
 void FActiveGameplayEffect::PostReplicatedAdd(const struct FActiveGameplayEffectsContainer &InArray)
 {
+	bool ShouldInvokeGameplayCueEvents = true;
+	if (PredictionKey.IsValidKey())
+	{
+		// PredictionKey will only be valid on the client that predicted it. So if this has a valid PredictionKey, we can assume we already predicted it and shouldn't invoke GameplayCues.
+		// We may need to do more bookkeeping here in the future. Possibly give the predicted gameplayeffect a chance to pass something off to the new replicated gameplay effect.
+
+		ShouldInvokeGameplayCueEvents = false;
+	}
+
 	const_cast<FActiveGameplayEffectsContainer&>(InArray).InternalOnActiveGameplayEffectAdded(*this);	// Const cast is ok. It is there to prevent mutation of the GameplayEffects array, which this wont do.
 
 	static const int32 MAX_DELTA_TIME = 3;
 
-	InArray.Owner->InvokeGameplayCueEvent(Spec, EGameplayCueEvent::WhileActive);
+	if (ShouldInvokeGameplayCueEvents)
+	{
+		InArray.Owner->InvokeGameplayCueEvent(Spec, EGameplayCueEvent::WhileActive);
+	}
+
 	// Was this actually just activated, or are we just finding out about it due to relevancy/join in progress?
 	float WorldTimeSeconds = InArray.GetWorldTime();
 	int32 GameStateTime = InArray.GetGameStateTime();
 
 	int32 DeltaGameStateTime = GameStateTime - StartGameStateTime;	// How long we think the effect has been playing (only 1 second precision!)
 
-	if (GameStateTime > 0 && FMath::Abs(DeltaGameStateTime) < MAX_DELTA_TIME)
+	if (ShouldInvokeGameplayCueEvents && GameStateTime > 0 && FMath::Abs(DeltaGameStateTime) < MAX_DELTA_TIME)
 	{
 		InArray.Owner->InvokeGameplayCueEvent(Spec, EGameplayCueEvent::OnActive);
 	}
 
 	// Set our local start time accordingly
-
 	StartWorldTime = WorldTimeSeconds - static_cast<float>(DeltaGameStateTime);
 }
 
@@ -1387,7 +1399,7 @@ void FActiveGameplayEffectsContainer::ExecuteActiveEffectsFrom(const FGameplayEf
 		// TODO: check replication policy. Right now we will replicate every execute via a multicast RPC
 
 		ABILITY_LOG(Log, TEXT("Invoking Execute GameplayCue for %s"), *Spec.ToSimpleString() );
-		Owner->NetMulticast_InvokeGameplayCueExecuted_FromSpec(Spec);
+		Owner->NetMulticast_InvokeGameplayCueExecuted_FromSpec(Spec, QualifierContext.PredictionKey());
 	}
 
 }
@@ -1570,12 +1582,12 @@ void FActiveGameplayEffectsContainer::StacksNeedToRecalculate()
 	}
 }
 
-FActiveGameplayEffect & FActiveGameplayEffectsContainer::CreateNewActiveGameplayEffect(const FGameplayEffectSpec &Spec, uint32 InPrevPredictionKey, uint32 InCurrPredictionKey)
+FActiveGameplayEffect & FActiveGameplayEffectsContainer::CreateNewActiveGameplayEffect(const FGameplayEffectSpec &Spec, FPredictionKey InPredictionKey)
 {
 	SCOPE_CYCLE_COUNTER(STAT_CreateNewActiveGameplayEffect);
 
 	FActiveGameplayEffectHandle NewHandle = FActiveGameplayEffectHandle::GenerateNewHandle(Owner);
-	FActiveGameplayEffect & NewEffect = *new (GameplayEffects)FActiveGameplayEffect(NewHandle, Spec, GetWorldTime(), GetGameStateTime(), InPrevPredictionKey, InCurrPredictionKey);
+	FActiveGameplayEffect & NewEffect = *new (GameplayEffects)FActiveGameplayEffect(NewHandle, Spec, GetWorldTime(), GetGameStateTime(), InPredictionKey);
 
 	// register callbacks with the timer manager
 	if (Owner)
@@ -1619,7 +1631,7 @@ FActiveGameplayEffect & FActiveGameplayEffectsContainer::CreateNewActiveGameplay
 		NewEffect.Spec.Duration.Get()->OnDirty = FAggregator::FOnDirty::CreateRaw(this, &FActiveGameplayEffectsContainer::OnDurationAggregatorDirty, Owner, NewHandle);
 	}
 	
-	if (InCurrPredictionKey == 0 || IsNetAuthority())	// Clients predicting a GameplayEffect must not call MarkItemDirty
+	if (InPredictionKey.IsValidKey() == false || IsNetAuthority())	// Clients predicting a GameplayEffect must not call MarkItemDirty
 	{
 		MarkItemDirty(NewEffect);
 	}
@@ -1628,12 +1640,12 @@ FActiveGameplayEffect & FActiveGameplayEffectsContainer::CreateNewActiveGameplay
 		// Clients predicting should call MarkArrayDirty to force the internal replication map to be rebuilt.
 		MarkArrayDirty();
 
-		UAbilitySystemComponent::FPredictionInfo PredictionInfo = Owner->GetPredictionKeyDelegate(InCurrPredictionKey);
+		UAbilitySystemComponent::FPredictionInfo& PredictionInfo = Owner->GetPredictionKeyDelegate(InPredictionKey.Current);
 		PredictionInfo.PredictionKeyClearDelegate.AddUObject(Owner, &UAbilitySystemComponent::RemoveActiveGameplayEffect_NoReturn, NewEffect.Handle);
 
-		if (InPrevPredictionKey)
+		if (InPredictionKey.Base) // Fixme: this should happen at a higher level
 		{
-			PredictionInfo.DependentPredictionKeys.AddUnique(InCurrPredictionKey);
+			PredictionInfo.DependentPredictionKeys.AddUnique(InPredictionKey.Current);
 		}
 	}
 
@@ -1685,7 +1697,24 @@ bool FActiveGameplayEffectsContainer::InternalRemoveActiveGameplayEffect(int32 I
 	if (ensure(Idx < GameplayEffects.Num()))
 	{
 		FActiveGameplayEffect& Effect = GameplayEffects[Idx];
-		Owner->InvokeGameplayCueEvent(GameplayEffects[Idx].Spec, EGameplayCueEvent::Removed);
+
+		bool ShouldInvokeGameplayCueEvent = true;
+		if (Effect.PredictionKey.IsValidKey() && !IsNetAuthority())
+		{
+			// Don't invoke GameplayCue event if we still have another GameplayEffect that shares the same predictionkey
+			for (const FActiveGameplayEffect& OtherEffect : GameplayEffects)
+			{
+				if (OtherEffect.PredictionKey == Effect.PredictionKey && (&OtherEffect != &Effect))
+				{
+					ShouldInvokeGameplayCueEvent = false;
+				}
+			}
+		}
+
+		if (ShouldInvokeGameplayCueEvent)
+		{
+			Owner->InvokeGameplayCueEvent(GameplayEffects[Idx].Spec, EGameplayCueEvent::Removed);
+		}
 
 		InternalOnActiveGameplayEffectRemoved(Effect);
 
diff --git a/Engine/Source/Runtime/GameplayAbilities/Private/GameplayPrediction.cpp b/Engine/Source/Runtime/GameplayAbilities/Private/GameplayPrediction.cpp
new file mode 100644
--- /dev/null
+++ b/Engine/Source/Runtime/GameplayAbilities/Private/GameplayPrediction.cpp
@@ -0,0 +1,63 @@
+// Copyright 1998-2014 Epic Games, Inc. All Rights Reserved.
+
+#include "AbilitySystemPrivatePCH.h"
+#include "GameplayPrediction.h"
+
+bool FPredictionKey::NetSerialize(FArchive& Ar, class UPackageMap* Map, bool& bOutSuccess)
+{
+	if (Ar.IsLoading())
+	{
+		Ar << Current;
+		if (Current > 0)
+		{
+			Ar << Base;
+		}
+		PredictiveConnection = Map;
+	}
+	else
+	{
+		/**
+		 *	Only serialize the payload if we have no owning connection (Client sending to server).
+		 *	or if the owning connection is this connection (Server only sends the prediction key to the client who gave it to us).
+		 */	
+		if (PredictiveConnection == nullptr || (Map == PredictiveConnection))
+		{
+			Ar << Current;
+			if (Current > 0)
+			{
+				Ar << Base;
+			}
+		}
+		else
+		{
+			int16 Payload = 0;
+			Ar << Payload;
+		}
+	}
+
+	bOutSuccess = true;
+	return true;
+}
+
+void FPredictionKey::GenerateNewPredictionKey()
+{
+	static int32 GKey = 1;
+	Current = GKey++;
+	bIsStale = false;
+}
+
+void FPredictionKey::GenerateDependantPredictionKey()
+{
+	if (Base == 0)
+	{
+		Base = Current;
+	}
+	GenerateNewPredictionKey();
+}
+
+FPredictionKey FPredictionKey::CreateNewPredictionKey()
+{
+	FPredictionKey NewKey;
+	NewKey.GenerateNewPredictionKey();
+	return NewKey;
+}
\ No newline at end of file
diff --git a/Engine/Source/Runtime/GameplayAbilities/Public/Abilities/GameplayAbility.h b/Engine/Source/Runtime/GameplayAbilities/Public/Abilities/GameplayAbility.h
--- a/Engine/Source/Runtime/GameplayAbilities/Public/Abilities/GameplayAbility.h
+++ b/Engine/Source/Runtime/GameplayAbilities/Public/Abilities/GameplayAbility.h
@@ -270,7 +270,7 @@ protected:
 	virtual void K2_CommitExecute();
 
 	/** Does the commmit atomically (consume resources, do cooldowns, etc) */
-	virtual void CommitExecute(const FGameplayAbilityActorInfo* ActorInfo, const FGameplayAbilityActivationInfo ActivationInfo);
+	virtual void CommitExecute(const FGameplayAbilitySpecHandle Handle, const FGameplayAbilityActorInfo* ActorInfo, const FGameplayAbilityActivationInfo ActivationInfo);
 
 	/** Do boilerplate init stuff and then call ActivateAbility */
 	void PreActivate(const FGameplayAbilitySpecHandle Handle, const FGameplayAbilityActorInfo* ActorInfo, const FGameplayAbilityActivationInfo ActivationInfo);
@@ -308,19 +308,36 @@ protected:
 
 	virtual void EndAbility(const FGameplayAbilitySpecHandle Handle, const FGameplayAbilityActorInfo* ActorInfo, const FGameplayAbilityActivationInfo ActivationInfo);
 
+	// -------------------------------------
+	//	GameplayEffects
+	//	
+	// -------------------------------------
+
+	UFUNCTION(BlueprintCallable, Category = Ability, FriendlyName="ApplyGameplayEffectToOwner")
+	FActiveGameplayEffectHandle K2_ApplyGameplayEffectToOwner(const UGameplayEffect* GameplayEffect, int32 GameplayEffectLevel = 1);
+
+	/** Non blueprintcallable, safe to call on CDO/NonInstance abilities */
+	FActiveGameplayEffectHandle ApplyGameplayEffectToOwner(const FGameplayAbilitySpecHandle Handle, const FGameplayAbilityActorInfo* ActorInfo, const FGameplayAbilityActivationInfo ActivationInfo, const UGameplayEffect* GameplayEffect, int32 GameplayEffectLevel);
+
+	UFUNCTION(BlueprintCallable, Category = Ability, FriendlyName = "ApplyGameplayEffectToTarget")
+	FActiveGameplayEffectHandle K2_ApplyGameplayEffectToTarget(FGameplayAbilityTargetDataHandle Target, const UGameplayEffect* GameplayEffect, int32 GameplayEffectLevel = 1);
+
+	/** Non blueprintcallable, safe to call on CDO/NonInstance abilities */
+	FActiveGameplayEffectHandle ApplyGameplayEffectToTarget(const FGameplayAbilitySpecHandle Handle, const FGameplayAbilityActorInfo* ActorInfo, const FGameplayAbilityActivationInfo ActivationInfo, FGameplayAbilityTargetDataHandle Target, const UGameplayEffect* GameplayEffect, int32 GameplayEffectLevel);
+
 	// -------------------------------------
 	//	GameplayCue
 	//	Abilities can invoke GameplayCues without having to create GameplayEffects
 	// -------------------------------------
 	
-	UFUNCTION(BlueprintCallable, Category = Ability, meta=(GameplayTagFilter="GameplayCue"))
-	virtual void ExecuteGameplayCue(FGameplayTag GameplayCueTag);
+	UFUNCTION(BlueprintCallable, Category = Ability, meta=(GameplayTagFilter="GameplayCue"), FriendlyName="ExecuteGameplayCue")
+	virtual void K2_ExecuteGameplayCue(FGameplayTag GameplayCueTag);
 
-	UFUNCTION(BlueprintCallable, Category = Ability, meta=(GameplayTagFilter="GameplayCue"))
-	virtual void AddGameplayCue(FGameplayTag GameplayCueTag);
+	UFUNCTION(BlueprintCallable, Category = Ability, meta=(GameplayTagFilter="GameplayCue"), FriendlyName="AddGameplayCue")
+	virtual void K2_AddGameplayCue(FGameplayTag GameplayCueTag);
 
-	UFUNCTION(BlueprintCallable, Category = Ability, meta=(GameplayTagFilter="GameplayCue"))
-	virtual void RemoveGameplayCue(FGameplayTag GameplayCueTag);
+	UFUNCTION(BlueprintCallable, Category = Ability, meta=(GameplayTagFilter="GameplayCue"), FriendlyName="RemoveGameplayCue")
+	virtual void K2_RemoveGameplayCue(FGameplayTag GameplayCueTag);
 
 
 	// -------------------------------------
@@ -340,6 +357,14 @@ protected:
 		}
 	}
 
+	void SetCurrentActivationInfo(const FGameplayAbilityActivationInfo ActivationInfo)
+	{
+		if (IsInstantiated())
+		{
+			CurrentActivationInfo = ActivationInfo;
+		}
+	}
+
 	void SetCurrentInfo(const FGameplayAbilitySpecHandle Handle, const FGameplayAbilityActorInfo* ActorInfo, const FGameplayAbilityActivationInfo ActivationInfo)
 	{
 		if (IsInstantiated())
@@ -368,11 +393,11 @@ protected:
 	bool	CheckCooldown(const FGameplayAbilityActorInfo* ActorInfo) const;
 
 	/** Applies CooldownGameplayEffect to the target */
-	void	ApplyCooldown(const FGameplayAbilityActorInfo* ActorInfo, const FGameplayAbilityActivationInfo ActivationInfo);
+	void	ApplyCooldown(const FGameplayAbilitySpecHandle Handle, const FGameplayAbilityActorInfo* ActorInfo, const FGameplayAbilityActivationInfo ActivationInfo);
 
 	bool	CheckCost(const FGameplayAbilityActorInfo* ActorInfo) const;
 
-	void	ApplyCost(const FGameplayAbilityActorInfo* ActorInfo, const FGameplayAbilityActivationInfo ActivationInfo);
+	void	ApplyCost(const FGameplayAbilitySpecHandle Handle, const FGameplayAbilityActorInfo* ActorInfo, const FGameplayAbilityActivationInfo ActivationInfo);
 
 	// -----------------------------------------------	
 
diff --git a/Engine/Source/Runtime/GameplayAbilities/Public/Abilities/GameplayAbilityTargetTypes.h b/Engine/Source/Runtime/GameplayAbilities/Public/Abilities/GameplayAbilityTargetTypes.h
--- a/Engine/Source/Runtime/GameplayAbilities/Public/Abilities/GameplayAbilityTargetTypes.h
+++ b/Engine/Source/Runtime/GameplayAbilities/Public/Abilities/GameplayAbilityTargetTypes.h
@@ -63,7 +63,7 @@ struct GAMEPLAYABILITIES_API FGameplayAbilityTargetData
 
 	virtual ~FGameplayAbilityTargetData() { }
 
-	TArray<FActiveGameplayEffectHandle> ApplyGameplayEffect(UGameplayEffect *GameplayEffect, const struct FGameplayAbilityActorInfo InstigatorInfo, float Level);
+	TArray<FActiveGameplayEffectHandle> ApplyGameplayEffect(const UGameplayEffect* GameplayEffect, const struct FGameplayAbilityActorInfo* InstigatorInfo, float Level);
 
 	virtual TArray<TWeakObjectPtr<AActor> >	GetActors() const
 	{
diff --git a/Engine/Source/Runtime/GameplayAbilities/Public/Abilities/GameplayAbilityTypes.h b/Engine/Source/Runtime/GameplayAbilities/Public/Abilities/GameplayAbilityTypes.h
--- a/Engine/Source/Runtime/GameplayAbilities/Public/Abilities/GameplayAbilityTypes.h
+++ b/Engine/Source/Runtime/GameplayAbilities/Public/Abilities/GameplayAbilityTypes.h
@@ -3,6 +3,7 @@
 #pragma once
 
 #include "GameplayEffect.h"
+#include "GameplayPrediction.h"
 #include "GameplayAbilityTypes.generated.h"
 
 class UGameplayEffect;
@@ -35,7 +36,7 @@ UENUM(BlueprintType)
 namespace EGameplayAbilityNetExecutionPolicy
 {
 	/**
-	 *	How does an ability execute on the network. Does a client "ask and predict", "ask and wait", "don't ask"
+	 *	How does an ability execute on the network. Does a client "ask and predict", "ask and wait", "don't ask (just do it)"
 	 */
 
 	enum Type
@@ -172,40 +173,53 @@ struct GAMEPLAYABILITIES_API FGameplayAbilityActivationInfo
 
 	FGameplayAbilityActivationInfo()
 		: ActivationMode(EGameplayAbilityActivationMode::Authority)
-		, PrevPredictionKey(0)
-		, CurrPredictionKey(0)
 	{
 
 	}
 
-	FGameplayAbilityActivationInfo(AActor * InActor, uint32 InPrevPredictionKey, uint32 InCurrPredictionKey)
-		: PrevPredictionKey(InPrevPredictionKey)
-		, CurrPredictionKey(InCurrPredictionKey)
+	FGameplayAbilityActivationInfo(AActor * InActor, FPredictionKey InPredictionKey)
+		: PredictionKey(InPredictionKey)
 	{
 		// On Init, we are either Authority or NonAuthority. We haven't been given a PredictionKey and we haven't been confirmed.
 		// NonAuthority essentially means 'I'm not sure what how I'm going to do this yet'.
 		ActivationMode = (InActor->Role == ROLE_Authority ? EGameplayAbilityActivationMode::Authority : EGameplayAbilityActivationMode::NonAuthority);
 	}
 
-	FGameplayAbilityActivationInfo(EGameplayAbilityActivationMode::Type InType, uint32 InPrevPredictionKey = 0, uint32 InCurrPredictionKey = 0)
+	FGameplayAbilityActivationInfo(EGameplayAbilityActivationMode::Type InType, FPredictionKey InPredictionKey = FPredictionKey())
 		: ActivationMode(InType)
-		, PrevPredictionKey(InPrevPredictionKey)
-		, CurrPredictionKey(InCurrPredictionKey)
+		, PredictionKey(InPredictionKey)
 	{
 	}
 
-	void GeneratePredictionKey(UAbilitySystemComponent * Component) const;
+	void GenerateNewPredictionKey() const;
 
 	void SetActivationConfirmed();
 
+	void SetPredictionStale();
+
+	FPredictionKey GetPredictionKeyForNewAction() const
+	{
+		return PredictionKey.IsStale() ? FPredictionKey() : PredictionKey;
+	}
+
+	FPredictionKey GetPredictionKey() const
+	{
+		return PredictionKey;
+	}
+
+	void SetPredictionKey(FPredictionKey NewKey) const
+	{
+		PredictionKey = NewKey;
+	}
+
 	UPROPERTY(BlueprintReadOnly, Category = "ActorInfo")
 	mutable TEnumAsByte<EGameplayAbilityActivationMode::Type>	ActivationMode;
 
-	UPROPERTY()
-	mutable uint32 PrevPredictionKey;
+private:
 
 	UPROPERTY()
-	mutable uint32 CurrPredictionKey;
+	mutable FPredictionKey	PredictionKey;
+	
 };
 
 // ----------------------------------------------------
@@ -282,8 +296,6 @@ struct GAMEPLAYABILITIES_API FGameplayEventData
 	FGameplayEventData()
 	: Instigator(NULL)
 	, Target(NULL)
-	, PrevPredictionKey(0)
-	, CurrPredictionKey(0)
 	{
 		// do nothing
 	}
@@ -299,8 +311,7 @@ struct GAMEPLAYABILITIES_API FGameplayEventData
 	UPROPERTY(EditAnywhere, BlueprintReadWrite, Category = GameplayAbilityTriggerPayload)
 	float Var3;
 
-	uint32 PrevPredictionKey;
-	uint32 CurrPredictionKey;
+	FPredictionKey PredictionKey;
 };
 
 /** 
diff --git a/Engine/Source/Runtime/GameplayAbilities/Public/Abilities/Tasks/AbilityTask_WaitGameplayEffectRemoved.h b/Engine/Source/Runtime/GameplayAbilities/Public/Abilities/Tasks/AbilityTask_WaitGameplayEffectRemoved.h
--- a/Engine/Source/Runtime/GameplayAbilities/Public/Abilities/Tasks/AbilityTask_WaitGameplayEffectRemoved.h
+++ b/Engine/Source/Runtime/GameplayAbilities/Public/Abilities/Tasks/AbilityTask_WaitGameplayEffectRemoved.h
@@ -22,6 +22,9 @@ class UAbilityTask_WaitGameplayEffectRemoved : public UAbilityTask
 	UPROPERTY(BlueprintAssignable)
 	FWaitGameplayEffectRemovedDelegate	OnRemoved;
 
+	UPROPERTY(BlueprintAssignable)
+	FWaitGameplayEffectRemovedDelegate	InvalidHandle;
+
 	virtual void Activate() override;
 
 	UFUNCTION()
diff --git a/Engine/Source/Runtime/GameplayAbilities/Public/AbilitySystemBlueprintLibrary.h b/Engine/Source/Runtime/GameplayAbilities/Public/AbilitySystemBlueprintLibrary.h
--- a/Engine/Source/Runtime/GameplayAbilities/Public/AbilitySystemBlueprintLibrary.h
+++ b/Engine/Source/Runtime/GameplayAbilities/Public/AbilitySystemBlueprintLibrary.h
@@ -25,13 +25,6 @@ class UAbilitySystemBlueprintLibrary : public UBlueprintFunctionLibrary
 	//		TargetData
 	// -------------------------------------------------------------------------------
 
-	UFUNCTION(BlueprintCallable, Category = "Ability|TargetData")
-	static TArray<FActiveGameplayEffectHandle> ApplyGameplayEffectToTargetData(FGameplayAbilityTargetDataHandle Target, UGameplayEffect *GameplayEffect, const FGameplayAbilityActorInfo InstigatorInfo, int32 Level=1);
-
-	/** Applies GameplayEffect to TargetData and returns the handle of the  */
-	UFUNCTION(BlueprintCallable, Category = "Ability|TargetData")
-	static FActiveGameplayEffectHandle ApplyGameplayEffectToTargetData_Single(FGameplayAbilityTargetDataHandle Target, UGameplayEffect *GameplayEffect, const FGameplayAbilityActorInfo InstigatorInfo, int32 Level=1);
-
 	UFUNCTION(BlueprintCallable, Category = "Ability|TargetData")
 	static FGameplayAbilityTargetDataHandle AppendTargetDataHandle(FGameplayAbilityTargetDataHandle TargetHandle, FGameplayAbilityTargetDataHandle HandleToAdd);
 
diff --git a/Engine/Source/Runtime/GameplayAbilities/Public/AbilitySystemComponent.h b/Engine/Source/Runtime/GameplayAbilities/Public/AbilitySystemComponent.h
--- a/Engine/Source/Runtime/GameplayAbilities/Public/AbilitySystemComponent.h
+++ b/Engine/Source/Runtime/GameplayAbilities/Public/AbilitySystemComponent.h
@@ -138,12 +138,7 @@ class GAMEPLAYABILITIES_API UAbilitySystemComponent : public UActorComponent, pu
 	 */
 
 	UPROPERTY(ReplicatedUsing=OnRep_PredictionKey)
-	uint32	ReplicatedPredictionKey;
-
-	UPROPERTY(transient)
-	uint32	LocalPredictionKey;
-
-	uint32	GetNextPredictionKey();
+	FPredictionKey	ReplicatedPredictionKey;
 
 	UFUNCTION()
 	void OnRep_PredictionKey();
@@ -151,24 +146,21 @@ class GAMEPLAYABILITIES_API UAbilitySystemComponent : public UActorComponent, pu
 	struct FPredictionInfo
 	{
 		FAbilitySystemComponentPredictionKeyClear PredictionKeyClearDelegate;
-		TArray<uint32> DependentPredictionKeys;
+		TArray<FPredictionKey::KeyType> DependentPredictionKeys;
 	};
 
-	TArray<TPair<uint32, FPredictionInfo > > PredictionDelegates;
+	TArray<TPair<FPredictionKey::KeyType, FPredictionInfo > > PredictionDelegates;
 
-	FPredictionInfo &	GetPredictionKeyDelegate(uint32 PredictionKey);
+	FPredictionInfo&	GetPredictionKeyDelegate(FPredictionKey::KeyType PredictionKey);
 
 	struct FPendingAbilityInfo
 	{
 		bool operator==(const FPendingAbilityInfo& Other) const
 		{
-			return PrevPredictionKey == Other.PrevPredictionKey
-				&& CurrPredictionKey == Other.CurrPredictionKey
-				&& Handle == Other.Handle;
+			return PredictionKey == Other.PredictionKey	&& Handle == Other.Handle;
 		}
 
-		uint32 PrevPredictionKey;
-		uint32 CurrPredictionKey;
+		FPredictionKey	PredictionKey;
 		FGameplayAbilitySpecHandle Handle;
 	};
 
@@ -189,11 +181,10 @@ class GAMEPLAYABILITIES_API UAbilitySystemComponent : public UActorComponent, pu
 
 		bool operator==(const FExecutingAbilityInfo& Other) const
 		{
-			return CurrPredictionKey == Other.CurrPredictionKey
-				&& State == Other.State;
+			return PredictionKey == Other.PredictionKey	&& State == Other.State;
 		}
 
-		uint32 CurrPredictionKey;
+		FPredictionKey PredictionKey;
 		EAbilityExecutionState State;
 		FGameplayAbilitySpecHandle Handle;
 	};
@@ -268,7 +259,7 @@ class GAMEPLAYABILITIES_API UAbilitySystemComponent : public UActorComponent, pu
 
 	FActiveGameplayEffectHandle ApplyGameplayEffectToTarget(UGameplayEffect *GameplayEffect, UAbilitySystemComponent *Target, float Level = FGameplayEffectLevelSpec::INVALID_LEVEL, FModifierQualifier BaseQualifier = FModifierQualifier());
 
-	UFUNCTION(BlueprintCallable, BlueprintAuthorityOnly, Category = GameplayEffects, meta=(FriendlyName = "ApplyGameplayEffectToSelf"))
+	UFUNCTION(BlueprintCallable, Category = GameplayEffects, meta=(FriendlyName = "ApplyGameplayEffectToSelf"))
 	FActiveGameplayEffectHandle K2_ApplyGameplayEffectToSelf(const UGameplayEffect *GameplayEffect, float Level, AActor *Instigator);
 	
 	FActiveGameplayEffectHandle ApplyGameplayEffectToSelf(const UGameplayEffect *GameplayEffect, float Level, AActor *Instigator, FModifierQualifier BaseQualifier = FModifierQualifier() );
@@ -316,24 +307,24 @@ class GAMEPLAYABILITIES_API UAbilitySystemComponent : public UActorComponent, pu
 	// GameplayCues can come from GameplayEffectSpecs
 
 	UFUNCTION(NetMulticast, unreliable)
-	void NetMulticast_InvokeGameplayCueExecuted_FromSpec(const FGameplayEffectSpec Spec);
+	void NetMulticast_InvokeGameplayCueExecuted_FromSpec(const FGameplayEffectSpec Spec, FPredictionKey PredictionKey);
 
 	void InvokeGameplayCueEvent(const FGameplayEffectSpec &Spec, EGameplayCueEvent::Type EventType);
 
 	// GameplayCues can also come on their own. For now these will have no additional parameters (just a tag) though that it something we could
 	// support down the road if we wanted.
 
-	void ExecuteGameplayCue(const FGameplayTag GameplayCueTag);
+	void ExecuteGameplayCue(const FGameplayTag GameplayCueTag, FPredictionKey PredictionKey);
 
-	void AddGameplayCue(const FGameplayTag GameplayCueTag);
+	void AddGameplayCue(const FGameplayTag GameplayCueTag, FPredictionKey PredictionKey);
 	
-	void RemoveGameplayCue(const FGameplayTag GameplayCueTag);
+	void RemoveGameplayCue(const FGameplayTag GameplayCueTag, FPredictionKey PredictionKey);
 
 	UFUNCTION(NetMulticast, unreliable)
-	void NetMulticast_InvokeGameplayCueExecuted(const FGameplayTag GameplayCueTag);
+	void NetMulticast_InvokeGameplayCueExecuted(const FGameplayTag GameplayCueTag, FPredictionKey PredictionKey);
 
 	UFUNCTION(NetMulticast, unreliable)
-	void NetMulticast_InvokeGameplayCueAdded(const FGameplayTag GameplayCueTag);
+	void NetMulticast_InvokeGameplayCueAdded(const FGameplayTag GameplayCueTag, FPredictionKey PredictionKey);
 
 	void InvokeGameplayCueEvent(const FGameplayTag GameplayCueTag, EGameplayCueEvent::Type EventType);
 
@@ -363,7 +354,7 @@ class GAMEPLAYABILITIES_API UAbilitySystemComponent : public UActorComponent, pu
 	FGameplayAbilitySpecHandle GiveAbility(FGameplayAbilitySpec AbilitySpec);
 
 	/** Attempts to activate the given ability */
-	bool TryActivateAbility(FGameplayAbilitySpecHandle AbilityToActivate, uint32 PrevPredictionKey = 0, uint32 CurrPredictionKey = 0, UGameplayAbility ** OutInstancedAbility = nullptr);
+	bool TryActivateAbility(FGameplayAbilitySpecHandle AbilityToActivate, FPredictionKey InPredictionKey=FPredictionKey(), UGameplayAbility ** OutInstancedAbility = nullptr);
 
 	void TriggerAbilityFromGameplayEvent(FGameplayAbilitySpecHandle AbilityToTrigger, FGameplayAbilityActorInfo* ActorInfo, FGameplayTag Tag, FGameplayEventData* Payload, UAbilitySystemComponent& Component);
 
@@ -396,7 +387,7 @@ class GAMEPLAYABILITIES_API UAbilitySystemComponent : public UActorComponent, pu
 	 *	instancing or anything else that the AbilitySystemComponent would provide, then it doesn't need the component to function.
 	 */
 
-	UPROPERTY(Replicated=OnRep_ActivateAbilities, BlueprintReadOnly, Category = "Abilities")
+	UPROPERTY(ReplicatedUsing=OnRep_ActivateAbilities, BlueprintReadOnly, Category = "Abilities")
 	TArray<FGameplayAbilitySpec>	ActivatableAbilities;
 
 	FGameplayAbilitySpec* FindAbilitySpecFromHandle(FGameplayAbilitySpecHandle Handle);
@@ -407,13 +398,13 @@ class GAMEPLAYABILITIES_API UAbilitySystemComponent : public UActorComponent, pu
 	void	OnRep_ActivateAbilities();
 
 	UFUNCTION(Server, reliable, WithValidation)
-	void	ServerTryActivateAbility(FGameplayAbilitySpecHandle AbilityToActivate, uint32 PrevPredictionKey, uint32 CurrPredictionKey);
+	void	ServerTryActivateAbility(FGameplayAbilitySpecHandle AbilityToActivate, FPredictionKey PredictionKey);
 
 	UFUNCTION(Client, Reliable)
-	void	ClientActivateAbilityFailed(FGameplayAbilitySpecHandle AbilityToActivate, uint32 PredictionKey);
+	void	ClientActivateAbilityFailed(FGameplayAbilitySpecHandle AbilityToActivate, int16 PredictionKey);
 
 	UFUNCTION(Client, Reliable)
-	void	ClientActivateAbilitySucceed(FGameplayAbilitySpecHandle AbilityToActivate, uint32 PredictionKey);
+	void	ClientActivateAbilitySucceed(FGameplayAbilitySpecHandle AbilityToActivate,int16 PredictionKey);
 	
 
 	// ----------------------------------------------------------------------------------------------------------------
@@ -556,6 +547,8 @@ class GAMEPLAYABILITIES_API UAbilitySystemComponent : public UActorComponent, pu
 
 private:
 
+	bool HasNetworkAuthorityToApplyGameplayEffect(const FModifierQualifier QualifierContext) const;
+
 	void ExecutePeriodicEffect(FActiveGameplayEffectHandle	Handle);
 
 	void ExecuteGameplayEffect(const FGameplayEffectSpec &Spec, const FModifierQualifier &QualifierContext);
diff --git a/Engine/Source/Runtime/GameplayAbilities/Public/GameplayEffect.h b/Engine/Source/Runtime/GameplayAbilities/Public/GameplayEffect.h
--- a/Engine/Source/Runtime/GameplayAbilities/Public/GameplayEffect.h
+++ b/Engine/Source/Runtime/GameplayAbilities/Public/GameplayEffect.h
@@ -788,8 +788,6 @@ struct FModifierQualifier
 
 	FModifierQualifier()
 		: MyType(EGameplayMod::Max)
-		, MyPrevPredictionKey(0)
-		, MyCurrPredictionKey(0)
 	{
 	}
 
@@ -808,30 +806,16 @@ struct FModifierQualifier
 
 	// ----------------------------------
 	// PredictionKey is used by networking/replication to specify that a GameplayEffect has been predictively created/added.
-
-	uint32 PreviousPredictionKey() const
-	{
-		return MyPrevPredictionKey;
-	}
-
-	uint32 CurrentPredictionKey() const
-	{
-		return MyCurrPredictionKey;
-	}
-
-	FModifierQualifier& PredictionKey(uint32 InPredictionKey)
+	
+	FModifierQualifier& PredictionKey(FPredictionKey InPredictionKey)
 	{
-		MyPrevPredictionKey = MyCurrPredictionKey;
-		MyCurrPredictionKey = InPredictionKey;
+		MyPredictionKey = InPredictionKey;
 		return *this;
 	}
 
-	// this has the same result as calling PredictionKey twice in a row to set both keys
-	FModifierQualifier& PredictionKeys(uint32 InPrevPredictionKey, uint32 InCurrPredictionKey)
+	FPredictionKey PredictionKey() const
 	{
-		MyPrevPredictionKey = InPrevPredictionKey;
-		MyCurrPredictionKey = InCurrPredictionKey;
-		return *this;
+		return MyPredictionKey;
 	}
 
 	// ----------------------------------
@@ -890,8 +874,7 @@ private:
 	EGameplayMod::Type	MyType;
 	FActiveGameplayEffectHandle MyIgnoreHandle;		// Do not modify this handle
 	FActiveGameplayEffectHandle MyExclusiveTargetHandle;	// Only modify this handle
-	uint32 MyPrevPredictionKey;
-	uint32 MyCurrPredictionKey;
+	FPredictionKey MyPredictionKey;
 };
 
 /**
@@ -1048,18 +1031,15 @@ struct FActiveGameplayEffect : public FFastArraySerializerItem
 	FActiveGameplayEffect()
 		: StartGameStateTime(0)
 		, StartWorldTime(0.f)
-		, PrevPredictionKey(0)
-		, CurrPredictionKey(0)
 	{
 	}
 
-	FActiveGameplayEffect(FActiveGameplayEffectHandle InHandle, const FGameplayEffectSpec &InSpec, float CurrentWorldTime, int32 InStartGameStateTime, uint32 InPrevPredictionKey, uint32 InCurrPredictionKey)
+	FActiveGameplayEffect(FActiveGameplayEffectHandle InHandle, const FGameplayEffectSpec &InSpec, float CurrentWorldTime, int32 InStartGameStateTime, FPredictionKey InPredictionKey)
 		: Handle(InHandle)
 		, Spec(InSpec)
 		, StartGameStateTime(InStartGameStateTime)
 		, StartWorldTime(CurrentWorldTime)
-		, PrevPredictionKey(InPrevPredictionKey)
-		, CurrPredictionKey(InCurrPredictionKey)
+		, PredictionKey(InPredictionKey)
 	{
 		for (FModifierSpec &Mod : Spec.Modifiers)
 		{
@@ -1072,6 +1052,9 @@ struct FActiveGameplayEffect : public FFastArraySerializerItem
 	UPROPERTY()
 	FGameplayEffectSpec Spec;
 
+	UPROPERTY()
+	FPredictionKey	PredictionKey;
+
 	/** Game time this started */
 	UPROPERTY()
 	int32 StartGameStateTime;
@@ -1085,10 +1068,6 @@ struct FActiveGameplayEffect : public FFastArraySerializerItem
 
 	FTimerHandle DurationHandle;
 
-	/** Prediction keys used to sync between client and server. If CurrPredictionKey != 0 then this effect does not use prediction. */
-	uint32 PrevPredictionKey;
-	uint32 CurrPredictionKey;
-
 	float GetTimeRemaining(float WorldTime)
 	{
 		float Duration = GetDuration();		
@@ -1119,7 +1098,7 @@ struct FActiveGameplayEffect : public FFastArraySerializerItem
 	}
 };
 
-/** Generic querying data structure for active GameplayEffects. Lets us ask things like:
+/** Generic querying data structure for active GameplayEffects. Lets us ask things like: ** 
  *		-Give me duration/magnitude of active gameplay effects with these tags
  *		-
  */
@@ -1163,7 +1142,7 @@ struct FActiveGameplayEffectsContainer : public FFastArraySerializer
 
 	UAbilitySystemComponent *Owner;
 	
-	FActiveGameplayEffect& CreateNewActiveGameplayEffect(const FGameplayEffectSpec &Spec, uint32 InPrevPredictionKey, uint32 InCurrPredictionKey);
+	FActiveGameplayEffect& CreateNewActiveGameplayEffect(const FGameplayEffectSpec &Spec, FPredictionKey InPredictionKey);
 
 	FActiveGameplayEffect* GetActiveGameplayEffect(const FActiveGameplayEffectHandle Handle);
 
diff --git a/Engine/Source/Runtime/GameplayAbilities/Public/GameplayEffectTypes.h b/Engine/Source/Runtime/GameplayAbilities/Public/GameplayEffectTypes.h
--- a/Engine/Source/Runtime/GameplayAbilities/Public/GameplayEffectTypes.h
+++ b/Engine/Source/Runtime/GameplayAbilities/Public/GameplayEffectTypes.h
@@ -4,6 +4,7 @@
 
 #include "GameplayTagAssetInterface.h"
 #include "AttributeSet.h"
+#include "GameplayPrediction.h"
 #include "GameplayEffectTypes.generated.h"
 
 #define SKILL_SYSTEM_AGGREGATOR_DEBUG 1
diff --git a/Engine/Source/Runtime/GameplayAbilities/Public/GameplayPrediction.h b/Engine/Source/Runtime/GameplayAbilities/Public/GameplayPrediction.h
new file mode 100644
--- /dev/null
+++ b/Engine/Source/Runtime/GameplayAbilities/Public/GameplayPrediction.h
@@ -0,0 +1,99 @@
+// Copyright 1998-2014 Epic Games, Inc. All Rights Reserved.
+
+#pragma once
+
+#include "GameplayPrediction.generated.h"
+
+USTRUCT()
+struct GAMEPLAYABILITIES_API FPredictionKey
+{
+	GENERATED_USTRUCT_BODY()
+
+	typedef int16 KeyType;
+
+	FPredictionKey()
+	: Current(0), Base(0), PredictiveConnection(nullptr), bIsStale(false)
+	{
+
+	}
+
+	bool IsValidKey() const
+	{
+		return Current > 0;
+	}
+
+	/**
+	 *	Stale means that a prediction key may be valid but is not usable for further prediction actions.
+	 *	E.g., we did predictive action A and are waiting for confirmation/rejection from server, we cannot
+	 *	then create a new predictive action B that is associated on the same key.
+	 */
+	bool IsStale() const
+	{
+		return Current > 0 && bIsStale == true;
+	}
+
+	bool operator==(const FPredictionKey& Other) const
+	{
+		return Current == Other.Current && Base == Other.Base;
+	}
+	
+	/** Construct a new prediction key with no dependancies */
+	static FPredictionKey CreateNewPredictionKey();
+
+	/** Create a new dependant prediction key: keep our existing base or use the current key as the base. */
+	void GenerateDependantPredictionKey();
+
+	UPROPERTY()
+	int16	Current;
+
+	UPROPERTY()
+	int16	Base;
+
+	UPROPERTY(NotReplicated)
+	UPackageMap*	PredictiveConnection;
+
+	UPROPERTY(NotReplicated)
+	bool bIsStale;
+
+	bool NetSerialize(FArchive& Ar, class UPackageMap* Map, bool& bOutSuccess);
+
+private:
+
+	void GenerateNewPredictionKey();
+
+	
+
+	FPredictionKey(int32 Key)
+		: Current(Key), Base(0), PredictiveConnection(nullptr)
+	{
+
+	}
+
+	FPredictionKey(int16 InKey, int16 PreviousKey)
+		: Current(InKey), Base(PreviousKey), PredictiveConnection(nullptr)
+	{
+
+	}
+};
+
+template<>
+struct TStructOpsTypeTraits<FPredictionKey> : public TStructOpsTypeTraitsBase
+{
+	enum
+	{
+		WithNetSerializer = true,
+		WithIdenticalViaEquality = true
+	};
+};
+
+// -----------------------------------------------------------------
+
+
+USTRUCT()
+struct GAMEPLAYABILITIES_API FScopedPredictionKey
+{
+	GENERATED_USTRUCT_BODY()
+
+	UPROPERTY()
+	FPredictionKey	Key;
+};
