Cleanup
-clarify GameplayEffect vs GameplayEffectSpec in blueprint facing functions
-expose level to GetOutgoingGameplayEffectSpec

diff --git a/Engine/Source/Runtime/GameplayAbilities/Private/Abilities/GameplayAbility.cpp b/Engine/Source/Runtime/GameplayAbilities/Private/Abilities/GameplayAbility.cpp
--- a/Engine/Source/Runtime/GameplayAbilities/Private/Abilities/GameplayAbility.cpp
+++ b/Engine/Source/Runtime/GameplayAbilities/Private/Abilities/GameplayAbility.cpp
@@ -421,10 +421,10 @@ USkeletalMeshComponent* UGameplayAbility::GetOwningComponentFromActorInfo() cons
 }
 
 /** Convenience method for abilities to get outgoing gameplay effect specs (for example, to pass on to projectiles to apply to whoever they hit) */
-FGameplayEffectSpecHandle UGameplayAbility::GetOutgoingSpec(UGameplayEffect* GameplayEffect) const
+FGameplayEffectSpecHandle UGameplayAbility::GetOutgoingGameplayEffectSpec(UGameplayEffect* GameplayEffect, float Level) const
 {
 	check(CurrentActorInfo && CurrentActorInfo->AbilitySystemComponent.IsValid());
-	return CurrentActorInfo->AbilitySystemComponent->GetOutgoingSpec(GameplayEffect);
+	return CurrentActorInfo->AbilitySystemComponent->GetOutgoingSpec(GameplayEffect, Level);
 }
 
 /** Fixme: Naming is confusing here */
diff --git a/Engine/Source/Runtime/GameplayAbilities/Private/AbilitySystemComponent.cpp b/Engine/Source/Runtime/GameplayAbilities/Private/AbilitySystemComponent.cpp
--- a/Engine/Source/Runtime/GameplayAbilities/Private/AbilitySystemComponent.cpp
+++ b/Engine/Source/Runtime/GameplayAbilities/Private/AbilitySystemComponent.cpp
@@ -171,12 +171,12 @@ float UAbilitySystemComponent::GetNumericAttribute(const FGameplayAttribute &Att
 	return Attribute.GetNumericValueChecked(AttributeSet);
 }
 
-FGameplayEffectSpecHandle UAbilitySystemComponent::GetOutgoingSpec(UGameplayEffect* GameplayEffect) const
+FGameplayEffectSpecHandle UAbilitySystemComponent::GetOutgoingSpec(UGameplayEffect* GameplayEffect, float Level) const
 {
 	SCOPE_CYCLE_COUNTER(STAT_GetOutgoingSpec);
 	// Fixme: we should build a map and cache these off. We can invalidate the map when an OutgoingGE modifier is applied or removed from us.
 
-	FGameplayEffectSpec* NewSpec = new FGameplayEffectSpec(GameplayEffect, AbilityActorInfo->Actor.Get(), FGameplayEffectLevelSpec::INVALID_LEVEL, GetCurveDataOverride());
+	FGameplayEffectSpec* NewSpec = new FGameplayEffectSpec(GameplayEffect, AbilityActorInfo->Actor.Get(), Level, GetCurveDataOverride());
 	if (ActiveGameplayEffects.ApplyActiveEffectsTo(*NewSpec, FModifierQualifier().Type(EGameplayMod::OutgoingGE)))
 	{
 		return FGameplayEffectSpecHandle(NewSpec);
diff --git a/Engine/Source/Runtime/GameplayAbilities/Public/Abilities/GameplayAbility.h b/Engine/Source/Runtime/GameplayAbilities/Public/Abilities/GameplayAbility.h
--- a/Engine/Source/Runtime/GameplayAbilities/Public/Abilities/GameplayAbility.h
+++ b/Engine/Source/Runtime/GameplayAbilities/Public/Abilities/GameplayAbility.h
@@ -320,7 +320,7 @@ protected:
 	FActiveGameplayEffectHandle ApplyGameplayEffectToOwner(const FGameplayAbilitySpecHandle Handle, const FGameplayAbilityActorInfo* ActorInfo, const FGameplayAbilityActivationInfo ActivationInfo, const UGameplayEffect* GameplayEffect, int32 GameplayEffectLevel);
 
 	UFUNCTION(BlueprintCallable, Category = Ability, FriendlyName = "ApplyGameplayEffectToTarget")
-	FActiveGameplayEffectHandle K2_ApplyGameplayEffectToTarget(FGameplayAbilityTargetDataHandle Target, const UGameplayEffect* GameplayEffect, int32 GameplayEffectLevel = 1);
+	FActiveGameplayEffectHandle K2_ApplyGameplayEffectToTarget(FGameplayAbilityTargetDataHandle TargetData, const UGameplayEffect* GameplayEffect, int32 GameplayEffectLevel = 1);
 
 	/** Non blueprintcallable, safe to call on CDO/NonInstance abilities */
 	FActiveGameplayEffectHandle ApplyGameplayEffectToTarget(const FGameplayAbilitySpecHandle Handle, const FGameplayAbilityActorInfo* ActorInfo, const FGameplayAbilityActivationInfo ActivationInfo, FGameplayAbilityTargetDataHandle Target, const UGameplayEffect* GameplayEffect, int32 GameplayEffectLevel);
@@ -385,7 +385,7 @@ protected:
 	USkeletalMeshComponent* GetOwningComponentFromActorInfo() const;
 
 	UFUNCTION(BlueprintCallable, Category=Ability)
-	FGameplayEffectSpecHandle GetOutgoingSpec(UGameplayEffect* GameplayEffect) const;	
+	FGameplayEffectSpecHandle GetOutgoingGameplayEffectSpec(UGameplayEffect* GameplayEffect, float Level=1.f) const;	
 
 	bool IsSupportedForNetworking() const override;
 
diff --git a/Engine/Source/Runtime/GameplayAbilities/Public/AbilitySystemComponent.h b/Engine/Source/Runtime/GameplayAbilities/Public/AbilitySystemComponent.h
--- a/Engine/Source/Runtime/GameplayAbilities/Public/AbilitySystemComponent.h
+++ b/Engine/Source/Runtime/GameplayAbilities/Public/AbilitySystemComponent.h
@@ -208,7 +208,7 @@ class GAMEPLAYABILITIES_API UAbilitySystemComponent : public UActorComponent, pu
 
 	/** Get an outgoing GameplayEffectSpec that is ready to be applied to other things. */
 	UFUNCTION(BlueprintCallable, Category = GameplayEffects)
-	FGameplayEffectSpecHandle GetOutgoingSpec(UGameplayEffect* GameplayEffect) const;
+	FGameplayEffectSpecHandle GetOutgoingSpec(UGameplayEffect* GameplayEffect, float Level) const;
 
 	/** Create an InstigatorContext for the owner of this AbilitySystemComponent */
 	UFUNCTION(BlueprintCallable, Category = GameplayEffects)
