small fix

diff --git a/Engine/Source/Runtime/GameplayAbilities/Private/AbilitySystemComponent.cpp b/Engine/Source/Runtime/GameplayAbilities/Private/AbilitySystemComponent.cpp
--- a/Engine/Source/Runtime/GameplayAbilities/Private/AbilitySystemComponent.cpp
+++ b/Engine/Source/Runtime/GameplayAbilities/Private/AbilitySystemComponent.cpp
@@ -215,7 +215,11 @@ FActiveGameplayEffectHandle UAbilitySystemComponent::K2_ApplyGameplayEffectToSel
 /** This is a helper function - it seems like this will be useful as a blueprint interface at the least, but Level parameter may need to be expanded */
 FActiveGameplayEffectHandle UAbilitySystemComponent::ApplyGameplayEffectToSelf(const UGameplayEffect *GameplayEffect, float Level, AActor *Instigator, FModifierQualifier BaseQualifier)
 {
-	check(GameplayEffect);
+	if (GameplayEffect == nullptr)
+	{
+		ABILITY_LOG(Error, TEXT("UAbilitySystemComponent::ApplyGameplayEffectToSelf called by Instigator %s with a null GameplayEffect."), *GetNameSafe(Instigator));
+		return FActiveGameplayEffectHandle();
+	}
 
 	if (HasNetworkAuthorityToApplyGameplayEffect(BaseQualifier))
 	{
