DebugString, WaitDelayTask

diff --git a/Engine/Source/Runtime/GameplayAbilities/Private/Abilities/Tasks/AbilityTask.cpp b/Engine/Source/Runtime/GameplayAbilities/Private/Abilities/Tasks/AbilityTask.cpp
--- a/Engine/Source/Runtime/GameplayAbilities/Private/Abilities/Tasks/AbilityTask.cpp
+++ b/Engine/Source/Runtime/GameplayAbilities/Private/Abilities/Tasks/AbilityTask.cpp
@@ -97,4 +97,9 @@ void UAbilityTask::OnDestroy(bool AbilityIsEnding)
 	}
 
 	MarkPendingKill();
+}
+
+FString UAbilityTask::GetDebugString() const
+{
+	return FString::Printf(TEXT("Generic %s"), *GetName());
 }
\ No newline at end of file
diff --git a/Engine/Source/Runtime/GameplayAbilities/Private/Abilities/Tasks/AbilityTask_PlayMontageAndWait.cpp b/Engine/Source/Runtime/GameplayAbilities/Private/Abilities/Tasks/AbilityTask_PlayMontageAndWait.cpp
--- a/Engine/Source/Runtime/GameplayAbilities/Private/Abilities/Tasks/AbilityTask_PlayMontageAndWait.cpp
+++ b/Engine/Source/Runtime/GameplayAbilities/Private/Abilities/Tasks/AbilityTask_PlayMontageAndWait.cpp
@@ -58,4 +58,19 @@ void UAbilityTask_PlayMontageAndWait::OnDestroy(bool AbilityEnded)
 	// (If we are destroyed, it will detect this and not do anything)
 
 	Super::OnDestroy(AbilityEnded);
+}
+
+FString UAbilityTask_PlayMontageAndWait::GetDebugString() const
+{
+	UAnimMontage* PlayingMontage = nullptr;
+	if (Ability.IsValid())
+	{
+		const FGameplayAbilityActorInfo* ActorInfo = Ability->GetCurrentActorInfo();
+		if (ActorInfo->AnimInstance.IsValid())
+		{
+			PlayingMontage = ActorInfo->AnimInstance->GetCurrentActiveMontage();
+		}
+	}
+
+	return FString::Printf(TEXT("PlayMontageAndWait. MontageToPlay: %s  (Currently Playing): %s"), *GetNameSafe(MontageToPlay), *GetNameSafe(PlayingMontage));
 }
\ No newline at end of file
diff --git a/Engine/Source/Runtime/GameplayAbilities/Private/Abilities/Tasks/AbilityTask_WaitDelay.cpp b/Engine/Source/Runtime/GameplayAbilities/Private/Abilities/Tasks/AbilityTask_WaitDelay.cpp
new file mode 100644
--- /dev/null
+++ b/Engine/Source/Runtime/GameplayAbilities/Private/Abilities/Tasks/AbilityTask_WaitDelay.cpp
@@ -0,0 +1,40 @@
+
+#include "AbilitySystemPrivatePCH.h"
+#include "Abilities/Tasks/AbilityTask_WaitDelay.h"
+
+UAbilityTask_WaitDelay::UAbilityTask_WaitDelay(const class FPostConstructInitializeProperties& PCIP)
+	: Super(PCIP)
+{
+	Time = 0.f;
+	TimeStarted = 0.f;
+}
+
+UAbilityTask_WaitDelay* UAbilityTask_WaitDelay::WaitDelay(UObject* WorldContextObject, float Time)
+{
+	auto MyObj = NewTask<UAbilityTask_WaitDelay>(WorldContextObject);
+	MyObj->Time = Time;
+	return MyObj;
+}
+
+void UAbilityTask_WaitDelay::Activate()
+{
+	UWorld* World = GetWorld();
+	TimeStarted = World->GetTimeSeconds();
+	World->GetTimerManager().SetTimer(this, &UAbilityTask_WaitDelay::OnTimeFinish, Time, false);
+}
+
+void UAbilityTask_WaitDelay::OnDestroy(bool AbilityEnded)
+{
+	Super::OnDestroy(AbilityEnded);
+}
+
+void UAbilityTask_WaitDelay::OnTimeFinish()
+{
+	OnFinish.Broadcast();
+}
+
+FString UAbilityTask_WaitDelay::GetDebugString() const
+{
+	float TimeLeft = GetWorld()->GetTimeSeconds() - TimeStarted;
+	return FString::Printf(TEXT("WaitDelay. Time: %.2f. TimeLeft: %.2f"), Time, TimeLeft);
+}
\ No newline at end of file
diff --git a/Engine/Source/Runtime/GameplayAbilities/Public/Abilities/Tasks/AbilityTask.h b/Engine/Source/Runtime/GameplayAbilities/Public/Abilities/Tasks/AbilityTask.h
--- a/Engine/Source/Runtime/GameplayAbilities/Public/Abilities/Tasks/AbilityTask.h
+++ b/Engine/Source/Runtime/GameplayAbilities/Public/Abilities/Tasks/AbilityTask.h
@@ -73,7 +73,14 @@ class GAMEPLAYABILITIES_API UAbilityTask : public UObject
 	/** Tick function for this task, if bTickingTask == true */
 	virtual void TickTask(float DeltaTime) {}
 
-	
+	/** Called when the ability is asked to confirm from an outside node. What this means depends on the individual task. By default, this does nothing other than ending if bEndTask is true. */
+	virtual void ExternalConfirm(bool bEndTask);
+
+	/** Called when the ability is asked to cancel from an outside node. What this means depends on the individual task. By default, this does nothing other than ending the task. */
+	virtual void ExternalCancel();
+
+	/** Return debug string describing task */
+	virtual FString GetDebugString() const;
 	
 	/** GameplayAbility that created us */
 	TWeakObjectPtr<UGameplayAbility> Ability;
@@ -105,14 +112,6 @@ class GAMEPLAYABILITIES_API UAbilityTask : public UObject
 	/** Called explicitly to end the task (usually by the task itself). Calls OnDestroy. */
 	void EndTask();
 
-	/** Called when the ability is asked to confirm from an outside node. What this means depends on the individual task. By default, this does nothing other than ending if bEndTask is true. */
-	virtual void ExternalConfirm(bool bEndTask);
-
-	/** Called when the ability is asked to cancel from an outside node. What this means depends on the individual task. By default, this does nothing other than ending the task. */
-	virtual void ExternalCancel();
-
-public:
-
 	/** This name allows us to find the task later so that we can end it. */
 	UPROPERTY()
 	FName InstanceName;
diff --git a/Engine/Source/Runtime/GameplayAbilities/Public/Abilities/Tasks/AbilityTask_PlayMontageAndWait.h b/Engine/Source/Runtime/GameplayAbilities/Public/Abilities/Tasks/AbilityTask_PlayMontageAndWait.h
--- a/Engine/Source/Runtime/GameplayAbilities/Public/Abilities/Tasks/AbilityTask_PlayMontageAndWait.h
+++ b/Engine/Source/Runtime/GameplayAbilities/Public/Abilities/Tasks/AbilityTask_PlayMontageAndWait.h
@@ -19,12 +19,13 @@ class UAbilityTask_PlayMontageAndWait : public UAbilityTask
 	UFUNCTION()
 	void OnMontageEnded(UAnimMontage* Montage, bool bInterrupted);
 
-
 	UFUNCTION(BlueprintCallable, Category="Ability|Tasks", meta = (FriendlyName="PlayMontageAndWait", HidePin = "WorldContextObject", DefaultToSelf = "WorldContextObject", BlueprintInternalUseOnly = "TRUE"))
 	static UAbilityTask_PlayMontageAndWait* CreatePlayMontageAndWaitProxy(UObject* WorldContextObject, UAnimMontage *MontageToPlay);
 
 	virtual void Activate() override;
 
+	virtual FString GetDebugString() const override;
+
 private:
 
 	virtual void OnDestroy(bool AbilityEnded) override;
diff --git a/Engine/Source/Runtime/GameplayAbilities/Public/Abilities/Tasks/AbilityTask_WaitDelay.h b/Engine/Source/Runtime/GameplayAbilities/Public/Abilities/Tasks/AbilityTask_WaitDelay.h
new file mode 100644
--- /dev/null
+++ b/Engine/Source/Runtime/GameplayAbilities/Public/Abilities/Tasks/AbilityTask_WaitDelay.h
@@ -0,0 +1,33 @@
+// Copyright 1998-2014 Epic Games, Inc. All Rights Reserved.
+#pragma once
+#include "AbilityTask.h"
+#include "Abilities/GameplayAbilityTypes.h"
+#include "AbilityTask_WaitDelay.generated.h"
+
+UCLASS(MinimalAPI)
+class UAbilityTask_WaitDelay : public UAbilityTask
+{
+	GENERATED_UCLASS_BODY()
+
+	DECLARE_DYNAMIC_MULTICAST_DELEGATE(FWaitDelayDelegate);
+
+	UPROPERTY(BlueprintAssignable)
+	FWaitDelayDelegate	OnFinish;
+
+	virtual void Activate() override;
+
+	/** Return debug string describing task */
+	virtual FString GetDebugString() const override;
+
+	/** Wait specified time. This is functionally the same as a standard Delay node. */
+	UFUNCTION(BlueprintCallable, Category="Ability|Tasks", meta = (HidePin = "WorldContextObject", DefaultToSelf = "WorldContextObject", BlueprintInternalUseOnly = "TRUE"))
+	static UAbilityTask_WaitDelay* WaitDelay(UObject* WorldContextObject, float Time);
+
+private:
+
+	void OnTimeFinish();
+
+	virtual void OnDestroy(bool AbilityEnded) override;
+	float Time;
+	float TimeStarted;
+};
\ No newline at end of file
