Prediction rework part 2
-Move formalized dependancy table out of AbilitySystemComponent
-Propert support for 'multicast' prediction events

diff --git a/Engine/Source/Runtime/GameplayAbilities/Private/AbilitySystemComponent.cpp b/Engine/Source/Runtime/GameplayAbilities/Private/AbilitySystemComponent.cpp
--- a/Engine/Source/Runtime/GameplayAbilities/Private/AbilitySystemComponent.cpp
+++ b/Engine/Source/Runtime/GameplayAbilities/Private/AbilitySystemComponent.cpp
@@ -156,7 +156,7 @@ bool UAbilitySystemComponent::IsOwnerActorAuthoritative() const
 
 bool UAbilitySystemComponent::HasNetworkAuthorityToApplyGameplayEffect(const FModifierQualifier QualifierContext) const
 {
-	return (IsOwnerActorAuthoritative() || QualifierContext.PredictionKey().IsStale()==false);
+	return (IsOwnerActorAuthoritative() || QualifierContext.PredictionKey().IsValidForMorePrediction());
 }
 
 void UAbilitySystemComponent::SetNumericAttribute(const FGameplayAttribute &Attribute, float NewFloatValue)
@@ -702,42 +702,7 @@ void UAbilitySystemComponent::OnRep_GameplayEffects()
 void UAbilitySystemComponent::OnRep_PredictionKey()
 {
 	// Every predictive action we've done up to and including the current value of ReplicatedPredictionKey needs to be wiped
-	int32 idx = 0;
-	for ( ; idx < PredictionDelegates.Num(); ++idx)
-	{
-		if (PredictionDelegates[idx].Key <= ReplicatedPredictionKey.Current)
-		{
-			PredictionDelegates[idx].Value.PredictionKeyClearDelegate.Broadcast();
-		}
-		else
-		{
-			break;
-		}
-	}
-	if (idx > 0)
-	{
-		// Remove everyone we called Broadcast on in one call
-		PredictionDelegates.RemoveAt(0, idx);
-	}
-}
-
-UAbilitySystemComponent::FPredictionInfo& UAbilitySystemComponent::GetPredictionKeyDelegate(FPredictionKey::KeyType PredictionKey)
-{
-	// See if we already have one for this key
-	for (auto & Item : PredictionDelegates)
-	{
-		if (Item.Key == PredictionKey)
-		{
-			return Item.Value;
-		}
-	}
-
-	// Create a new one
-	TPair<FPredictionKey::KeyType, FPredictionInfo> NewPair;
-	NewPair.Key = PredictionKey;
-	
-	PredictionDelegates.Add(NewPair);
-	return PredictionDelegates.Last().Value;
+	FPredictionKeyDelegates::CatchUpTo(ReplicatedPredictionKey.Current);
 }
 
 // ---------------------------------------------------------------------------------------
diff --git a/Engine/Source/Runtime/GameplayAbilities/Private/AbilitySystemComponent_Abilities.cpp b/Engine/Source/Runtime/GameplayAbilities/Private/AbilitySystemComponent_Abilities.cpp
--- a/Engine/Source/Runtime/GameplayAbilities/Private/AbilitySystemComponent_Abilities.cpp
+++ b/Engine/Source/Runtime/GameplayAbilities/Private/AbilitySystemComponent_Abilities.cpp
@@ -388,8 +388,13 @@ bool UAbilitySystemComponent::TryActivateAbility(FGameplayAbilitySpecHandle Hand
 		// This execution is now officially EGameplayAbilityActivationMode:Predicting and has a PredictionKey
 		ActivationInfo.GenerateNewPredictionKey();
 
+		FPredictionKey ThisPredictionKey = ActivationInfo.GetPredictionKeyForNewAction();
+		
 		// This must be called immediately after GeneratePredictionKey to prevent problems with recursively activating abilities
-		ServerTryActivateAbility(Handle, ActivationInfo.GetPredictionKeyForNewAction());
+		ServerTryActivateAbility(Handle, ThisPredictionKey);
+
+		// If this PredictionKey is rejected, we will call OnClientActivateAbilityFailed.
+		ThisPredictionKey.NewRejectedDelegate().BindUObject(this, &UAbilitySystemComponent::OnClientActivateAbilityFailed, Handle, ThisPredictionKey.Current);
 
 		if (Ability->GetInstancingPolicy() == EGameplayAbilityInstancingPolicy::InstancedPerExecution)
 		{
@@ -408,7 +413,7 @@ bool UAbilitySystemComponent::TryActivateAbility(FGameplayAbilitySpecHandle Hand
 			}
 			else
 			{
-				ensure(false);
+				ABILITY_LOG(Error, TEXT("TryActivateAbility called on ability %s that is InstancePerExecution and Replicated. This is an invalid configuration."), *Ability->GetName() );
 			}
 		}
 		else
@@ -460,19 +465,10 @@ void UAbilitySystemComponent::ServerTryActivateAbility_Implementation(FGameplayA
 			{
 				switch(ExecutingAbilityInfo.State)
 				{
-					case EAbilityExecutionState::Executing:
-					{
-						FGameplayAbilitySpec* Spec = FindAbilitySpecFromHandle(ExecutingAbilityInfo.Handle);
-						if (Spec)
-						{
-							// This feels wrong? We are changing the PredictionKey state on another ability?
-							Spec->ActivationInfo.SetPredictionKey(PredictionKey);
-						}
-						break;
-					}
 					case EAbilityExecutionState::Failed:
 						ClientActivateAbilityFailed(Handle, PredictionKey.Current);
 						break;
+					case EAbilityExecutionState::Executing:
 					case EAbilityExecutionState::Succeeded:
 						ClientActivateAbilitySucceed(Handle, PredictionKey.Current);
 						break;
@@ -520,39 +516,40 @@ bool UAbilitySystemComponent::ServerTryActivateAbility_Validate(FGameplayAbility
 	return true;
 }
 
+static_assert(sizeof(int16) == sizeof(FPredictionKey::KeyType), "Sizeof PredictionKey::KeyType does not match RPC parameters in AbilitySystemComponent ClientActivateAbilityFailed_Implementation");
+
 void UAbilitySystemComponent::ClientActivateAbilityFailed_Implementation(FGameplayAbilitySpecHandle Handle, int16 PredictionKey)
 {
+	// If this was predicted, we must use the PredictionKey system to end the ability
 	if (PredictionKey > 0)
 	{
-		for (int32 idx = 0; idx < PredictionDelegates.Num(); ++idx)
-		{
-			FPredictionKey::KeyType Key = PredictionDelegates[idx].Key;
-			if (Key == PredictionKey)
-			{
-				ABILITY_LOG(Warning, TEXT("Failed ActivateAbility, clearing prediction data %d"), PredictionKey);
-				PredictionDelegates[idx].Value.PredictionKeyClearDelegate.Broadcast();
+		FPredictionKeyDelegates::BroadcastRejectedDelegate(PredictionKey);
+	}
 
-				// Find the actual UGameplayAbility
-				FGameplayAbilitySpec* Spec = FindAbilitySpecFromHandle(Handle);
-				if (Spec)
-				{
-					UGameplayAbility* AbilityToActivate = Spec->Ability;
-					for (FPredictionKey::KeyType DependentKey : PredictionDelegates[idx].Value.DependentPredictionKeys)
-					{
-						ClientActivateAbilityFailed(Handle, DependentKey);
-					}
-				}
-				PredictionDelegates.RemoveAt(idx);
-				break;
-			}
-			else if(Key >= PredictionKey)
-			{
-				break;
-			}
+	// If this was not predicted... there is not much to do right now. We may need to keep track on the AbilitySpec if we are waiting on a confirm in these cases.
+}
+
+void UAbilitySystemComponent::OnClientActivateAbilityFailed(FGameplayAbilitySpecHandle Handle, FPredictionKey::KeyType PredictionKey)
+{
+	// Find the actual UGameplayAbility		
+	FGameplayAbilitySpec* Spec = FindAbilitySpecFromHandle(Handle);
+
+	ABILITY_LOG(Warning, TEXT("ClientActivateAbilityFailed_Implementation. PredictionKey :%d"), PredictionKey);
+
+	TArray<UGameplayAbility*> Instances = Spec->GetAbilityInstances();
+	for (UGameplayAbility* Ability : Instances)
+	{
+		if (Ability->CurrentActivationInfo.GetPredictionKey().Current == PredictionKey)
+		{
+			ABILITY_LOG(Warning, TEXT("Ending Ability %s"), *Ability->GetName());
+			Ability->K2_EndAbility();
 		}
 	}
 }
 
+
+static_assert(sizeof(int16) == sizeof(FPredictionKey::KeyType), "Sizeof PredictionKey::KeyType does not match RPC parameters in AbilitySystemComponent ClientActivateAbilitySucceed_Implementation");
+
 void UAbilitySystemComponent::ClientActivateAbilitySucceed_Implementation(FGameplayAbilitySpecHandle Handle, int16 PredictionKey)
 {
 	FGameplayAbilitySpec* Spec = FindAbilitySpecFromHandle(Handle);
diff --git a/Engine/Source/Runtime/GameplayAbilities/Private/GameplayEffect.cpp b/Engine/Source/Runtime/GameplayAbilities/Private/GameplayEffect.cpp
--- a/Engine/Source/Runtime/GameplayAbilities/Private/GameplayEffect.cpp
+++ b/Engine/Source/Runtime/GameplayAbilities/Private/GameplayEffect.cpp
@@ -1151,8 +1151,11 @@ void FActiveGameplayEffect::PostReplicatedAdd(const struct FActiveGameplayEffect
 	{
 		// PredictionKey will only be valid on the client that predicted it. So if this has a valid PredictionKey, we can assume we already predicted it and shouldn't invoke GameplayCues.
 		// We may need to do more bookkeeping here in the future. Possibly give the predicted gameplayeffect a chance to pass something off to the new replicated gameplay effect.
-
-		ShouldInvokeGameplayCueEvents = false;
+		
+		if (InArray.HasPredictedEffectWithPredictedKey(PredictionKey))
+		{
+			ShouldInvokeGameplayCueEvents = false;
+		}
 	}
 
 	const_cast<FActiveGameplayEffectsContainer&>(InArray).InternalOnActiveGameplayEffectAdded(*this);	// Const cast is ok. It is there to prevent mutation of the GameplayEffects array, which this wont do.
@@ -1582,7 +1585,7 @@ void FActiveGameplayEffectsContainer::StacksNeedToRecalculate()
 	}
 }
 
-FActiveGameplayEffect & FActiveGameplayEffectsContainer::CreateNewActiveGameplayEffect(const FGameplayEffectSpec &Spec, FPredictionKey InPredictionKey)
+FActiveGameplayEffect& FActiveGameplayEffectsContainer::CreateNewActiveGameplayEffect(const FGameplayEffectSpec &Spec, FPredictionKey InPredictionKey)
 {
 	SCOPE_CYCLE_COUNTER(STAT_CreateNewActiveGameplayEffect);
 
@@ -1640,13 +1643,8 @@ FActiveGameplayEffect & FActiveGameplayEffectsContainer::CreateNewActiveGameplay
 		// Clients predicting should call MarkArrayDirty to force the internal replication map to be rebuilt.
 		MarkArrayDirty();
 
-		UAbilitySystemComponent::FPredictionInfo& PredictionInfo = Owner->GetPredictionKeyDelegate(InPredictionKey.Current);
-		PredictionInfo.PredictionKeyClearDelegate.AddUObject(Owner, &UAbilitySystemComponent::RemoveActiveGameplayEffect_NoReturn, NewEffect.Handle);
-
-		if (InPredictionKey.Base) // Fixme: this should happen at a higher level
-		{
-			PredictionInfo.DependentPredictionKeys.AddUnique(InPredictionKey.Current);
-		}
+		// Once replicated state has caught up to this prediction key, we must remove this gameplay effect.
+		InPredictionKey.NewCaughtUpDelegate().BindUObject(Owner, &UAbilitySystemComponent::RemoveActiveGameplayEffect_NoReturn, NewEffect.Handle);
 	}
 
 	InternalOnActiveGameplayEffectAdded(NewEffect);	
@@ -1699,15 +1697,12 @@ bool FActiveGameplayEffectsContainer::InternalRemoveActiveGameplayEffect(int32 I
 		FActiveGameplayEffect& Effect = GameplayEffects[Idx];
 
 		bool ShouldInvokeGameplayCueEvent = true;
-		if (Effect.PredictionKey.IsValidKey() && !IsNetAuthority())
+		if (!IsNetAuthority() && Effect.PredictionKey.IsValidKey() && Effect.PredictionKey.WasReceived() == false)
 		{
-			// Don't invoke GameplayCue event if we still have another GameplayEffect that shares the same predictionkey
-			for (const FActiveGameplayEffect& OtherEffect : GameplayEffects)
+			// This was an effect that we predicted. Don't invoke GameplayCue event if we have another GameplayEffect that shares the same predictionkey and was received from the server
+			if (HasReceivedEffectWithPredictedKey(Effect.PredictionKey))
 			{
-				if (OtherEffect.PredictionKey == Effect.PredictionKey && (&OtherEffect != &Effect))
-				{
-					ShouldInvokeGameplayCueEvent = false;
-				}
+				ShouldInvokeGameplayCueEvent = false;
 			}
 		}
 
@@ -2008,7 +2003,7 @@ TArray<float> FActiveGameplayEffectsContainer::GetActiveEffectsTimeRemaining(con
 
 	TArray<float>	ReturnList;
 
-	for (const FActiveGameplayEffect &Effect : GameplayEffects)
+	for (const FActiveGameplayEffect& Effect : GameplayEffects)
 	{
 		if (Query.TagContainer)
 		{
@@ -2038,6 +2033,32 @@ FOnGameplayAttributeChange& FActiveGameplayEffectsContainer::RegisterGameplayAtt
 	return AttributeChangeDelegates.FindOrAdd(Attribute);
 }
 
+bool FActiveGameplayEffectsContainer::HasReceivedEffectWithPredictedKey(FPredictionKey PredictionKey) const
+{
+	for (const FActiveGameplayEffect& Effect : GameplayEffects)
+	{
+		if (Effect.PredictionKey == PredictionKey && Effect.PredictionKey.WasReceived() == true)
+		{
+			return true;
+		}
+	}
+
+	return false;
+}
+
+bool FActiveGameplayEffectsContainer::HasPredictedEffectWithPredictedKey(FPredictionKey PredictionKey) const
+{
+	for (const FActiveGameplayEffect& Effect : GameplayEffects)
+	{
+		if (Effect.PredictionKey == PredictionKey && Effect.PredictionKey.WasReceived() == false)
+		{
+			return true;
+		}
+	}
+
+	return false;
+}
+
 // --------------------------------------------------------------------------------------------------------------------------------------------------------
 //
 //	FGameplayEffectLevelSpec
diff --git a/Engine/Source/Runtime/GameplayAbilities/Private/GameplayPrediction.cpp b/Engine/Source/Runtime/GameplayAbilities/Private/GameplayPrediction.cpp
--- a/Engine/Source/Runtime/GameplayAbilities/Private/GameplayPrediction.cpp
+++ b/Engine/Source/Runtime/GameplayAbilities/Private/GameplayPrediction.cpp
@@ -41,18 +41,29 @@ bool FPredictionKey::NetSerialize(FArchive& Ar, class UPackageMap* Map, bool& bO
 
 void FPredictionKey::GenerateNewPredictionKey()
 {
-	static int32 GKey = 1;
+	static KeyType GKey = 1;
 	Current = GKey++;
 	bIsStale = false;
 }
 
 void FPredictionKey::GenerateDependantPredictionKey()
 {
+	KeyType Previous = 0;
 	if (Base == 0)
 	{
 		Base = Current;
 	}
+	else
+	{
+		Previous = Current;
+	}
+
 	GenerateNewPredictionKey();
+
+	if (Previous > 0)
+	{
+		FPredictionKeyDelegates::AddDependancy(Current, Previous);
+	}
 }
 
 FPredictionKey FPredictionKey::CreateNewPredictionKey()
@@ -60,4 +71,104 @@ FPredictionKey FPredictionKey::CreateNewPredictionKey()
 	FPredictionKey NewKey;
 	NewKey.GenerateNewPredictionKey();
 	return NewKey;
+}
+
+FPredictionKeyEvent& FPredictionKey::NewRejectedDelegate()
+{
+	return FPredictionKeyDelegates::NewRejectedDelegate(Current);
+}
+
+FPredictionKeyEvent& FPredictionKey::NewCaughtUpDelegate()
+{
+	return FPredictionKeyDelegates::NewCaughtUpDelegate(Current);
+}
+
+// -------------------------------------
+
+FPredictionKeyDelegates& FPredictionKeyDelegates::Get()
+{
+	static FPredictionKeyDelegates StaticMap;
+	return StaticMap;
+}
+
+FPredictionKeyEvent& FPredictionKeyDelegates::NewRejectedDelegate(FPredictionKey::KeyType Key)
+{
+	TArray<FPredictionKeyEvent>& DelegateList = Get().DelegateMap.FindOrAdd(Key).RejectedDelegates;
+	DelegateList.Add(FPredictionKeyEvent());
+	return DelegateList.Top();
+}
+
+FPredictionKeyEvent& FPredictionKeyDelegates::NewCaughtUpDelegate(FPredictionKey::KeyType Key)
+{
+	TArray<FPredictionKeyEvent>& DelegateList = Get().DelegateMap.FindOrAdd(Key).CaughtUpDelegates;
+	DelegateList.Add(FPredictionKeyEvent());
+	return DelegateList.Top();
+
+}
+
+void FPredictionKeyDelegates::BroadcastRejectedDelegate(FPredictionKey::KeyType Key)
+{
+	TArray<FPredictionKeyEvent>& DelegateList = Get().DelegateMap.FindOrAdd(Key).RejectedDelegates;
+	for (auto Delegate : DelegateList)
+	{
+		Delegate.ExecuteIfBound();
+	}
+}
+
+void FPredictionKeyDelegates::BroadcastCaughtUpDelegate(FPredictionKey::KeyType Key)
+{
+	TArray<FPredictionKeyEvent>& DelegateList = Get().DelegateMap.FindOrAdd(Key).CaughtUpDelegates;
+	for (auto Delegate : DelegateList)
+	{
+		Delegate.ExecuteIfBound();
+	}
+}
+
+void FPredictionKeyDelegates::Reject(FPredictionKey::KeyType Key)
+{
+	FDelegates* DelPtr = Get().DelegateMap.Find(Key);
+	if (DelPtr)
+	{
+		for (auto Delegate : DelPtr->RejectedDelegates)
+		{
+			Delegate.ExecuteIfBound();
+		}
+
+		DelPtr->RejectedDelegates.Empty();
+	}
+}
+
+void FPredictionKeyDelegates::CatchUpTo(FPredictionKey::KeyType Key)
+{
+	for (auto MapIt = Get().DelegateMap.CreateIterator(); MapIt; ++MapIt)
+	{
+		if (MapIt.Key() <= Key)
+		{		
+			for (auto Delegate : MapIt.Value().CaughtUpDelegates)
+			{
+				Delegate.ExecuteIfBound();
+			}
+			
+			MapIt.RemoveCurrent();
+		}
+	}
+}
+
+void FPredictionKeyDelegates::CaughtUp(FPredictionKey::KeyType Key)
+{
+	FDelegates* DelPtr = Get().DelegateMap.Find(Key);
+	if (DelPtr)
+	{
+		for (auto Delegate : DelPtr->CaughtUpDelegates)
+		{
+			Delegate.ExecuteIfBound();
+		}
+		Get().DelegateMap.Remove(Key);
+	}
+}
+
+void FPredictionKeyDelegates::AddDependancy(FPredictionKey::KeyType ThisKey, FPredictionKey::KeyType DependsOn)
+{
+	NewRejectedDelegate(DependsOn).BindStatic(&FPredictionKeyDelegates::Reject, ThisKey);
+	NewCaughtUpDelegate(DependsOn).BindStatic(&FPredictionKeyDelegates::CaughtUp, ThisKey);
 }
\ No newline at end of file
diff --git a/Engine/Source/Runtime/GameplayAbilities/Public/Abilities/GameplayAbilityTypes.h b/Engine/Source/Runtime/GameplayAbilities/Public/Abilities/GameplayAbilityTypes.h
--- a/Engine/Source/Runtime/GameplayAbilities/Public/Abilities/GameplayAbilityTypes.h
+++ b/Engine/Source/Runtime/GameplayAbilities/Public/Abilities/GameplayAbilityTypes.h
@@ -199,7 +199,7 @@ struct GAMEPLAYABILITIES_API FGameplayAbilityActivationInfo
 
 	FPredictionKey GetPredictionKeyForNewAction() const
 	{
-		return PredictionKey.IsStale() ? FPredictionKey() : PredictionKey;
+		return PredictionKey.IsValidForMorePrediction() ? PredictionKey : FPredictionKey();
 	}
 
 	FPredictionKey GetPredictionKey() const
diff --git a/Engine/Source/Runtime/GameplayAbilities/Public/AbilitySystemComponent.h b/Engine/Source/Runtime/GameplayAbilities/Public/AbilitySystemComponent.h
--- a/Engine/Source/Runtime/GameplayAbilities/Public/AbilitySystemComponent.h
+++ b/Engine/Source/Runtime/GameplayAbilities/Public/AbilitySystemComponent.h
@@ -143,16 +143,6 @@ class GAMEPLAYABILITIES_API UAbilitySystemComponent : public UActorComponent, pu
 	UFUNCTION()
 	void OnRep_PredictionKey();
 
-	struct FPredictionInfo
-	{
-		FAbilitySystemComponentPredictionKeyClear PredictionKeyClearDelegate;
-		TArray<FPredictionKey::KeyType> DependentPredictionKeys;
-	};
-
-	TArray<TPair<FPredictionKey::KeyType, FPredictionInfo > > PredictionDelegates;
-
-	FPredictionInfo&	GetPredictionKeyDelegate(FPredictionKey::KeyType PredictionKey);
-
 	struct FPendingAbilityInfo
 	{
 		bool operator==(const FPendingAbilityInfo& Other) const
@@ -403,6 +393,8 @@ class GAMEPLAYABILITIES_API UAbilitySystemComponent : public UActorComponent, pu
 	UFUNCTION(Client, Reliable)
 	void	ClientActivateAbilityFailed(FGameplayAbilitySpecHandle AbilityToActivate, int16 PredictionKey);
 
+	void	OnClientActivateAbilityFailed(FGameplayAbilitySpecHandle AbilityToActivate, FPredictionKey::KeyType PredictionKey);
+
 	UFUNCTION(Client, Reliable)
 	void	ClientActivateAbilitySucceed(FGameplayAbilitySpecHandle AbilityToActivate,int16 PredictionKey);
 	
diff --git a/Engine/Source/Runtime/GameplayAbilities/Public/GameplayEffect.h b/Engine/Source/Runtime/GameplayAbilities/Public/GameplayEffect.h
--- a/Engine/Source/Runtime/GameplayAbilities/Public/GameplayEffect.h
+++ b/Engine/Source/Runtime/GameplayAbilities/Public/GameplayEffect.h
@@ -1222,6 +1222,10 @@ struct FActiveGameplayEffectsContainer : public FFastArraySerializer
 
 	FOnGameplayAttributeChange& RegisterGameplayAttributeEvent(FGameplayAttribute Attribute);
 
+	bool HasReceivedEffectWithPredictedKey(FPredictionKey PredictionKey) const;
+
+	bool HasPredictedEffectWithPredictedKey(FPredictionKey PredictionKey) const;
+
 private:
 
 	FTimerHandle StackHandle;
diff --git a/Engine/Source/Runtime/GameplayAbilities/Public/GameplayPrediction.h b/Engine/Source/Runtime/GameplayAbilities/Public/GameplayPrediction.h
--- a/Engine/Source/Runtime/GameplayAbilities/Public/GameplayPrediction.h
+++ b/Engine/Source/Runtime/GameplayAbilities/Public/GameplayPrediction.h
@@ -4,6 +4,8 @@
 
 #include "GameplayPrediction.generated.h"
 
+DECLARE_DELEGATE(FPredictionKeyEvent);
+
 USTRUCT()
 struct GAMEPLAYABILITIES_API FPredictionKey
 {
@@ -27,15 +29,35 @@ struct GAMEPLAYABILITIES_API FPredictionKey
 	 *	E.g., we did predictive action A and are waiting for confirmation/rejection from server, we cannot
 	 *	then create a new predictive action B that is associated on the same key.
 	 */
-	bool IsStale() const
+	bool IsValidForMorePrediction() const
+	{
+		return Current > 0 && bIsStale == false;
+	}
+
+	bool WasReceived() const
 	{
-		return Current > 0 && bIsStale == true;
+		return PredictiveConnection != nullptr;
 	}
 
 	bool operator==(const FPredictionKey& Other) const
 	{
 		return Current == Other.Current && Base == Other.Base;
 	}
+
+	bool DependsOn(KeyType Key)
+	{
+		return (Current == Key || Base == Key);
+	}
+
+	FString ToString() const
+	{
+		return FString::Printf(TEXT("[%d/%d]"), Current, Base);
+	}
+
+	friend uint32 GetTypeHash(const FPredictionKey& InKey)
+	{
+		return ((InKey.Current << 16) | InKey.Base);
+	}
 	
 	/** Construct a new prediction key with no dependancies */
 	static FPredictionKey CreateNewPredictionKey();
@@ -43,6 +65,9 @@ struct GAMEPLAYABILITIES_API FPredictionKey
 	/** Create a new dependant prediction key: keep our existing base or use the current key as the base. */
 	void GenerateDependantPredictionKey();
 
+	FPredictionKeyEvent& NewRejectedDelegate();
+	FPredictionKeyEvent& NewCaughtUpDelegate();
+
 	UPROPERTY()
 	int16	Current;
 
@@ -89,11 +114,38 @@ struct TStructOpsTypeTraits<FPredictionKey> : public TStructOpsTypeTraitsBase
 // -----------------------------------------------------------------
 
 
-USTRUCT()
-struct GAMEPLAYABILITIES_API FScopedPredictionKey
+struct FPredictionKeyDelegates
 {
-	GENERATED_USTRUCT_BODY()
 
-	UPROPERTY()
-	FPredictionKey	Key;
+public:
+
+	struct FDelegates
+	{
+	public:
+
+		/** This delegate is called if the prediction key is associated with an action that is explicitly rejected by the server. */
+		TArray<FPredictionKeyEvent>	RejectedDelegates;
+
+		/** This delegate is called when replicated state has caught up with the prediction key. Doesnt imply rejection or acceptance. */
+		TArray<FPredictionKeyEvent>	CaughtUpDelegates;
+	};
+
+	TMap<FPredictionKey::KeyType, FDelegates>	DelegateMap;
+
+	static FPredictionKeyDelegates& Get();
+
+	static FPredictionKeyEvent&	NewRejectedDelegate(FPredictionKey::KeyType Key);
+	static FPredictionKeyEvent&	NewCaughtUpDelegate(FPredictionKey::KeyType Key);
+
+	static void	BroadcastRejectedDelegate(FPredictionKey::KeyType Key);
+	static void	BroadcastCaughtUpDelegate(FPredictionKey::KeyType Key);
+
+	static void Reject(FPredictionKey::KeyType Key);
+	static void CatchUpTo(FPredictionKey::KeyType Key);
+
+	static void AddDependancy(FPredictionKey::KeyType ThisKey, FPredictionKey::KeyType DependsOn);
+
+private:
+
+	static void CaughtUp(FPredictionKey::KeyType Key);
 };
