Prediction work part 2:
-Attribute modification prediction
-Improvements to prediction delegates/callbacks

diff --git a/Engine/Source/Runtime/GameplayAbilities/Private/AbilitySystemComponent.cpp b/Engine/Source/Runtime/GameplayAbilities/Private/AbilitySystemComponent.cpp
--- a/Engine/Source/Runtime/GameplayAbilities/Private/AbilitySystemComponent.cpp
+++ b/Engine/Source/Runtime/GameplayAbilities/Private/AbilitySystemComponent.cpp
@@ -606,6 +606,11 @@ void UAbilitySystemComponent::AddDependancyToAttribute(FGameplayAttribute Attrib
 	ActiveGameplayEffects.AddDependancyToAttribute(Attribute, InDependant);
 }
 
+void UAbilitySystemComponent::SetBaseAttributeValueFromReplication(float NewValue, FGameplayAttribute Attribute)
+{
+	ActiveGameplayEffects.SetBaseAttributeValueFromReplication(Attribute, NewValue);
+}
+
 bool UAbilitySystemComponent::CanApplyAttributeModifiers(const UGameplayEffect *GameplayEffect, float Level, AActor *Instigator)
 {
 	return ActiveGameplayEffects.CanApplyAttributeModifiers(GameplayEffect, Level, Instigator);
diff --git a/Engine/Source/Runtime/GameplayAbilities/Private/AttributeSet.cpp b/Engine/Source/Runtime/GameplayAbilities/Private/AttributeSet.cpp
--- a/Engine/Source/Runtime/GameplayAbilities/Private/AttributeSet.cpp
+++ b/Engine/Source/Runtime/GameplayAbilities/Private/AttributeSet.cpp
@@ -4,6 +4,7 @@
 #include "AttributeSet.h"
 #include "AbilitySystemComponent.h"
 #include "HAL/OutputDevices.h"
+#include "AbilitySystemGlobals.h"
 
 #include "ComponentReregisterContext.h"	
 
@@ -81,6 +82,11 @@ void UAttributeSet::ClampFromMetaDataTable(const UDataTable *DataTable)
 	PrintDebug();
 }
 
+UAbilitySystemComponent* UAttributeSet::GetOwningAbilitySystemComponent() const
+{
+	return UAbilitySystemGlobals::GetAbilitySystemComponentFromActor(CastChecked<AActor>(GetOuter()));
+}
+
 void UAttributeSet::PrintDebug()
 {
 	
diff --git a/Engine/Source/Runtime/GameplayAbilities/Private/GameplayEffect.cpp b/Engine/Source/Runtime/GameplayAbilities/Private/GameplayEffect.cpp
--- a/Engine/Source/Runtime/GameplayAbilities/Private/GameplayEffect.cpp
+++ b/Engine/Source/Runtime/GameplayAbilities/Private/GameplayEffect.cpp
@@ -175,7 +175,7 @@ void FGameplayEffectSpec::InitModifiers(const FGlobalCurveDataOverride *CurveDat
 		ModifierLevel->ApplyNewDef(ModInfo.LevelInfo, NewLevelSpec);
 
 		// This creates a new FModifierSpec that we own.
-		Modifiers.Emplace(FModifierSpec(ModInfo, NewLevelSpec, CurveData, Owner, Level));
+		new (Modifiers)FModifierSpec(ModInfo, NewLevelSpec, CurveData, Owner, Level);
 	}	
 }
 
@@ -185,6 +185,11 @@ void FGameplayEffectSpec::MakeUnique()
 	{
 		ModSpec.Aggregator.MakeUnique();
 	}
+
+	Duration.MakeUnique();
+	Period.MakeUnique();
+	ChanceToApplyToTarget.MakeUnique();
+	ChanceToExecuteOnGameplayEffect.MakeUnique();
 }
 
 int32 FGameplayEffectSpec::ApplyModifiersFrom(const FGameplayEffectSpec &InSpec, const FModifierQualifier &QualifierContext)
@@ -1146,6 +1151,12 @@ void FActiveGameplayEffect::PreReplicatedRemove(const struct FActiveGameplayEffe
 
 void FActiveGameplayEffect::PostReplicatedAdd(const struct FActiveGameplayEffectsContainer &InArray)
 {
+	if (Spec.Def == nullptr)
+	{
+		ABILITY_LOG(Error, TEXT("Received ReplicatedGameplayEffect with no UGameplayEffect def."));
+		return;
+	}
+
 	bool ShouldInvokeGameplayCueEvents = true;
 	if (PredictionKey.IsValidKey())
 	{
@@ -1404,7 +1415,6 @@ void FActiveGameplayEffectsContainer::ExecuteActiveEffectsFrom(const FGameplayEf
 		ABILITY_LOG(Log, TEXT("Invoking Execute GameplayCue for %s"), *Spec.ToSimpleString() );
 		Owner->NetMulticast_InvokeGameplayCueExecuted_FromSpec(Spec, QualifierContext.PredictionKey());
 	}
-
 }
 
 void FActiveGameplayEffectsContainer::ExecutePeriodicGameplayEffect(FActiveGameplayEffectHandle Handle)
@@ -1455,6 +1465,18 @@ FAggregatorRef & FActiveGameplayEffectsContainer::FindOrCreateAttributeAggregato
 	return OngoingPropertyEffects.Add(Attribute, FAggregatorRef(NewPropertyAggregator));
 }
 
+void FActiveGameplayEffectsContainer::SetBaseAttributeValueFromReplication(FGameplayAttribute Attribute, float BaseValue)
+{
+	FAggregatorRef* RefPtr = OngoingPropertyEffects.Find(Attribute);
+	if (RefPtr && RefPtr->IsValid())
+	{
+		check(RefPtr->Get()->BaseData.Magnitude.IsStatic());
+		RefPtr->Get()->BaseData.Magnitude.SetValue(BaseValue);
+		RefPtr->Get()->MarkDirty();
+		OnPropertyAggregatorDirty(RefPtr->Get(), Attribute);
+	}
+}
+
 float FActiveGameplayEffectsContainer::GetGameplayEffectDuration(FActiveGameplayEffectHandle Handle) const
 {
 	// Could make this a map for quicker lookup
@@ -1644,7 +1666,8 @@ FActiveGameplayEffect& FActiveGameplayEffectsContainer::CreateNewActiveGameplayE
 		MarkArrayDirty();
 
 		// Once replicated state has caught up to this prediction key, we must remove this gameplay effect.
-		InPredictionKey.NewCaughtUpDelegate().BindUObject(Owner, &UAbilitySystemComponent::RemoveActiveGameplayEffect_NoReturn, NewEffect.Handle);
+		InPredictionKey.NewRejectOrCaughtUpDelegate(FPredictionKeyEvent::CreateUObject(Owner, &UAbilitySystemComponent::RemoveActiveGameplayEffect_NoReturn, NewEffect.Handle));
+		
 	}
 
 	InternalOnActiveGameplayEffectAdded(NewEffect);	
diff --git a/Engine/Source/Runtime/GameplayAbilities/Private/GameplayPrediction.cpp b/Engine/Source/Runtime/GameplayAbilities/Private/GameplayPrediction.cpp
--- a/Engine/Source/Runtime/GameplayAbilities/Private/GameplayPrediction.cpp
+++ b/Engine/Source/Runtime/GameplayAbilities/Private/GameplayPrediction.cpp
@@ -83,6 +83,11 @@ FPredictionKeyEvent& FPredictionKey::NewCaughtUpDelegate()
 	return FPredictionKeyDelegates::NewCaughtUpDelegate(Current);
 }
 
+void FPredictionKey::NewRejectOrCaughtUpDelegate(FPredictionKeyEvent Event)
+{
+	FPredictionKeyDelegates::NewRejectOrCaughtUpDelegate(Current, Event);
+}
+
 // -------------------------------------
 
 FPredictionKeyDelegates& FPredictionKeyDelegates::Get()
@@ -103,7 +108,13 @@ FPredictionKeyEvent& FPredictionKeyDelegates::NewCaughtUpDelegate(FPredictionKey
 	TArray<FPredictionKeyEvent>& DelegateList = Get().DelegateMap.FindOrAdd(Key).CaughtUpDelegates;
 	DelegateList.Add(FPredictionKeyEvent());
 	return DelegateList.Top();
+}
 
+void FPredictionKeyDelegates::NewRejectOrCaughtUpDelegate(FPredictionKey::KeyType Key, FPredictionKeyEvent NewEvent)
+{
+	FDelegates& Delegates = Get().DelegateMap.FindOrAdd(Key);
+	Delegates.CaughtUpDelegates.Add(NewEvent);
+	Delegates.RejectedDelegates.Add(NewEvent);
 }
 
 void FPredictionKeyDelegates::BroadcastRejectedDelegate(FPredictionKey::KeyType Key)
@@ -134,7 +145,7 @@ void FPredictionKeyDelegates::Reject(FPredictionKey::KeyType Key)
 			Delegate.ExecuteIfBound();
 		}
 
-		DelPtr->RejectedDelegates.Empty();
+		Get().DelegateMap.Remove(Key);
 	}
 }
 
diff --git a/Engine/Source/Runtime/GameplayAbilities/Public/AbilitySystemComponent.h b/Engine/Source/Runtime/GameplayAbilities/Public/AbilitySystemComponent.h
--- a/Engine/Source/Runtime/GameplayAbilities/Public/AbilitySystemComponent.h
+++ b/Engine/Source/Runtime/GameplayAbilities/Public/AbilitySystemComponent.h
@@ -258,6 +258,8 @@ class GAMEPLAYABILITIES_API UAbilitySystemComponent : public UActorComponent, pu
 
 	void AddDependancyToAttribute(FGameplayAttribute Attribute, const TWeakPtr<FAggregator> InDependant);
 
+	void SetBaseAttributeValueFromReplication(float NewValue, FGameplayAttribute Attribute);
+
 	/** Tests if all modifiers in this GameplayEffect will leave the attribute > 0.f */
 	bool CanApplyAttributeModifiers(const UGameplayEffect *GameplayEffect, float Level, AActor *Instigator);
 
diff --git a/Engine/Source/Runtime/GameplayAbilities/Public/AttributeSet.h b/Engine/Source/Runtime/GameplayAbilities/Public/AttributeSet.h
--- a/Engine/Source/Runtime/GameplayAbilities/Public/AttributeSet.h
+++ b/Engine/Source/Runtime/GameplayAbilities/Public/AttributeSet.h
@@ -85,6 +85,8 @@ public:
 
 	void ClampFromMetaDataTable(const UDataTable *DataTable);
 
+	UAbilitySystemComponent* GetOwningAbilitySystemComponent() const;
+
 	virtual void PrintDebug();
 };
 
@@ -294,4 +296,19 @@ private:
 	};
 
 	TMap<FName, FAttributeSetDefaulsCollection>	Defaults;
-};
\ No newline at end of file
+};
+
+/**
+ *	This is a helper macro that can be used in RepNotify functions to handle attributes that will be predictively modified by clients.
+ *	
+ *	void UMyHealthSet::OnRep_Health()
+ *	{
+ *		GAMEPLAYATTRIBUTE_REPNOTIFY(UMyHealthSet, Health);
+ *	}
+ */
+
+#define GAMEPLAYATTRIBUTE_REPNOTIFY(C, P) \
+{ \
+	static UProperty* ThisProperty = FindFieldChecked<UProperty>(C::StaticClass(), GET_MEMBER_NAME_CHECKED(C, P)); \
+	GetOwningAbilitySystemComponent()->SetBaseAttributeValueFromReplication(P, FGameplayAttribute(ThisProperty)); \
+}
diff --git a/Engine/Source/Runtime/GameplayAbilities/Public/GameplayEffect.h b/Engine/Source/Runtime/GameplayAbilities/Public/GameplayEffect.h
--- a/Engine/Source/Runtime/GameplayAbilities/Public/GameplayEffect.h
+++ b/Engine/Source/Runtime/GameplayAbilities/Public/GameplayEffect.h
@@ -484,7 +484,7 @@ struct FAggregatorRef
 		WeakPtr = SharedPtr;
 	}
 
-	FAggregatorRef(struct FAggregatorRef *Src)
+	FAggregatorRef(const FAggregatorRef *Src)
 	{
 		SetSoftRef(Src);
 	}
@@ -500,7 +500,7 @@ struct FAggregatorRef
 		SharedPtr.Reset();
 	}
 
-	void SetSoftRef(FAggregatorRef *Src)
+	void SetSoftRef(const FAggregatorRef *Src)
 	{
 		check(!SharedPtr.IsValid());
 		WeakPtr = Src->SharedPtr;
@@ -543,7 +543,8 @@ struct TStructOpsTypeTraits< FAggregatorRef > : public TStructOpsTypeTraitsBase
 {
 	enum
 	{
-		WithNetSerializer = true
+		WithNetSerializer = true,
+		WithCopy = false
 	};
 };
 
@@ -1225,6 +1226,8 @@ struct FActiveGameplayEffectsContainer : public FFastArraySerializer
 	bool HasReceivedEffectWithPredictedKey(FPredictionKey PredictionKey) const;
 
 	bool HasPredictedEffectWithPredictedKey(FPredictionKey PredictionKey) const;
+		
+	void SetBaseAttributeValueFromReplication(FGameplayAttribute Attribute, float BaseBalue);
 
 private:
 
diff --git a/Engine/Source/Runtime/GameplayAbilities/Public/GameplayPrediction.h b/Engine/Source/Runtime/GameplayAbilities/Public/GameplayPrediction.h
--- a/Engine/Source/Runtime/GameplayAbilities/Public/GameplayPrediction.h
+++ b/Engine/Source/Runtime/GameplayAbilities/Public/GameplayPrediction.h
@@ -6,6 +6,32 @@
 
 DECLARE_DELEGATE(FPredictionKeyEvent);
 
+/**
+ *	FPredictionKey is a generic way of supporting Clientside Prediction in the GameplayAbility system.
+ *	A FPredictionKey is essentially an ID for indentifying predictive actions and side effects that are
+ *	done on a client. UAbilitySystemComponent supports syncrhonization of the predicton key and its side effects
+ *	between client and server.
+ *	
+ *	Essentially, anything can be associated with a PredictionKey, for example activating an Ability.
+ *	The client can generates a fresh PredictionKey and sends it to the server in his ServerTryActivateAbility call.
+ *	The server can confirm or reject this call (ClientActivateAbilitySucceed/Failed). 
+ *	
+ *	While the client is predictiving his ability, he is creating side effects (GameplayEffects, TriggeredEvents, Animations, etc).
+ *	As the client predicts these side effects, he associates each one with the prediction key generated at the start of the ability
+ *	activation.
+ *	
+ *	If the ability activation is rejected, the client can immediately revert these side effects. 
+ *	If the ability activation is accepted, the client must wait until the replicated side effects are sent to the server.
+ *		(the ClientActivatbleAbilitySucceed RPC will be immediately sent. Property replication may happen a few frames later).
+ *		Once replication of the server created side effects is finished, the client can undo his locally predictive side effects.
+ *		
+ *	The main things FPredictionKey itself provides are:
+ *		-Unique ID and a system for having dependant chains of Prediction Keys ("Current" and "Base" integers)
+ *		-A special implementation of ::NetSerialize *** which only serializes the prediction key to the predicting client ***
+ *			-This is important as it allows us to serializez prediction keys in replicated state, knowing that only clients that gave the server the prediction key will actually see them!
+ *	
+ */
+
 USTRUCT()
 struct GAMEPLAYABILITIES_API FPredictionKey
 {
@@ -19,34 +45,60 @@ struct GAMEPLAYABILITIES_API FPredictionKey
 
 	}
 
+	UPROPERTY()
+	int16	Current;
+
+	UPROPERTY()
+	int16	Base;
+
+	UPROPERTY(NotReplicated)
+	UPackageMap*	PredictiveConnection;
+
+	UPROPERTY(NotReplicated)
+	bool bIsStale;
+
+	/** Construct a new prediction key with no dependancies */
+	static FPredictionKey CreateNewPredictionKey();
+
+	/** Create a new dependant prediction key: keep our existing base or use the current key as the base. */
+	void GenerateDependantPredictionKey();
+
+	/** Creates new delegate called only when this key is rejected. */
+	FPredictionKeyEvent& NewRejectedDelegate();
+
+	/** Creates new delegate called only when replicated state catches up to this key. */
+	FPredictionKeyEvent& NewCaughtUpDelegate();
+	
+	/** Add a new delegate that is called if the key is rejected or caught up to. */
+	void NewRejectOrCaughtUpDelegate(FPredictionKeyEvent Event);
+	
+	bool NetSerialize(FArchive& Ar, class UPackageMap* Map, bool& bOutSuccess);
+
 	bool IsValidKey() const
 	{
 		return Current > 0;
 	}
 
-	/**
-	 *	Stale means that a prediction key may be valid but is not usable for further prediction actions.
-	 *	E.g., we did predictive action A and are waiting for confirmation/rejection from server, we cannot
-	 *	then create a new predictive action B that is associated on the same key.
-	 */
+	/** Can this key be used for more predictive actions, or has it already been sent off to the server? */
 	bool IsValidForMorePrediction() const
 	{
 		return Current > 0 && bIsStale == false;
 	}
 
+	/** Was this PredictoinKey received from a NetSerialize or created locallay? */
 	bool WasReceived() const
 	{
 		return PredictiveConnection != nullptr;
 	}
 
-	bool operator==(const FPredictionKey& Other) const
+	bool DependsOn(KeyType Key)
 	{
-		return Current == Other.Current && Base == Other.Base;
+		return (Current == Key || Base == Key);
 	}
 
-	bool DependsOn(KeyType Key)
+	bool operator==(const FPredictionKey& Other) const
 	{
-		return (Current == Key || Base == Key);
+		return Current == Other.Current && Base == Other.Base;
 	}
 
 	FString ToString() const
@@ -58,36 +110,11 @@ struct GAMEPLAYABILITIES_API FPredictionKey
 	{
 		return ((InKey.Current << 16) | InKey.Base);
 	}
-	
-	/** Construct a new prediction key with no dependancies */
-	static FPredictionKey CreateNewPredictionKey();
-
-	/** Create a new dependant prediction key: keep our existing base or use the current key as the base. */
-	void GenerateDependantPredictionKey();
-
-	FPredictionKeyEvent& NewRejectedDelegate();
-	FPredictionKeyEvent& NewCaughtUpDelegate();
-
-	UPROPERTY()
-	int16	Current;
-
-	UPROPERTY()
-	int16	Base;
-
-	UPROPERTY(NotReplicated)
-	UPackageMap*	PredictiveConnection;
-
-	UPROPERTY(NotReplicated)
-	bool bIsStale;
-
-	bool NetSerialize(FArchive& Ar, class UPackageMap* Map, bool& bOutSuccess);
 
 private:
 
 	void GenerateNewPredictionKey();
 
-	
-
 	FPredictionKey(int32 Key)
 		: Current(Key), Base(0), PredictiveConnection(nullptr)
 	{
@@ -113,6 +140,11 @@ struct TStructOpsTypeTraits<FPredictionKey> : public TStructOpsTypeTraitsBase
 
 // -----------------------------------------------------------------
 
+/**
+ *	This is a data structure for registering delegates associated with prediction key rejection and replicated state 'catching up'.
+ *	Delegates should be registered that revert side effects created with prediction keys.
+ *	
+ */
 
 struct FPredictionKeyDelegates
 {
@@ -136,6 +168,7 @@ public:
 
 	static FPredictionKeyEvent&	NewRejectedDelegate(FPredictionKey::KeyType Key);
 	static FPredictionKeyEvent&	NewCaughtUpDelegate(FPredictionKey::KeyType Key);
+	static void NewRejectOrCaughtUpDelegate(FPredictionKey::KeyType Key, FPredictionKeyEvent NewEvent);
 
 	static void	BroadcastRejectedDelegate(FPredictionKey::KeyType Key);
 	static void	BroadcastCaughtUpDelegate(FPredictionKey::KeyType Key);
