Prediction work part 3: scoped prediction window

diff --git a/Engine/Source/Runtime/GameplayAbilities/Private/Abilities/GameplayAbility.cpp b/Engine/Source/Runtime/GameplayAbilities/Private/Abilities/GameplayAbility.cpp
--- a/Engine/Source/Runtime/GameplayAbilities/Private/Abilities/GameplayAbility.cpp
+++ b/Engine/Source/Runtime/GameplayAbilities/Private/Abilities/GameplayAbility.cpp
@@ -630,17 +630,22 @@ FActiveGameplayEffectHandle UGameplayAbility::K2_ApplyGameplayEffectToTarget(FGa
 FActiveGameplayEffectHandle UGameplayAbility::ApplyGameplayEffectToTarget(const FGameplayAbilitySpecHandle Handle, const FGameplayAbilityActorInfo* ActorInfo, const FGameplayAbilityActivationInfo ActivationInfo, FGameplayAbilityTargetDataHandle Target, const UGameplayEffect* GameplayEffect, int32 GameplayEffectLevel)
 {
 	FActiveGameplayEffectHandle EffectHandle;
-	for (auto Data : Target.Data)
+
+
+	if (ActivationInfo.ActivationMode == EGameplayAbilityActivationMode::Authority || ActivationInfo.ActivationMode == EGameplayAbilityActivationMode::Predicting)
 	{
-		TArray<FActiveGameplayEffectHandle> EffectHandles = Data->ApplyGameplayEffect(GameplayEffect, ActorInfo, (float)GameplayEffectLevel);
-		if (EffectHandles.Num() > 0)
+		for (auto Data : Target.Data)
 		{
-			EffectHandle = EffectHandles[0];
-			if (EffectHandles.Num() > 1)
+			TArray<FActiveGameplayEffectHandle> EffectHandles = Data->ApplyGameplayEffect(GameplayEffect, ActorInfo, (float)GameplayEffectLevel, FModifierQualifier().PredictionKey(ActivationInfo.GetPredictionKeyForNewAction()));
+			if (EffectHandles.Num() > 0)
 			{
-				ABILITY_LOG(Warning, TEXT("ApplyGameplayEffectToTargetData_Single called on TargetData with multiple actor targets. %s"), *GetName());
+				EffectHandle = EffectHandles[0];
+				if (EffectHandles.Num() > 1)
+				{
+					ABILITY_LOG(Warning, TEXT("ApplyGameplayEffectToTargetData_Single called on TargetData with multiple actor targets. %s"), *GetName());
+				}
 			}
 		}
 	}
 	return EffectHandle;
-}
\ No newline at end of file
+}
diff --git a/Engine/Source/Runtime/GameplayAbilities/Private/Abilities/Tasks/AbilityTask_WaitInputRelease.cpp b/Engine/Source/Runtime/GameplayAbilities/Private/Abilities/Tasks/AbilityTask_WaitInputRelease.cpp
--- a/Engine/Source/Runtime/GameplayAbilities/Private/Abilities/Tasks/AbilityTask_WaitInputRelease.cpp
+++ b/Engine/Source/Runtime/GameplayAbilities/Private/Abilities/Tasks/AbilityTask_WaitInputRelease.cpp
@@ -1,6 +1,7 @@
 
 #include "AbilitySystemPrivatePCH.h"
 #include "Abilities/Tasks/AbilityTask_WaitInputRelease.h"
+#include "GameplayPrediction.h"
 
 UAbilityTask_WaitInputRelease::UAbilityTask_WaitInputRelease(const class FPostConstructInitializeProperties& PCIP)
 	: Super(PCIP)
@@ -13,10 +14,12 @@ void UAbilityTask_WaitInputRelease::OnReleaseCallback(UGameplayAbility* InAbilit
 {
 	float ElapsedTime = GetWorld()->GetTimeSeconds() - StartTime;
 
+	FScopedPredictionWindow	ScopedPrediction(InAbility);
+
 	if (Ability->GetCurrentActivationInfo().ActivationMode != EGameplayAbilityActivationMode::Authority)
 	{
 		// Tell the server we released.
-		AbilitySystemComponent->ServerInputRelease(Ability->GetCurrentAbilitySpecHandle());
+		AbilitySystemComponent->ServerInputRelease(Ability->GetCurrentAbilitySpecHandle(), ScopedPrediction.ScopedPredictionKey);
 	}
 
 	// We are done. Kill us so we don't keep getting broadcast messages
diff --git a/Engine/Source/Runtime/GameplayAbilities/Private/AbilitySystemComponent_Abilities.cpp b/Engine/Source/Runtime/GameplayAbilities/Private/AbilitySystemComponent_Abilities.cpp
--- a/Engine/Source/Runtime/GameplayAbilities/Private/AbilitySystemComponent_Abilities.cpp
+++ b/Engine/Source/Runtime/GameplayAbilities/Private/AbilitySystemComponent_Abilities.cpp
@@ -6,6 +6,7 @@
 #include "Abilities/GameplayAbility.h"
 #include "Abilities/GameplayAbilityTargetActor.h"
 #include "Abilities/Tasks/AbilityTask.h"
+#include "GameplayPrediction.h"
 
 #include "Net/UnrealNetwork.h"
 #include "MessageLog.h"
@@ -896,13 +897,15 @@ void UAbilitySystemComponent::InputCancel()
 	CancelCallbacks.Broadcast();
 }
 
-bool UAbilitySystemComponent::ServerInputRelease_Validate(FGameplayAbilitySpecHandle Handle)
+bool UAbilitySystemComponent::ServerInputRelease_Validate(FGameplayAbilitySpecHandle Handle, FPredictionKey ScopedPedictionKey)
 {
 	return true;
 }
 
-void UAbilitySystemComponent::ServerInputRelease_Implementation(FGameplayAbilitySpecHandle Handle)
+void UAbilitySystemComponent::ServerInputRelease_Implementation(FGameplayAbilitySpecHandle Handle, FPredictionKey ScopedPedictionKey)
 {
+	FScopedPredictionWindow ScopedPrediction(this, ScopedPedictionKey);
+
 	FGameplayAbilitySpec* Spec = FindAbilitySpecFromHandle(Handle);
 	if (Spec)
 	{
diff --git a/Engine/Source/Runtime/GameplayAbilities/Private/AttributeSet.cpp b/Engine/Source/Runtime/GameplayAbilities/Private/AttributeSet.cpp
--- a/Engine/Source/Runtime/GameplayAbilities/Private/AttributeSet.cpp
+++ b/Engine/Source/Runtime/GameplayAbilities/Private/AttributeSet.cpp
@@ -254,7 +254,8 @@ void FAttributeSetInitter::PreloadAttributeSetData(UCurveTable* CurveData)
 		TSubclassOf<UAttributeSet> Set = FindBestAttributeClass(ClassList, SetName);
 		if (!Set)
 		{
-			ABILITY_LOG(Warning, TEXT("FAttributeSetInitter::PreloadAttributeSetData Unable to match AttributeSet from %s (row: %s)"), *SetName, *RowName);
+			// This is ok, we may have rows in here that don't correspond directly to attributes
+			ABILITY_LOG(Verbose, TEXT("FAttributeSetInitter::PreloadAttributeSetData Unable to match AttributeSet from %s (row: %s)"), *SetName, *RowName);
 			continue;
 		}
 
diff --git a/Engine/Source/Runtime/GameplayAbilities/Private/GameplayAbilityTargetTypes.cpp b/Engine/Source/Runtime/GameplayAbilities/Private/GameplayAbilityTargetTypes.cpp
--- a/Engine/Source/Runtime/GameplayAbilities/Private/GameplayAbilityTargetTypes.cpp
+++ b/Engine/Source/Runtime/GameplayAbilities/Private/GameplayAbilityTargetTypes.cpp
@@ -7,7 +7,7 @@
 #include "AbilitySystemBlueprintLibrary.h"
 
 
-TArray<FActiveGameplayEffectHandle> FGameplayAbilityTargetData::ApplyGameplayEffect(const UGameplayEffect* GameplayEffect, const FGameplayAbilityActorInfo* InstigatorInfo, float Level)
+TArray<FActiveGameplayEffectHandle> FGameplayAbilityTargetData::ApplyGameplayEffect(const UGameplayEffect* GameplayEffect, const FGameplayAbilityActorInfo* InstigatorInfo, float Level, FModifierQualifier Qualifier)
 {
 	// TODO: Improve relationship between InstigatorContext and FGameplayAbilityTargetData/FHitResult (or use something different between HitResult)
 
@@ -38,7 +38,7 @@ TArray<FActiveGameplayEffectHandle> FGameplayAbilityTargetData::ApplyGameplayEff
 			UAbilitySystemComponent* TargetComponent = UAbilitySystemBlueprintLibrary::GetAbilitySystemComponent(TargetActor.Get());
 			if (TargetComponent)
 			{
-				AppliedHandles.Add( InstigatorInfo->AbilitySystemComponent->ApplyGameplayEffectSpecToTarget(SpecToApply, TargetComponent) );
+				AppliedHandles.Add( InstigatorInfo->AbilitySystemComponent->ApplyGameplayEffectSpecToTarget(SpecToApply, TargetComponent, Qualifier) );
 			}
 		}
 	}
diff --git a/Engine/Source/Runtime/GameplayAbilities/Private/GameplayPrediction.cpp b/Engine/Source/Runtime/GameplayAbilities/Private/GameplayPrediction.cpp
--- a/Engine/Source/Runtime/GameplayAbilities/Private/GameplayPrediction.cpp
+++ b/Engine/Source/Runtime/GameplayAbilities/Private/GameplayPrediction.cpp
@@ -182,4 +182,60 @@ void FPredictionKeyDelegates::AddDependancy(FPredictionKey::KeyType ThisKey, FPr
 {
 	NewRejectedDelegate(DependsOn).BindStatic(&FPredictionKeyDelegates::Reject, ThisKey);
 	NewCaughtUpDelegate(DependsOn).BindStatic(&FPredictionKeyDelegates::CaughtUp, ThisKey);
-}
\ No newline at end of file
+}
+
+// -------------------------------------
+
+FScopedPredictionWindow::FScopedPredictionWindow(UAbilitySystemComponent* AbilitySystemComponent, FPredictionKey InPredictionKey)
+{
+	Owner = AbilitySystemComponent;
+	Owner->ScopedPedictionKey = InPredictionKey;
+	ClearScopedPredictionKey = true;
+
+}
+
+FScopedPredictionWindow::FScopedPredictionWindow(UGameplayAbility* GameplayAbilityInstance)
+{
+	Ability = GameplayAbilityInstance;
+	check(Ability);
+
+	Owner = GameplayAbilityInstance->GetCurrentActorInfo()->AbilitySystemComponent.Get();
+	check(Owner);
+
+	ClearScopedPredictionKey = false;
+
+	FGameplayAbilityActivationInfo& ActivationInfo = Ability->GetCurrentActivationInfoRef();
+
+
+	if (ActivationInfo.ActivationMode == EGameplayAbilityActivationMode::Authority)
+	{
+		ActivationInfo.SetPredictionKey(Owner->ScopedPedictionKey);
+	}
+	else
+	{
+		ClientPrevActivationMode = static_cast<int8>(ActivationInfo.ActivationMode);
+
+		ActivationInfo.GenerateNewPredictionKey();
+		ScopedPredictionKey = ActivationInfo.GetPredictionKeyForNewAction();
+	}
+}
+
+FScopedPredictionWindow::~FScopedPredictionWindow()
+{
+	if (Ability)
+	{
+		// Restore old activation info settings
+		FGameplayAbilityActivationInfo& ActivationInfo = Ability->GetCurrentActivationInfoRef();
+		if (ActivationInfo.ActivationMode != EGameplayAbilityActivationMode::Authority)
+		{
+			ActivationInfo.ActivationMode = static_cast< TEnumAsByte<EGameplayAbilityActivationMode::Type> > (ClientPrevActivationMode);
+		}
+
+		ActivationInfo.SetPredictionStale();
+	}
+
+	if (ClearScopedPredictionKey && Owner)
+	{
+		Owner->ScopedPedictionKey = FPredictionKey();
+	}
+}
diff --git a/Engine/Source/Runtime/GameplayAbilities/Public/Abilities/GameplayAbility.h b/Engine/Source/Runtime/GameplayAbilities/Public/Abilities/GameplayAbility.h
--- a/Engine/Source/Runtime/GameplayAbilities/Public/Abilities/GameplayAbility.h
+++ b/Engine/Source/Runtime/GameplayAbilities/Public/Abilities/GameplayAbility.h
@@ -151,6 +151,12 @@ public:
 		return CurrentActivationInfo;
 	}
 
+	FGameplayAbilityActivationInfo& GetCurrentActivationInfoRef()
+	{
+		check(IsInstantiated());
+		return CurrentActivationInfo;
+	}
+
 	/** Gets the current AbilitySpecHandle- can only be called on instanced abilities. */
 	FGameplayAbilitySpecHandle GetCurrentAbilitySpecHandle() const
 	{
diff --git a/Engine/Source/Runtime/GameplayAbilities/Public/Abilities/GameplayAbilityTargetTypes.h b/Engine/Source/Runtime/GameplayAbilities/Public/Abilities/GameplayAbilityTargetTypes.h
--- a/Engine/Source/Runtime/GameplayAbilities/Public/Abilities/GameplayAbilityTargetTypes.h
+++ b/Engine/Source/Runtime/GameplayAbilities/Public/Abilities/GameplayAbilityTargetTypes.h
@@ -63,7 +63,7 @@ struct GAMEPLAYABILITIES_API FGameplayAbilityTargetData
 
 	virtual ~FGameplayAbilityTargetData() { }
 
-	TArray<FActiveGameplayEffectHandle> ApplyGameplayEffect(const UGameplayEffect* GameplayEffect, const struct FGameplayAbilityActorInfo* InstigatorInfo, float Level);
+	TArray<FActiveGameplayEffectHandle> ApplyGameplayEffect(const UGameplayEffect* GameplayEffect, const struct FGameplayAbilityActorInfo* InstigatorInfo, float Level, FModifierQualifier Qualifier = FModifierQualifier());
 
 	virtual TArray<TWeakObjectPtr<AActor> >	GetActors() const
 	{
diff --git a/Engine/Source/Runtime/GameplayAbilities/Public/AbilitySystemComponent.h b/Engine/Source/Runtime/GameplayAbilities/Public/AbilitySystemComponent.h
--- a/Engine/Source/Runtime/GameplayAbilities/Public/AbilitySystemComponent.h
+++ b/Engine/Source/Runtime/GameplayAbilities/Public/AbilitySystemComponent.h
@@ -140,6 +140,8 @@ class GAMEPLAYABILITIES_API UAbilitySystemComponent : public UActorComponent, pu
 	UPROPERTY(ReplicatedUsing=OnRep_PredictionKey)
 	FPredictionKey	ReplicatedPredictionKey;
 
+	FPredictionKey	ScopedPedictionKey;
+
 	UFUNCTION()
 	void OnRep_PredictionKey();
 
@@ -430,7 +432,7 @@ class GAMEPLAYABILITIES_API UAbilitySystemComponent : public UActorComponent, pu
 
 	/** Sent by abilities when they *need* to tell server when activation input is released. (Not sent by defauly only sent when using AbilityTask_WaitInputRelease) */
 	UFUNCTION(Server, reliable, WithValidation)
-	void ServerInputRelease(FGameplayAbilitySpecHandle AbilityToActivate);
+	void ServerInputRelease(FGameplayAbilitySpecHandle AbilityToActivate, FPredictionKey ScopedPedictionKey);
 
 	UFUNCTION(BlueprintCallable, Category="Abilities")
 	void InputConfirm();
diff --git a/Engine/Source/Runtime/GameplayAbilities/Public/GameplayEffect.h b/Engine/Source/Runtime/GameplayAbilities/Public/GameplayEffect.h
--- a/Engine/Source/Runtime/GameplayAbilities/Public/GameplayEffect.h
+++ b/Engine/Source/Runtime/GameplayAbilities/Public/GameplayEffect.h
@@ -769,115 +769,6 @@ private:
 	mutable FGameplayModifierEvaluatedData	CachedData;
 };
 
-/** 
- * Qualification context for applying modifiers.
- *	For example a Modifier may be setup in data to only apply to OutgoingGE mods.
- *	The FModifierQualifier is the data structure to hold the 'what type of modifier are we applying' data.
- *
- *  This should ideally only hold data that is outside of FGameplayEffectSpec or FGameplayModifierSpec.
- *	For example, specs know what they can and can't modify. We don't need to duplicate that there. 
- *  FModifierQualifier is meant to hold the data that comes from the calling context that is not intrinsic 
- *  to the existing data structures.
- *
- *	This class uses an optional initialization idiom such that you can do things like:
- *		FModifierQualifier().Handle(InHandle).(InType).etc...
- */
-USTRUCT(BlueprintType)
-struct FModifierQualifier
-{
-	GENERATED_USTRUCT_BODY()
-
-	FModifierQualifier()
-		: MyType(EGameplayMod::Max)
-	{
-	}
-
-	// ----------------------------------
-
-	FModifierQualifier& Type(EGameplayMod::Type InType)
-	{
-		MyType = InType;
-		return *this;
-	}
-
-	EGameplayMod::Type Type() const
-	{
-		return MyType;
-	}
-
-	// ----------------------------------
-	// PredictionKey is used by networking/replication to specify that a GameplayEffect has been predictively created/added.
-	
-	FModifierQualifier& PredictionKey(FPredictionKey InPredictionKey)
-	{
-		MyPredictionKey = InPredictionKey;
-		return *this;
-	}
-
-	FPredictionKey PredictionKey() const
-	{
-		return MyPredictionKey;
-	}
-
-	// ----------------------------------
-	// IgnoreHandle - ignore this handle completely. For example when executing an active gameplay effect, we never
-	// want to apply it to itself (either as IncomingGE or activeGE). It is ignored in all contexts.
-	//
-
-	FModifierQualifier& IgnoreHandle(FActiveGameplayEffectHandle InHandle)
-	{
-		MyIgnoreHandle = InHandle;
-		return *this;
-	}
-
-	FActiveGameplayEffectHandle IgnoreHandle() const
-	{
-		return MyIgnoreHandle;
-	}
-
-	// ----------------------------------
-	// ExclusiveTarget - sometimes we need to only apply a modifier to a specific active gameplay effect.
-	// We may only be able to use the handle if there are multiple instances of the same gameplay effect.
-	// This only applies in the context of applying/executing to a target. 'We only modify this active effect'.
-	// ExclusiveTarget is not checked in the context of applying outgoing/incoming GE modifiers to the spec.
-	//
-
-	FModifierQualifier& ExclusiveTarget(FActiveGameplayEffectHandle InHandle)
-	{
-		MyExclusiveTargetHandle = InHandle;
-		return *this;
-	}
-
-	FActiveGameplayEffectHandle ExclusiveTarget() const
-	{
-		return MyExclusiveTargetHandle;
-	}
-
-	// ----------------------------------
-
-	bool TestTarget(FActiveGameplayEffectHandle InHandle) const
-	{
-		if (MyIgnoreHandle.IsValid() && MyIgnoreHandle == InHandle)
-			return false;
-
-		if (MyExclusiveTargetHandle.IsValid() && MyExclusiveTargetHandle != InHandle)
-			return false;
-
-		return true;
-	}
-
-	FString ToString() const
-	{
-		return EGameplayModToString(MyType);
-	}
-
-private:
-	EGameplayMod::Type	MyType;
-	FActiveGameplayEffectHandle MyIgnoreHandle;		// Do not modify this handle
-	FActiveGameplayEffectHandle MyExclusiveTargetHandle;	// Only modify this handle
-	FPredictionKey MyPredictionKey;
-};
-
 /**
  * Modifier Specification
  *	-Const data (FGameplayModifierInfo) tells us what we modify, what we can modify
diff --git a/Engine/Source/Runtime/GameplayAbilities/Public/GameplayEffectTypes.h b/Engine/Source/Runtime/GameplayAbilities/Public/GameplayEffectTypes.h
--- a/Engine/Source/Runtime/GameplayAbilities/Public/GameplayEffectTypes.h
+++ b/Engine/Source/Runtime/GameplayAbilities/Public/GameplayEffectTypes.h
@@ -20,6 +20,16 @@ class UAbilitySystemComponent;
 struct FGameplayEffectSpec;
 struct FGameplayEffectModCallbackData;
 
+FString EGameplayModOpToString(int32 Type);
+
+FString EGameplayModToString(int32 Type);
+
+FString EGameplayModEffectToString(int32 Type);
+
+FString EGameplayEffectCopyPolicyToString(int32 Type);
+
+FString EGameplayEffectStackingPolicyToString(int32 Type);
+
 UENUM(BlueprintType)
 namespace EGameplayModOp
 {
@@ -271,6 +281,118 @@ struct TStructOpsTypeTraits< FGameplayEffectInstigatorContext > : public TStruct
 
 // -----------------------------------------------------------
 
+/** 
+ * Qualification context for applying modifiers.
+ *	For example a Modifier may be setup in data to only apply to OutgoingGE mods.
+ *	The FModifierQualifier is the data structure to hold the 'what type of modifier are we applying' data.
+ *
+ *  This should ideally only hold data that is outside of FGameplayEffectSpec or FGameplayModifierSpec.
+ *	For example, specs know what they can and can't modify. We don't need to duplicate that there. 
+ *  FModifierQualifier is meant to hold the data that comes from the calling context that is not intrinsic 
+ *  to the existing data structures.
+ *
+ *	This class uses an optional initialization idiom such that you can do things like:
+ *		FModifierQualifier().Handle(InHandle).(InType).etc...
+ */
+USTRUCT(BlueprintType)
+struct FModifierQualifier
+{
+	GENERATED_USTRUCT_BODY()
+
+	FModifierQualifier()
+		: MyType(EGameplayMod::Max)
+	{
+	}
+
+	// ----------------------------------
+
+	FModifierQualifier& Type(EGameplayMod::Type InType)
+	{
+		MyType = InType;
+		return *this;
+	}
+
+	EGameplayMod::Type Type() const
+	{
+		return MyType;
+	}
+
+	// ----------------------------------
+	// PredictionKey is used by networking/replication to specify that a GameplayEffect has been predictively created/added.
+	
+	FModifierQualifier& PredictionKey(FPredictionKey InPredictionKey)
+	{
+		MyPredictionKey = InPredictionKey;
+		return *this;
+	}
+
+	FPredictionKey PredictionKey() const
+	{
+		return MyPredictionKey;
+	}
+
+	// ----------------------------------
+	// IgnoreHandle - ignore this handle completely. For example when executing an active gameplay effect, we never
+	// want to apply it to itself (either as IncomingGE or activeGE). It is ignored in all contexts.
+	//
+
+	FModifierQualifier& IgnoreHandle(FActiveGameplayEffectHandle InHandle)
+	{
+		MyIgnoreHandle = InHandle;
+		return *this;
+	}
+
+	FActiveGameplayEffectHandle IgnoreHandle() const
+	{
+		return MyIgnoreHandle;
+	}
+
+	// ----------------------------------
+	// ExclusiveTarget - sometimes we need to only apply a modifier to a specific active gameplay effect.
+	// We may only be able to use the handle if there are multiple instances of the same gameplay effect.
+	// This only applies in the context of applying/executing to a target. 'We only modify this active effect'.
+	// ExclusiveTarget is not checked in the context of applying outgoing/incoming GE modifiers to the spec.
+	//
+
+	FModifierQualifier& ExclusiveTarget(FActiveGameplayEffectHandle InHandle)
+	{
+		MyExclusiveTargetHandle = InHandle;
+		return *this;
+	}
+
+	FActiveGameplayEffectHandle ExclusiveTarget() const
+	{
+		return MyExclusiveTargetHandle;
+	}
+
+	// ----------------------------------
+
+	bool TestTarget(FActiveGameplayEffectHandle InHandle) const
+	{
+		if (MyIgnoreHandle.IsValid() && MyIgnoreHandle == InHandle)
+			return false;
+
+		if (MyExclusiveTargetHandle.IsValid() && MyExclusiveTargetHandle != InHandle)
+			return false;
+
+		return true;
+	}
+
+	FString ToString() const
+	{
+		return EGameplayModToString(MyType);
+	}
+
+private:
+	EGameplayMod::Type	MyType;
+	FActiveGameplayEffectHandle MyIgnoreHandle;		// Do not modify this handle
+	FActiveGameplayEffectHandle MyExclusiveTargetHandle;	// Only modify this handle
+	FPredictionKey MyPredictionKey;
+};
+
+// -----------------------------------------------------------
+
+
 USTRUCT(BlueprintType)
 struct FGameplayCueParameters
 {
@@ -298,17 +420,6 @@ namespace EGameplayCueEvent
 	};
 }
 
-FString EGameplayModOpToString(int32 Type);
-
-FString EGameplayModToString(int32 Type);
-
-FString EGameplayModEffectToString(int32 Type);
-
-FString EGameplayEffectCopyPolicyToString(int32 Type);
-
-FString EGameplayEffectStackingPolicyToString(int32 Type);
-
-
 DECLARE_DELEGATE_OneParam(FOnGameplayAttributeEffectExecuted, struct FGameplayModifierEvaluatedData&);
 
 DECLARE_MULTICAST_DELEGATE_TwoParams(FOnGameplayEffectTagCountChanged, const FGameplayTag, int32 );
diff --git a/Engine/Source/Runtime/GameplayAbilities/Public/GameplayPrediction.h b/Engine/Source/Runtime/GameplayAbilities/Public/GameplayPrediction.h
--- a/Engine/Source/Runtime/GameplayAbilities/Public/GameplayPrediction.h
+++ b/Engine/Source/Runtime/GameplayAbilities/Public/GameplayPrediction.h
@@ -182,3 +182,34 @@ private:
 
 	static void CaughtUp(FPredictionKey::KeyType Key);
 };
+
+
+// -----------------------------------------------------------------
+
+class UAbilitySystemComponent;
+class UGameplayAbility;
+
+/**
+ *	A structure for allowing scoped prediction windows.
+ */
+
+struct FScopedPredictionWindow
+{
+	/** To be called on server when a new prediction key is received from the client (In an RPC). */
+	FScopedPredictionWindow(UAbilitySystemComponent* AbilitySystemComponent, FPredictionKey InPredictionKey);
+
+	/** To be called in the callsite where the predictive code will take place. This generates a new PredictionKey and acts as a synchonization point between client and server for that key.  */
+	FScopedPredictionWindow(UGameplayAbility* GameplayAbilityInstance);
+
+	~FScopedPredictionWindow();
+
+	FPredictionKey	ScopedPredictionKey;
+
+private:
+
+	UAbilitySystemComponent* Owner;
+	UGameplayAbility* Ability;
+
+	int8 ClientPrevActivationMode;
+	bool ClearScopedPredictionKey;
+};
\ No newline at end of file
