nonunity fix

diff --git a/Engine/Source/Runtime/GameplayAbilities/Private/GameplayPrediction.cpp b/Engine/Source/Runtime/GameplayAbilities/Private/GameplayPrediction.cpp
--- a/Engine/Source/Runtime/GameplayAbilities/Private/GameplayPrediction.cpp
+++ b/Engine/Source/Runtime/GameplayAbilities/Private/GameplayPrediction.cpp
@@ -1,6 +1,8 @@
 // Copyright 1998-2014 Epic Games, Inc. All Rights Reserved.
 
 #include "AbilitySystemPrivatePCH.h"
+#include "GameplayAbility.h"
+#include "AbilitySystemComponent.h"
 #include "GameplayPrediction.h"
 
 bool FPredictionKey::NetSerialize(FArchive& Ar, class UPackageMap* Map, bool& bOutSuccess)
