massive commenting

diff --git a/Engine/Source/Runtime/GameplayAbilities/Private/GameplayEffect.cpp b/Engine/Source/Runtime/GameplayAbilities/Private/GameplayEffect.cpp
--- a/Engine/Source/Runtime/GameplayAbilities/Private/GameplayEffect.cpp
+++ b/Engine/Source/Runtime/GameplayAbilities/Private/GameplayEffect.cpp
@@ -1146,6 +1146,12 @@ bool FActiveGameplayEffect::CanBeStacked(const FActiveGameplayEffect& Other) con
 
 void FActiveGameplayEffect::PreReplicatedRemove(const struct FActiveGameplayEffectsContainer &InArray)
 {
+	if (Spec.Def == nullptr)
+	{
+		ABILITY_LOG(Error, TEXT("Received PreReplicatedRemove with no UGameplayEffect def."));
+		return;
+	}
+
 	const_cast<FActiveGameplayEffectsContainer&>(InArray).InternalOnActiveGameplayEffectRemoved(*this);	// Const cast is ok. It is there to prevent mutation of the GameplayEffects array, which this wont do.
 	
 	InArray.Owner->InvokeGameplayCueEvent(Spec, EGameplayCueEvent::Removed);
@@ -1276,7 +1282,8 @@ void FActiveGameplayEffectsContainer::ApplySpecToActiveEffectsAndAttributes(FGam
 /** This is the main function that executes a GameplayEffect on Attributes and ActiveGameplayEffects */
 void FActiveGameplayEffectsContainer::ExecuteActiveEffectsFrom(const FGameplayEffectSpec &Spec, const FModifierQualifier &QualifierContext)
 {
-	bool InvokeGameplayCueExecute = false;
+	// If there are no modifiers, we always want to apply the GameplayCue. If there are modifiers, we only want to invoke the GameplayCue if one of them went through (coudl be blocked by immunity or % chance roll)
+	bool InvokeGameplayCueExecute = (Spec.Modifiers.Num() == 0);
 
 	// check if this is a stacking effect and if it is the active stacking effect.
 	if (Spec.GetStackingType() != EGameplayEffectStackingPolicy::Unlimited)
@@ -1409,7 +1416,7 @@ void FActiveGameplayEffectsContainer::ExecuteActiveEffectsFrom(const FGameplayEf
 			}
 		}
 	}
-
+	
 	if (InvokeGameplayCueExecute && Spec.Def->GameplayCues.Num())
 	{
 		// TODO: check replication policy. Right now we will replicate every execute via a multicast RPC
@@ -1475,7 +1482,6 @@ void FActiveGameplayEffectsContainer::SetBaseAttributeValueFromReplication(FGame
 		check(RefPtr->Get()->BaseData.Magnitude.IsStatic());
 		RefPtr->Get()->BaseData.Magnitude.SetValue(BaseValue);
 		RefPtr->Get()->MarkDirty();
-		OnPropertyAggregatorDirty(RefPtr->Get(), Attribute);
 	}
 }
 
diff --git a/Engine/Source/Runtime/GameplayAbilities/Public/GameplayPrediction.h b/Engine/Source/Runtime/GameplayAbilities/Public/GameplayPrediction.h
--- a/Engine/Source/Runtime/GameplayAbilities/Public/GameplayPrediction.h
+++ b/Engine/Source/Runtime/GameplayAbilities/Public/GameplayPrediction.h
@@ -6,17 +6,249 @@
 
 DECLARE_DELEGATE(FPredictionKeyEvent);
 
+/**
+ *	
+ *	Overview of Gameplay Ability Prediction
+ *	
+ *	High Level Goals:
+ *	At the GameplayAbility level (implementing an ability) prediction is transparent. An ability says "Do X->Y->Z", and we will automatically predict the parts of that that we can. 
+ *	We wish to avoid having logic such as "If Authority: Do X. Else: Do predictive version of X" in the ability itself.
+ *	
+ *	At this point, not all cases are solved, but we have a very solid framework for working with client side prediction.
+ *		
+ *	What do we currently predict?
+ *	-Ability activation
+ *	-Triggered Events
+ *	-GameplayEffect application:
+ *		-Attribute modification	(EXCEPTIONS: meta attributes such as Damage/Healing and certain cases of ongoing % based modifiers. Explained below)
+ *		-GameplayTag modification
+ *	-Gameplay Cue events (both from within predictive gameplay effect and on their own)
+ *		
+ *	-Montages
+ *	-Movement (built into UE4 UCharacterMovement)
+ *	
+ *	
+ *	Some things we don't predict (most of these we potentially could, but currently dont):
+ *	-GameplayEffect removal
+ *	-GameplayEffect periodic effects (dots ticking)
+ *	
+ *	
+ *	Problems we attempt to solve:
+ *	1. "Can I do this?" Basic protocol for prediction.
+ *	2. "Undo" How to undo side effects when a prediction fails.
+ *	3. "Redo" How to avoid replaying side effects that we predicted locally but that also get replicated from the server.
+ *	4. "Completeness" How to be sue we /really/ predicted all side effects.
+ *	5. "Dependencies" How to manage dependent prediction and chains of predicted events.
+ *	6. "Override" How to override state predictively that is otherwise replicated/owned by the server.
+ *	
+ *	---------------------------------------------------------
+ *	
+ *	Implementation Details
+ *	
+ *	*** PredictionKey ***
+ *	
+ *	A fundamental concept in this system is the Prediction Key (FPredictionKey). A prediction key on its own is simply a unique ID that is generated in a central place on the client. The client will send his prediction key to the server, 
+ *	and associate predictive actions and side effects with this key. The server may respond with an accept/reject for the prediction key, and will also associate the server-side created side effects with this prediction key.
+ *	
+ *	(IMPORTANT) FPredictionKey always replicate client -> server, but when replicating server -> clients they *only* replicate to the client that sent the prediction key to the server in the first place. 
+ *	This happens in FPredictionKey::NetSerialize. All other clients will receive an invalid (0) prediction key when a prediction key sent from a client is replicated back down through a replicated property.
+ *	
+ *	
+ *	*** Ability Activation ***
+ *	
+ *	Ability Activation is a first class predictive action. Whenever a client predictively activates an ability, he explicitly asks the server and the server explicitly responds. Once an ability has been
+ *	predictively activated, the client has a valid 'prediction window' where predictive side effects can happen which are not explicitly 'asked about'. (E.g., we do not explicitly ask 'Can I decrement mana, Can I
+ *	put this ability on cooldown. Those actions are considered logically atomic with activating an ability).
+ *	
+ *	
+ *	AbilitySystemComponent provides a set of functions for communicating ability activation between clients and server: TryActivateAbility -> ServerTryActivateAbility ->  ClientActivateAbility(Failed/Succeed).
+ *	
+ *	1. Client calls TryActivateAbility which generates a new FPredictionKey and calls ServerTryActivateAbility.
+ *	2. Client continues (before hearing back from server) and calls ActivateAbility with the generated PredictionKey associated with the Ability's ActivationInfo.
+ *	3. Any side effects that happen /before the call to ActivatAbility finish/ have the generated FPredictionKey associated with them.
+ *	4. Server decides if the ability really happened in ServerTryActivateAbility, calls ClientActivateAbility(Failed/Succeed) and sets UAbilitySystemComponent::ReplicatedPredictionKey to the generated key that was sent.
+ *	5. If client receives ClientAbilityFailed, he immediately kills the ability and rolls back side effects that were associated with the prediction key.
+ *		5a. 'Rolling back' is accomplished via FPredictionKeyDelegates and FPredictionKey::NewRejectedDelegate/NewCaughtUpDelegate/NewRejectOrCaughtUpDelegate. 
+ *
+ *
+ *		Registering the callback in TryActivateAbility:
+ *		
+ *		// If this PredictionKey is rejected, we will call OnClientActivateAbilityFailed.
+ *		ThisPredictionKey.NewRejectedDelegate().BindUObject(this, &UAbilitySystemComponent::OnClientActivateAbilityFailed, Handle, ThisPredictionKey.Current);
+ *		
+ *		
+ *		Invoking the callback in ClientActivateAbilityFailed_Implementation:
+ *		FPredictionKeyDelegates::BroadcastRejectedDelegate(PredictionKey);
+ *		
+ *	6. If accepted, client must wait until property replication catches up (the Succeed RPC will be sent immediately, property replication will happen on its own). Once the ReplicatedPredictionKey catches up to the
+ *		key used previous steps, the client can undo his predictive side effects. See UAbilitySystemComponent::OnRep_PredictionKey.
+ *						
+ *			 
+ *			 
+ *			 
+ *	*** GameplayEffect Prediction ***
+ *
+ *	GameplayEffects are considered side effects of prediction and are not explicitly asked about.
+ *	
+ *	1. GameplayEffects are only applied on clients if there is a valid prediction key. (If no prediction key, it simply skips the application on client).
+ *	2. Attributes, GameplayCues, and GameplayTags are all predicted if the GameplayEffect is predicted.
+ *	3. When the FActiveGameplayEffect is created, it stores the prediction key (FActiveGameplayEffect::PredictionKey)
+ *		3a. Instant effects are explained below in "Attribute Prediction".
+ *	4. On the server, the same prediction key is also set on the server's FActiveGameplayEffect that will be replicated down.
+ *	5. As a client, if you get a replicated FActiveGameplayEffect with a valid prediction key on it, you check to see if you have an ActiveGameplayEffect with that same key, if there is match, we do not apply
+ *		the 'on applied' type of logic, e.g., GameplayCues. The solves the "Redo" problem. However we will have 2 of the 'same' GameplayEffects in our ActiveGameplayEffects container, temporarily:
+ *	6. At the same time, UAbilitySystemComponent::ReplicatedPredictionKey will catch up and the predictive effects will be removed. When they are removed in this case, we again check PredicitonKey and decide 
+ *		if we should not do the 'On Remove' logic / GameplayCue.
+ *		
+ *	At this point, we have effectively predicted a gameplay effect as a side effect and handled the 'Undo' and 'Redo' problems.
+ *		
+ *	
+ *	*** Attribute Prediction ***
+ *	
+ *	Since attributes are replicated as standard uproperties, predicting modification to them can be tricky ("Override" problem). Instantaneous modification can be even harder since these are non stateful by nature.
+ *	(E.g., rolling back an attribute mod is difficult if there is no book keeping past the modification). This makes the "Undo" and "Redo" problem also hard in this case.
+ *	
+ *	The basic plan of attack is to treat attribute prediction as delta prediction rather than absolute value prediction. We do not predict that we have 90 mana, we predict that we have -10 mana from the server value, until 
+ *	the server confirms our prediction key. Basically, treat instant modifications as /infinite duration modifications/ to attributes while they are done predictively. The solves "Undo" and "Redo".
+ *	
+ *	For the "override" problem, we can handle this in the properties OnRep by treating the replicated (server) value as the 'base value' instead of 'final value' of the attribute, and to 
+ *	reaggregate our 'final value' after a replication happens.
+ *	
+ *	
+ *	1. We treat predictive instant gameplay effects as infinite duration gamepaly effects. See UAbilitySystemComponent::ApplyGameplayEffectSpecToSelf.
+ *	2. We have to *always* receive RepNotify calls on our attributes (not just when there is a change from last local value, since we will predict the change ahead of time). Done with REPNOTIFY_Always.
+ *	3. In the attribute RepNotify, we call into the AbilitySystemComponent::ActiveGameplayEffects to update our 'final value' give the new 'base value'. the GAMEPLAYATTRIBUTE_REPNOTIFY can do this.
+ *	4. Everything else will work like above (GameplayEffect prediction) : when the prediction key is caught up, the predictive GameplayEffect is removed and we will return to the server given value.
+ *	
+ *	
+ *	Example:
+ *	
+ *	void UMyHealthSet::GetLifetimeReplicatedProps(TArray< FLifetimeProperty > & OutLifetimeProps) const
+ *	{
+ *		Super::GetLifetimeReplicatedProps(OutLifetimeProps);
+ *
+ *		DOREPLIFETIME_CONDITION_NOTIFY(UMyHealthSet, Health, COND_None, REPNOTIFY_Always);
+ *	}
+ *	
+ *  void UMyHealthSet::OnRep_Health()
+ *  {
+ *		GAMEPLAYATTRIBUTE_REPNOTIFY(UMyHealthSet, Health);
+ *  }
+ *  
+ *  *** Gameplay Cue Events ***
+ *  
+ *  Outside of GameplayEffects which are already explained, Gameplay Cues can be activated on their own. These functions (UAbilitySystemComponent::ExecuteGameplayCue etc)  take network role and prediction keys into account.
+ *  
+ *  1. In UAbilitySystemComponent::ExecuteGameplayCue, if authority then do the multicast event (with replication key). If non authority but w/ a valid prediction key, predict the GameplayCue.
+ *  2. On the receiving end (NetMulticast_InvokeGameplayCueExecuted etc), if there is a replication key, then don't do the event (assume you predicted it).
+ *  
+ *  Remember that FPredictionKeys only replicate to the originating owner. This is an intrinsic property of FReplicationKey.
+ *	
+ *	*** Triggered Data Prediction ***
+ *	
+ *	Triggered Data is currently used to activate abilities. Essentially this all goes through the same code path as ActivateAbility. Rather than the ability being activated from input press, it is activated from
+ *	another game code driven event. Clients are able to predictively execute these events which predictively activate abilities. 
+ *	
+ *	There are some nuances to however, since the server will also run the code that triggers events. The server won't just wait to hear from the client. The server will keep a list of triggered abilities that have been
+ *	activated from a predictive ability. When receiving a TryActivate from a triggered ability, the server will look to see if /he/ has already run this ability, and respond with that information.
+ *	
+ *	There is work left to do on Triggered Events and replication. (explained at the end).
+ *	
+ *	---------------------------------------------------------	
+ *	
+ *	Advanced topic!	
+ *	
+ *	*** Dependencies ***
+ *	
+ *	We can have situations such as "Ability X activates and immediately triggers an event which activates Ability Y which triggers another Ability Z". The dependency chain is X->Y->Z.
+ *	Each of those abilities could be rejected by the server. If Y is rejected, then Z also never happened, but the server never tries to run Z, so the server doesn't explicitly decide 'no Z can't run'.
+ *	
+ *	To handle this, we have a concept of a Base PredictionKey, which is a member of FPredictionKey. When calling TryActivateAbility, we pass in the current PredictionKey (if applicable). That prediction key
+ *	is used as the base for any new prediction keys generated. We build a chain of keys this way, and can then invalidate Z if Y is rejected.
+ *	
+ *	This is slightly more nuanced though. In the X->Y->Z case, the server will only recieve the PredictionKey for X before trying to run the chain himself. E.g., he will TryActivate Y and Z with the original prediction key
+ *	sent to him from the client. Where as the client will generate a new PredictionKey each time he calls TryActivateAbility. The client *has* to generate a new PRedictionKey for each ability activate, since each activate
+ *	is not logically atomic. Each side effect produced in the chain of events has to have a unique PredictionKey. We cannot have GameplayEffects produced in X have the same PredictionKey produced in Z.
+ *	
+ *	To get around this, The prediction key of X is considered the Base key for Y and Z. The dependancy from Y to Z is kept completely client side, which is done in by FPredictionKeyDelegates::AddDependancy. We add delegates
+ *	to reject/catchup Z if Y rejected/confirmed.
+ *	
+ *	This dependency system allows us to have multiple predictive actions that are not logically atomic within a single prediction window/scope.
+ *
+ *	
+ *	
+ *	*** Additional Prediction Windows (within an Ability) ***
+ *	
+ *	As stated, A prediction key is only usable during a single logical scope. Once ActivateAbility returns, we are essentially done with that key. If the ability is waiting on an external event or timer, by the time it
+ *	returns, we will have gotten a confirm/reject from the server. Any side effects produced after this will no longer be tied to the lifespan of the original key.
+ *	
+ *	This isn't that bad, except that abilities will sometimes want to react to player input. For example, 'a hold down and charge' ability wants to instantly predict some stuff when the button is released. It is possible
+ *	to create a new prediction window within an ability with FScopedPredictionWindow. 
+ *	
+ *	FScopedPredictionWindows provides a way to send the server a new prediction key and have the server pick up and use that key within the same logical scope.
+ *	
+ *	UAbilityTask_WaitInputRelease::OnReleaseCallback is a good example. The flow of events is as followed:
+ *	1. Client enters UAbilityTask_WaitInputRelease::OnReleaseCallback and starts a new FScopedPredictionWindow. This creates a new prediction key for this scope (FScopedPredictionWindow::ScopedPredictionKey).
+ *	2. Client calls AbilitySystemComponent->ServerInputRelease which passes ScopedPrediction.ScopedPredictionKey as a parameter.
+ *	3. Server runs ServerInputRelease_Implementation which takes the passed in PredictionKey and sets it as UAbilitySystemComponent::ScopedPredictionKey with an FScopedPredictionWindow.
+ *	4. Server runs UAbilityTask_WaitInputRelease::OnReleaseCallback /within the same scope/
+ *	5. When the server hits the FScopedPredictionWindow in ::OnReleaseCallback, it gets the prediction key from UAbilitySystemComponent::ScopedPredictionKey. That is now used for all side effects within this logical scope.
+ *	6. Once the server ends this scoped prediction window, the prediction key used is finished and set to ReplicatedPredictionKey.
+ *	7. All side effects created in this scope now share a key between client and server.
+ *	
+ *	The key to this working is that ::OnReleaseCallback calls ::ServerInputRelease which calls ::OnReleaseCallback on the server. There is no room for anything else to happen and use the given prediction key.
+ *	
+ *	While there is no "Try/Failed/Succeed" calls in this example, all side effects are procedurally grouped/atomic. This solves the "Undo" and "Redo" problems for any arbitrary function calls that run on the server and client.
+ *	
+ *	
+ *	---------------------------------------------------------
+ *	
+ *	Unsupported / Issues/ Todo
+ *	
+ *	Triggered events do not explicitly replicate. E.g., if a triggered event only runs on the server, the client will never hear about it. This also prevents us from doing cross player/AI etc events. Support
+ *	for this should eventually be added and it should follow the same pattern that GameplayEffect and GameplayCues follow (predict triggered event with a prediction key, ignore the RPC event if it has a prediction key).
+ *	
+ *		 
+ *	*** Predicting "Meta" Attributes such as Damage/Healing vs "real" attributes such as Health ***
+ *	
+ *	We are unable to apply meta attributes predictively. Meta attributes only work on instant effects, in the back end of GameplayEffect (Pre/Post Modify Attribute on the UAttributeSet). These events are not called when
+ *	applying duration-based gameplay effects. E.g., a GameplayEffect that modifies damage for 5 seconds doesn't make sense. 
+ *	
+ *	In order to support this, we would probably add some limited support for duration based meta attributes, and move the transform of the instant gameplay effect from the front end (UAbilitySystemComponent::ApplyGameplayEffectSpecToSelf)
+ *	to the backend (UAttributeSet::PostModifyAttribute). 
+ * 
+ *			 
+ *	*** Predicting ongoing multiplicitive GameplayEffects ***
+ *	
+ *	There are also limitations when predicting % based gameplay effects. Since the server replicates down the 'final value' of an attribute, but not the entire aggregator chain of what is modifying it, we may run into cases where
+ *	the client cannot accurately predict new gameplay effects.
+ *	
+ *	For example: 
+ *	-Client has a perm +10% movement speed buff with base movement speed of 500 -> 550 is the final movement speed for this client.
+ *	-Client has an ability which grants an additional 10% movement speed buff. It is expected to *sum* the % based multipliers for a final 20% bonus to 500 -> 600 movement speed.
+ *	-However on the client, we just apply a 10% buff to 550 -> 605.
+ *	
+ *	This will need to be fixed by replicating down the aggregator chain for attributes. We already replicate some of this data, but not the full modifier list. We will need to look into supporting this eventually.
+ *	 
+ *	
+ *	
+ *	
+ *	
+ */
+
+
+
 /**
  *	FPredictionKey is a generic way of supporting Clientside Prediction in the GameplayAbility system.
- *	A FPredictionKey is essentially an ID for indentifying predictive actions and side effects that are
- *	done on a client. UAbilitySystemComponent supports syncrhonization of the predicton key and its side effects
+ *	A FPredictionKey is essentially an ID for identifying predictive actions and side effects that are
+ *	done on a client. UAbilitySystemComponent supports synchronization of the prediction key and its side effects
  *	between client and server.
  *	
  *	Essentially, anything can be associated with a PredictionKey, for example activating an Ability.
  *	The client can generates a fresh PredictionKey and sends it to the server in his ServerTryActivateAbility call.
  *	The server can confirm or reject this call (ClientActivateAbilitySucceed/Failed). 
  *	
- *	While the client is predictiving his ability, he is creating side effects (GameplayEffects, TriggeredEvents, Animations, etc).
+ *	While the client is predictively his ability, he is creating side effects (GameplayEffects, TriggeredEvents, Animations, etc).
  *	As the client predicts these side effects, he associates each one with the prediction key generated at the start of the ability
  *	activation.
  *	
@@ -28,7 +260,7 @@ DECLARE_DELEGATE(FPredictionKeyEvent);
  *	The main things FPredictionKey itself provides are:
  *		-Unique ID and a system for having dependant chains of Prediction Keys ("Current" and "Base" integers)
  *		-A special implementation of ::NetSerialize *** which only serializes the prediction key to the predicting client ***
- *			-This is important as it allows us to serializez prediction keys in replicated state, knowing that only clients that gave the server the prediction key will actually see them!
+ *			-This is important as it allows us to serialize prediction keys in replicated state, knowing that only clients that gave the server the prediction key will actually see them!
  *	
  */
 
@@ -207,8 +439,8 @@ struct FScopedPredictionWindow
 
 private:
 
-	UAbilitySystemComponent* Owner;
-	UGameplayAbility* Ability;
+	TWeakObjectPtr<UAbilitySystemComponent> Owner;
+	TWeakObjectPtr<UGameplayAbility> Ability;
 
 	int8 ClientPrevActivationMode;
 	bool ClearScopedPredictionKey;
