Fixed Clang compile error

diff --git a/Engine/Source/Runtime/GameplayAbilities/Private/GameplayPrediction.cpp b/Engine/Source/Runtime/GameplayAbilities/Private/GameplayPrediction.cpp
--- a/Engine/Source/Runtime/GameplayAbilities/Private/GameplayPrediction.cpp
+++ b/Engine/Source/Runtime/GameplayAbilities/Private/GameplayPrediction.cpp
@@ -199,10 +199,10 @@ FScopedPredictionWindow::FScopedPredictionWindow(UAbilitySystemComponent* Abilit
 FScopedPredictionWindow::FScopedPredictionWindow(UGameplayAbility* GameplayAbilityInstance)
 {
 	Ability = GameplayAbilityInstance;
-	check(Ability);
+	check(Ability.IsValid());
 
 	Owner = GameplayAbilityInstance->GetCurrentActorInfo()->AbilitySystemComponent.Get();
-	check(Owner);
+	check(Owner.IsValid());
 
 	ClearScopedPredictionKey = false;
 
@@ -224,7 +224,7 @@ FScopedPredictionWindow::FScopedPredictionWindow(UGameplayAbility* GameplayAbili
 
 FScopedPredictionWindow::~FScopedPredictionWindow()
 {
-	if (Ability)
+	if (Ability.IsValid())
 	{
 		// Restore old activation info settings
 		FGameplayAbilityActivationInfo& ActivationInfo = Ability->GetCurrentActivationInfoRef();
@@ -236,7 +236,7 @@ FScopedPredictionWindow::~FScopedPredictionWindow()
 		ActivationInfo.SetPredictionStale();
 	}
 
-	if (ClearScopedPredictionKey && Owner)
+	if (ClearScopedPredictionKey && Owner.IsValid())
 	{
 		Owner->ScopedPedictionKey = FPredictionKey();
 	}
