scoped prediction for target data / confirm targeting

diff --git a/Engine/Source/Runtime/GameplayAbilities/Private/Abilities/Tasks/AbilityTask_WaitTargetData.cpp b/Engine/Source/Runtime/GameplayAbilities/Private/Abilities/Tasks/AbilityTask_WaitTargetData.cpp
--- a/Engine/Source/Runtime/GameplayAbilities/Private/Abilities/Tasks/AbilityTask_WaitTargetData.cpp
+++ b/Engine/Source/Runtime/GameplayAbilities/Private/Abilities/Tasks/AbilityTask_WaitTargetData.cpp
@@ -81,7 +81,7 @@ bool UAbilityTask_WaitTargetData::BeginSpawningActor(UObject* WorldContextObject
 			// Register with the TargetData callbacks if we are expecting client to send them
 			if (!CDO->ShouldProduceTargetDataOnServer)
 			{
-				//Problem here - if there's targeting data just sitting around, fire events don't get hooked up because of this if-else, even if we don't end the task.
+				// Problem here - if there's targeting data just sitting around, fire events don't get hooked up because of this if-else, even if we don't end the task.
 				if (AbilitySystemComponent->ReplicatedTargetData.IsValid(0))
 				{
 					ValidData.Broadcast(AbilitySystemComponent->ReplicatedTargetData);
@@ -162,6 +162,8 @@ void UAbilityTask_WaitTargetData::FinishSpawningActor(UObject* WorldContextObjec
 void UAbilityTask_WaitTargetData::OnTargetDataReplicatedCallback(FGameplayAbilityTargetDataHandle Data)
 {
 	check(AbilitySystemComponent.IsValid());
+	
+	FScopedPredictionWindow	ScopedPrediction(Ability.Get());
 
 	/** 
 	 *  Call into the TargetActor to sanitize/verify the data. If this returns false, we are rejecting
@@ -201,10 +203,12 @@ void UAbilityTask_WaitTargetData::OnTargetDataReplicatedCancelledCallback()
 void UAbilityTask_WaitTargetData::OnTargetDataReadyCallback(FGameplayAbilityTargetDataHandle Data)
 {
 	check(AbilitySystemComponent.IsValid());
+
+	FScopedPredictionWindow	ScopedPrediction(Ability.Get());
 	
 	if (ShouldReplicateDataToServer())
 	{
-		AbilitySystemComponent->ServerSetReplicatedTargetData(Data);
+		AbilitySystemComponent->ServerSetReplicatedTargetData(Data, ScopedPrediction.ScopedPredictionKey);
 	}
 
 	ValidData.Broadcast(Data);
diff --git a/Engine/Source/Runtime/GameplayAbilities/Private/AbilitySystemComponent.cpp b/Engine/Source/Runtime/GameplayAbilities/Private/AbilitySystemComponent.cpp
--- a/Engine/Source/Runtime/GameplayAbilities/Private/AbilitySystemComponent.cpp
+++ b/Engine/Source/Runtime/GameplayAbilities/Private/AbilitySystemComponent.cpp
@@ -340,7 +340,7 @@ FActiveGameplayEffectHandle UAbilitySystemComponent::ApplyGameplayEffectSpecToSe
 
 	// Make sure we create our copy of the spec in the right place first...
 	FActiveGameplayEffectHandle	MyHandle;
-	bool bInvokeGameplayCueApplied = UGameplayEffect::INFINITE_DURATION == Spec.GetDuration();
+	bool bInvokeGameplayCueApplied = UGameplayEffect::INSTANT_APPLICATION != Spec.GetDuration(); // Cache this now before possibly modifying predictive instant effect to infinite duration effect.
 
 	FGameplayEffectSpec* OurCopyOfSpec = NULL;
 	TSharedPtr<FGameplayEffectSpec> StackSpec;
@@ -578,7 +578,7 @@ void UAbilitySystemComponent::InvokeGameplayCueEvent(const FGameplayTag Gameplay
 
 void UAbilitySystemComponent::NetMulticast_InvokeGameplayCueExecuted_FromSpec_Implementation(const FGameplayEffectSpec Spec, FPredictionKey PredictionKey)
 {
-	if (PredictionKey.IsValidKey() == false)
+	if (IsOwnerActorAuthoritative() || PredictionKey.IsValidKey() == false)
 	{
 		InvokeGameplayCueEvent(Spec, EGameplayCueEvent::Executed);
 	}
@@ -586,7 +586,7 @@ void UAbilitySystemComponent::NetMulticast_InvokeGameplayCueExecuted_FromSpec_Im
 
 void UAbilitySystemComponent::NetMulticast_InvokeGameplayCueExecuted_Implementation(const FGameplayTag GameplayCueTag, FPredictionKey PredictionKey)
 {
-	if (PredictionKey.IsValidKey() == false)
+	if (IsOwnerActorAuthoritative() || PredictionKey.IsValidKey() == false)
 	{
 		InvokeGameplayCueEvent(GameplayCueTag, EGameplayCueEvent::Executed);
 	}
@@ -594,7 +594,7 @@ void UAbilitySystemComponent::NetMulticast_InvokeGameplayCueExecuted_Implementat
 
 void UAbilitySystemComponent::NetMulticast_InvokeGameplayCueAdded_Implementation(const FGameplayTag GameplayCueTag, FPredictionKey PredictionKey)
 {
-	if (PredictionKey.IsValidKey() == false)
+	if (IsOwnerActorAuthoritative() || PredictionKey.IsValidKey() == false)
 	{
 		InvokeGameplayCueEvent(GameplayCueTag, EGameplayCueEvent::OnActive);
 	}
diff --git a/Engine/Source/Runtime/GameplayAbilities/Private/AbilitySystemComponent_Abilities.cpp b/Engine/Source/Runtime/GameplayAbilities/Private/AbilitySystemComponent_Abilities.cpp
--- a/Engine/Source/Runtime/GameplayAbilities/Private/AbilitySystemComponent_Abilities.cpp
+++ b/Engine/Source/Runtime/GameplayAbilities/Private/AbilitySystemComponent_Abilities.cpp
@@ -877,10 +877,13 @@ void UAbilitySystemComponent::AbilitySpectInputReleased(FGameplayAbilitySpec& Sp
 
 void UAbilitySystemComponent::InputConfirm()
 {
+	FPredictionKey NewKey = FPredictionKey::CreateNewPredictionKey();
+	FScopedPredictionWindow ScopedPrediction(this, NewKey);
+
 	if (GetOwnerRole() != ROLE_Authority && ConfirmCallbacks.IsBound())
 	{
 		// Tell the server we confirmed input.
-		ServerSetReplicatedConfirm(true);
+		ServerSetReplicatedConfirm(true, NewKey);
 	}
 	
 	ConfirmCallbacks.Broadcast();
@@ -888,10 +891,13 @@ void UAbilitySystemComponent::InputConfirm()
 
 void UAbilitySystemComponent::InputCancel()
 {
+	FPredictionKey NewKey = FPredictionKey::CreateNewPredictionKey();
+	FScopedPredictionWindow ScopedPrediction(this, NewKey);
+
 	if (GetOwnerRole() != ROLE_Authority && CancelCallbacks.IsBound())
 	{
 		// Tell the server we confirmed input.
-		ServerSetReplicatedConfirm(false);
+		ServerSetReplicatedConfirm(false, NewKey);
 	}
 
 	CancelCallbacks.Broadcast();
@@ -941,8 +947,10 @@ void UAbilitySystemComponent::TargetCancel()
 
 // --------------------------------------------------------------------------
 
-void UAbilitySystemComponent::ServerSetReplicatedConfirm_Implementation(bool Confirmed)
+void UAbilitySystemComponent::ServerSetReplicatedConfirm_Implementation(bool Confirmed, FPredictionKey PredictionKey)
 {
+	FScopedPredictionWindow ScopedPrediction(this, PredictionKey);
+
 	if (Confirmed)
 	{
 		ReplicatedConfirmAbility = true;
@@ -955,20 +963,22 @@ void UAbilitySystemComponent::ServerSetReplicatedConfirm_Implementation(bool Con
 	}
 }
 
-bool UAbilitySystemComponent::ServerSetReplicatedConfirm_Validate(bool Confirmed)
+bool UAbilitySystemComponent::ServerSetReplicatedConfirm_Validate(bool Confirmed, FPredictionKey PredictionKey)
 {
 	return true;
 }
 
 // -------
 
-void UAbilitySystemComponent::ServerSetReplicatedTargetData_Implementation(FGameplayAbilityTargetDataHandle Confirmed)
+void UAbilitySystemComponent::ServerSetReplicatedTargetData_Implementation(FGameplayAbilityTargetDataHandle Confirmed, FPredictionKey PredictionKey)
 {
+	FScopedPredictionWindow ScopedPrediction(this, PredictionKey);
+
 	ReplicatedTargetData = Confirmed;
 	ReplicatedTargetDataDelegate.Broadcast(ReplicatedTargetData);
 }
 
-bool UAbilitySystemComponent::ServerSetReplicatedTargetData_Validate(FGameplayAbilityTargetDataHandle Confirmed)
+bool UAbilitySystemComponent::ServerSetReplicatedTargetData_Validate(FGameplayAbilityTargetDataHandle Confirmed, FPredictionKey PredictionKey)
 {
 	return true;
 }
diff --git a/Engine/Source/Runtime/GameplayAbilities/Private/GameplayEffect.cpp b/Engine/Source/Runtime/GameplayAbilities/Private/GameplayEffect.cpp
--- a/Engine/Source/Runtime/GameplayAbilities/Private/GameplayEffect.cpp
+++ b/Engine/Source/Runtime/GameplayAbilities/Private/GameplayEffect.cpp
@@ -1201,6 +1201,14 @@ void FActiveGameplayEffect::PostReplicatedAdd(const struct FActiveGameplayEffect
 	StartWorldTime = WorldTimeSeconds - static_cast<float>(DeltaGameStateTime);
 }
 
+void FActiveGameplayEffect::PostReplicatedChange(const struct FActiveGameplayEffectsContainer &InArray)
+{
+	if (Spec.Def == nullptr)
+	{
+		ABILITY_LOG(Error, TEXT("Received ReplicatedGameplayEffect with no UGameplayEffect def."));
+	}
+}
+
 // --------------------------------------------------------------------------------------------------------------------------------------------------------
 //
 //	FActiveGameplayEffectsContainer
diff --git a/Engine/Source/Runtime/GameplayAbilities/Private/GameplayPrediction.cpp b/Engine/Source/Runtime/GameplayAbilities/Private/GameplayPrediction.cpp
--- a/Engine/Source/Runtime/GameplayAbilities/Private/GameplayPrediction.cpp
+++ b/Engine/Source/Runtime/GameplayAbilities/Private/GameplayPrediction.cpp
@@ -217,8 +217,16 @@ FScopedPredictionWindow::FScopedPredictionWindow(UGameplayAbility* GameplayAbili
 	{
 		ClientPrevActivationMode = static_cast<int8>(ActivationInfo.ActivationMode);
 
-		ActivationInfo.GenerateNewPredictionKey();
-		ScopedPredictionKey = ActivationInfo.GetPredictionKeyForNewAction();
+		if (Owner->ScopedPedictionKey.IsValidKey())
+		{
+			ActivationInfo.SetPredictionKey(Owner->ScopedPedictionKey);
+			ActivationInfo.ActivationMode = EGameplayAbilityActivationMode::Predicting;
+		}
+		else
+		{
+			ActivationInfo.GenerateNewPredictionKey();
+			ScopedPredictionKey = ActivationInfo.GetPredictionKeyForNewAction();
+		}
 	}
 }
 
diff --git a/Engine/Source/Runtime/GameplayAbilities/Public/Abilities/GameplayAbilityTargetActor.h b/Engine/Source/Runtime/GameplayAbilities/Public/Abilities/GameplayAbilityTargetActor.h
--- a/Engine/Source/Runtime/GameplayAbilities/Public/Abilities/GameplayAbilityTargetActor.h
+++ b/Engine/Source/Runtime/GameplayAbilities/Public/Abilities/GameplayAbilityTargetActor.h
@@ -17,7 +17,7 @@ public:
 	bool StaticTargetFunction;
 
 	/** The TargetData this class produces can be entirely generated on the server. We don't require the client to send us full or partial TargetData (possibly just a 'confirm') */
-	UPROPERTY()
+	UPROPERTY(EditAnywhere, Category=Advanced)
 	bool ShouldProduceTargetDataOnServer;
 
 	/** Describes where the targeting action starts, usually the player character or a socket on the player character. */
diff --git a/Engine/Source/Runtime/GameplayAbilities/Public/AbilitySystemComponent.h b/Engine/Source/Runtime/GameplayAbilities/Public/AbilitySystemComponent.h
--- a/Engine/Source/Runtime/GameplayAbilities/Public/AbilitySystemComponent.h
+++ b/Engine/Source/Runtime/GameplayAbilities/Public/AbilitySystemComponent.h
@@ -517,10 +517,10 @@ class GAMEPLAYABILITIES_API UAbilitySystemComponent : public UActorComponent, pu
 	 */
 
 	UFUNCTION(Server, reliable, WithValidation)
-	void ServerSetReplicatedConfirm(bool Confirmed);
+	void ServerSetReplicatedConfirm(bool Confirmed, FPredictionKey PredictionKey);
 
 	UFUNCTION(Server, reliable, WithValidation)
-	void ServerSetReplicatedTargetData(FGameplayAbilityTargetDataHandle ReplicatedTargetData);
+	void ServerSetReplicatedTargetData(FGameplayAbilityTargetDataHandle ReplicatedTargetData, FPredictionKey PredictionKey);
 
 	UFUNCTION(Server, reliable, WithValidation)
 	void ServerSetReplicatedTargetDataCancelled();
@@ -563,8 +563,7 @@ private:
 		return (GlobalCurveDataOverride.Overrides.Num() > 0 ? &GlobalCurveDataOverride : NULL);
 	}
 
-	FGlobalCurveDataOverride	GlobalCurveDataOverride;
-	
+	FGlobalCurveDataOverride	GlobalCurveDataOverride;	
 
 	// --------------------------------------------
 	
diff --git a/Engine/Source/Runtime/GameplayAbilities/Public/GameplayEffect.h b/Engine/Source/Runtime/GameplayAbilities/Public/GameplayEffect.h
--- a/Engine/Source/Runtime/GameplayAbilities/Public/GameplayEffect.h
+++ b/Engine/Source/Runtime/GameplayAbilities/Public/GameplayEffect.h
@@ -366,6 +366,11 @@ public:
 	TSubclassOf<class UGameplayEffectStackingExtension> StackingExtension;
 
 
+	/** Hack! Fixes issues in PIE when you create a dataasset and it cant be used in networking until you reload the editor */
+	bool IsNameStableForNetworking() const override
+	{
+		return true;
+	}
 
 protected:
 	void ValidateStacking();
@@ -982,7 +987,7 @@ struct FActiveGameplayEffect : public FFastArraySerializerItem
 
 	void PreReplicatedRemove(const struct FActiveGameplayEffectsContainer &InArray);
 	void PostReplicatedAdd(const struct FActiveGameplayEffectsContainer &InArray);
-	void PostReplicatedChange(const struct FActiveGameplayEffectsContainer &InArray) { }
+	void PostReplicatedChange(const struct FActiveGameplayEffectsContainer &InArray);
 
 	bool operator==(const FActiveGameplayEffect& Other)
 	{
