more comments

diff --git a/Engine/Source/Runtime/GameplayAbilities/Public/GameplayPrediction.h b/Engine/Source/Runtime/GameplayAbilities/Public/GameplayPrediction.h
--- a/Engine/Source/Runtime/GameplayAbilities/Public/GameplayPrediction.h
+++ b/Engine/Source/Runtime/GameplayAbilities/Public/GameplayPrediction.h
@@ -15,6 +15,9 @@ DECLARE_DELEGATE(FPredictionKeyEvent);
  *	We wish to avoid having logic such as "If Authority: Do X. Else: Do predictive version of X" in the ability itself.
  *	
  *	At this point, not all cases are solved, but we have a very solid framework for working with client side prediction.
+ *
+ *	When we say "client side prediction" we really mean client predicting game simulation state. Things can still be 'completely client side' without having to work within a prediction system.
+ *	For example, footsteps are completely client side and never interact with this system. But clients predicting their mana going from 100 to 90 when they cast a spell is 'client side prediction'.
  *		
  *	What do we currently predict?
  *	-Ability activation
@@ -231,7 +234,18 @@ DECLARE_DELEGATE(FPredictionKeyEvent);
  *	This will need to be fixed by replicating down the aggregator chain for attributes. We already replicate some of this data, but not the full modifier list. We will need to look into supporting this eventually.
  *	 
  *	
+ *	*** "Weak Prediction" ***
+ *	
+ *	We will probably still have cases that do not fit well into this system. Some situations will exist where a prediction key exchange is not feasible. For example, an ability where any one that player collides with/touches
+ *	receives a GameplayEffect that slows them and their material blue. Since we can't send Server RPCs every time this happens (and the server couldn't necessarily handle the message at his point in the simulation), there is no 
+ *	way to correlate the gameplay effect side effects between client and server.
+ *	
+ *	One approach here may be to think about a weaker form of prediction. One where there is not a fresh prediction key used and instead the server assumes the client will predict all side effects from an entire ability. This would
+ *	at least solve the "redo" problem but would not solve the "completeness" problem. If the client side prediction could be made as minimal as possible - for example only predicting an initial particle effect rather than
+ *	predicting the state and attribute change - then the problems get less severe.
  *	
+ *	I can envision a weak prediction mode which is what (certain abilities? All abilities?) fall back to when there is no fresh prediction key that can accurately correlate side effects. When in weak prediction mode, perhaps
+ *	only certain actions can be predicted - for example GameplayCue execute events, but not OnAdded/OnRemove events.
  *	
  *	
  */
