Moving/renaming classes for actor placement targeting.

diff --git a/Engine/Source/Runtime/GameplayAbilities/Public/Abilities/GameplayAbilityTargetActor_ActorPlacement.cpp b/Engine/Source/Runtime/GameplayAbilities/Public/Abilities/GameplayAbilityTargetActor_ActorPlacement.cpp
new file mode 100644
--- /dev/null
+++ b/Engine/Source/Runtime/GameplayAbilities/Public/Abilities/GameplayAbilityTargetActor_ActorPlacement.cpp
@@ -0,0 +1,45 @@
+// Copyright 1998-2014 Epic Games, Inc. All Rights Reserved.
+
+#include "AbilitySystemPrivatePCH.h"
+#include "AbilitySystemComponent.h"
+#include "GameplayAbilityTargetActor_ActorPlacement.h"
+#include "GameplayAbilityWorldReticle_ActorVisualization.h"
+
+// --------------------------------------------------------------------------------------------------------------------------------------------------------
+//
+//	AGameplayAbilityTargetActor_ActorPlacement
+//
+// --------------------------------------------------------------------------------------------------------------------------------------------------------
+
+AGameplayAbilityTargetActor_ActorPlacement::AGameplayAbilityTargetActor_ActorPlacement(const class FPostConstructInitializeProperties& PCIP)
+	: Super(PCIP)
+{
+}
+
+void AGameplayAbilityTargetActor_ActorPlacement::EndPlay(const EEndPlayReason::Type EndPlayReason)
+{
+	if (TurretReticleActor.IsValid())
+	{
+		TurretReticleActor.Get()->Destroy();
+	}
+
+	Super::EndPlay(EndPlayReason);
+}
+
+void AGameplayAbilityTargetActor_ActorPlacement::StartTargeting(UGameplayAbility* InAbility)
+{
+	Super::StartTargeting(InAbility);
+	if (AGameplayAbilityWorldReticle* CachedReticleActor = ReticleActor.Get())
+	{
+		if (AActor *TurretActor = GetWorld()->SpawnActor(TurretClass))
+		{
+			TurretReticleActor = GetWorld()->SpawnActor<AGameplayAbilityWorldReticle_ActorVisualization>();
+			TurretReticleActor->AttachRootComponentToActor(CachedReticleActor);
+			TurretReticleActor->InitializeReticleTurretInformation(TurretActor, TurretMaterial);
+			GetWorld()->DestroyActor(TurretActor);
+		}
+	}
+}
+
+//Might want to override this function to allow for a radius check against the ground, possibly including a height check. Or might want to do it in ground trace.
+//FHitResult AGameplayAbilityTargetActor_ActorPlacement::PerformTrace(AActor* InSourceActor) const
diff --git a/Engine/Source/Runtime/GameplayAbilities/Public/Abilities/GameplayAbilityTargetActor_ActorPlacement.h b/Engine/Source/Runtime/GameplayAbilities/Public/Abilities/GameplayAbilityTargetActor_ActorPlacement.h
new file mode 100644
--- /dev/null
+++ b/Engine/Source/Runtime/GameplayAbilities/Public/Abilities/GameplayAbilityTargetActor_ActorPlacement.h
@@ -0,0 +1,32 @@
+// Copyright 1998-2014 Epic Games, Inc. All Rights Reserved.
+
+#pragma once
+
+#include "GameplayAbilityTargetActor_GroundTrace.h"
+#include "GameplayAbilityWorldReticle_ActorVisualization.h"
+#include "GameplayAbilityTargetActor_ActorPlacement.generated.h"
+
+UCLASS(Blueprintable)
+class AGameplayAbilityTargetActor_ActorPlacement : public AGameplayAbilityTargetActor_GroundTrace
+{
+	GENERATED_UCLASS_BODY()
+
+	virtual void EndPlay(const EEndPlayReason::Type EndPlayReason) override;
+
+	virtual void StartTargeting(UGameplayAbility* InAbility) override;
+
+	/** Actor we intend to place. */
+	UPROPERTY(BlueprintReadWrite, EditAnywhere, meta = (ExposeOnSpawn = true), Category = Targeting)
+	UClass* TurretClass;		//Using a special class for replication purposes. (Not implemented yet)
+	
+	/** Override Material 0 on our turret's meshes with this material for visualization. */
+	UPROPERTY(BlueprintReadWrite, EditAnywhere, meta = (ExposeOnSpawn = true), Category = Projectile)
+	UMaterialInterface *TurretMaterial;
+
+	/** Visualization for the turret. */
+	TWeakObjectPtr<AGameplayAbilityWorldReticle_ActorVisualization> TurretReticleActor;
+
+protected:
+	//Consider whether we want size-checking here, or in GroundTrace
+	//virtual FHitResult PerformTrace(AActor* InSourceActor) const override;
+};
\ No newline at end of file
diff --git a/Engine/Source/Runtime/GameplayAbilities/Public/Abilities/GameplayAbilityWorldReticle_ActorVisualization.cpp b/Engine/Source/Runtime/GameplayAbilities/Public/Abilities/GameplayAbilityWorldReticle_ActorVisualization.cpp
new file mode 100644
--- /dev/null
+++ b/Engine/Source/Runtime/GameplayAbilities/Public/Abilities/GameplayAbilityWorldReticle_ActorVisualization.cpp
@@ -0,0 +1,72 @@
+// Copyright 1998-2013 Epic Games, Inc. All Rights Reserved.
+
+#include "AbilitySystemPrivatePCH.h"
+#include "GameplayAbilityTargetActor.h"
+#include "GameplayAbilityWorldReticle_ActorVisualization.h"
+
+// --------------------------------------------------------------------------------------------------------------------------------------------------------
+//
+//	AGameplayAbilityWorldReticle_ActorVisualization
+//
+// --------------------------------------------------------------------------------------------------------------------------------------------------------
+
+AGameplayAbilityWorldReticle_ActorVisualization::AGameplayAbilityWorldReticle_ActorVisualization(const class FPostConstructInitializeProperties& PCIP)
+	: Super(PCIP)
+{
+	//PrimaryActorTick.bCanEverTick = true;
+	//PrimaryActorTick.TickGroup = TG_PrePhysics;
+
+	CollisionComponent = PCIP.CreateDefaultSubobject<UCapsuleComponent>(this, TEXT("CollisionCapsule0"));
+	CollisionComponent->InitCapsuleSize(0.f, 0.f);
+	CollisionComponent->AlwaysLoadOnClient = true;
+	CollisionComponent->bAbsoluteScale = true;
+	//CollisionComponent->AlwaysLoadOnServer = true;
+	CollisionComponent->bCanEverAffectNavigation = false;
+	CollisionComponent->BodyInstance.bEnableCollision_DEPRECATED = false;
+	CollisionComponent->SetCollisionEnabled(ECollisionEnabled::NoCollision);
+
+	//TSubobjectPtr<USceneComponent> SceneComponent = PCIP.CreateDefaultSubobject<USceneComponent>(this, TEXT("RootComponent0"));
+	RootComponent = CollisionComponent;
+}
+
+void AGameplayAbilityWorldReticle_ActorVisualization::InitializeReticleTurretInformation(AActor* TurretActor, UMaterialInterface *TurretMaterial)
+{
+	if (TurretActor)
+	{
+		//Get components
+		TArray<UMeshComponent*> MeshComps;
+		USceneComponent* MyRoot = GetRootComponent();
+		TurretActor->GetComponents(MeshComps);
+		check(MyRoot);
+
+		for (UMeshComponent* MeshComp : MeshComps)
+		{
+			//Special case: If we don't clear the root component explicitly, the component will be destroyed along with the original turret actor.
+			if (MeshComp == TurretActor->GetRootComponent())
+			{
+				TurretActor->SetRootComponent(NULL);
+			}
+
+			//Disable collision on turret mesh parts so it doesn't interfere with aiming or any other client-side collision/prediction/physics stuff
+			MeshComp->SetCollisionEnabled(ECollisionEnabled::NoCollision);		//All mesh components are primitive components, so no cast is needed
+
+			//Move components from one actor to the other, attaching as needed. Hierarchy should not be important, but we can do fixups if it becomes important later.
+			MeshComp->DetachFromParent();
+			MeshComp->AttachTo(MyRoot);
+			MeshComp->Rename(nullptr, this);
+			if (TurretMaterial)
+			{
+				MeshComp->SetMaterial(0, TurretMaterial);
+			}
+		}
+	}
+}
+
+void AGameplayAbilityWorldReticle_ActorVisualization::EndPlay(const EEndPlayReason::Type EndPlayReason)
+{
+	TArray<UActorComponent*> CompArray;
+	GetComponents(CompArray);
+	CompArray.Num();
+
+	Super::EndPlay(EndPlayReason);
+}
diff --git a/Engine/Source/Runtime/GameplayAbilities/Public/Abilities/GameplayAbilityWorldReticle_ActorVisualization.h b/Engine/Source/Runtime/GameplayAbilities/Public/Abilities/GameplayAbilityWorldReticle_ActorVisualization.h
new file mode 100644
--- /dev/null
+++ b/Engine/Source/Runtime/GameplayAbilities/Public/Abilities/GameplayAbilityWorldReticle_ActorVisualization.h
@@ -0,0 +1,29 @@
+// Copyright 1998-2014 Epic Games, Inc. All Rights Reserved.
+#pragma once
+
+#include "GameplayAbilityWorldReticle_ActorVisualization.generated.h"
+
+class AGameplayAbilityTargetActor;
+
+/** This is a dummy reticle for internal use by turret placement tasks. It builds a custom visual model of the turret being placed. */
+UCLASS(notplaceable)
+class AGameplayAbilityWorldReticle_ActorVisualization : public AGameplayAbilityWorldReticle
+{
+	GENERATED_UCLASS_BODY()
+
+public:
+
+	void InitializeReticleTurretInformation(AActor* TurretActor, UMaterialInterface *TurretMaterial);
+
+	/** Hardcoded collision component, so other objects don't think they can collide with the visualization turret */
+	UPROPERTY()
+	TSubobjectPtr<class UCapsuleComponent> CollisionComponent;
+
+	UPROPERTY()
+	TArray<UActorComponent*> TurretComponents;
+
+	/** Overridable function called whenever this actor is being removed from a level */
+	virtual void EndPlay(const EEndPlayReason::Type EndPlayReason);
+
+protected:
+};
