Fixing non-unity

diff --git a/Engine/Source/Runtime/GameplayAbilities/Public/Abilities/GameplayAbilityWorldReticle_ActorVisualization.h b/Engine/Source/Runtime/GameplayAbilities/Public/Abilities/GameplayAbilityWorldReticle_ActorVisualization.h
--- a/Engine/Source/Runtime/GameplayAbilities/Public/Abilities/GameplayAbilityWorldReticle_ActorVisualization.h
+++ b/Engine/Source/Runtime/GameplayAbilities/Public/Abilities/GameplayAbilityWorldReticle_ActorVisualization.h
@@ -1,6 +1,8 @@
 // Copyright 1998-2014 Epic Games, Inc. All Rights Reserved.
 #pragma once
 
+#include "GameplayAbilityWorldReticle.h"
+
 #include "GameplayAbilityWorldReticle_ActorVisualization.generated.h"
 
 class AGameplayAbilityTargetActor;
