Debug bool on targeting tasks deactivated by default and exposed to users by request.

diff --git a/Engine/Source/Runtime/GameplayAbilities/Private/Abilities/GameplayAbilityTargetActor_Trace.cpp b/Engine/Source/Runtime/GameplayAbilities/Private/Abilities/GameplayAbilityTargetActor_Trace.cpp
--- a/Engine/Source/Runtime/GameplayAbilities/Private/Abilities/GameplayAbilityTargetActor_Trace.cpp
+++ b/Engine/Source/Runtime/GameplayAbilities/Private/Abilities/GameplayAbilityTargetActor_Trace.cpp
@@ -43,8 +43,6 @@ void AGameplayAbilityTargetActor_Trace::StartTargeting(UGameplayAbility* InAbili
 	Super::StartTargeting(InAbility);
 	SourceActor = InAbility->GetCurrentActorInfo()->Actor.Get();
 	
-	bDebug = true;
-
 	ReticleActor = GetWorld()->SpawnActor<AGameplayAbilityWorldReticle>(ReticleClass, GetActorLocation(), GetActorRotation());
 	if (AGameplayAbilityWorldReticle* CachedReticleActor = ReticleActor.Get())
 	{
diff --git a/Engine/Source/Runtime/GameplayAbilities/Public/Abilities/GameplayAbilityTargetActor.h b/Engine/Source/Runtime/GameplayAbilities/Public/Abilities/GameplayAbilityTargetActor.h
--- a/Engine/Source/Runtime/GameplayAbilities/Public/Abilities/GameplayAbilityTargetActor.h
+++ b/Engine/Source/Runtime/GameplayAbilities/Public/Abilities/GameplayAbilityTargetActor.h
@@ -70,6 +70,7 @@ public:
 	UPROPERTY(BlueprintReadOnly, Replicated, Category = Targeting)
 	AActor* SourceActor;
 
-	UPROPERTY(BlueprintReadOnly, Replicated, Category = Targeting)
+	/** Draw the debug information (if applicable) for this targeting actor. */
+	UPROPERTY(BlueprintReadWrite, EditAnywhere, Replicated, meta = (ExposeOnSpawn = true), Category = Targeting)
 	bool bDebug;
 };
\ No newline at end of file
