Adding an include to pacify GUBP.

diff --git a/Engine/Source/Runtime/GameplayAbilities/Private/GameplayAbilityTypes.cpp b/Engine/Source/Runtime/GameplayAbilities/Private/GameplayAbilityTypes.cpp
--- a/Engine/Source/Runtime/GameplayAbilities/Private/GameplayAbilityTypes.cpp
+++ b/Engine/Source/Runtime/GameplayAbilities/Private/GameplayAbilityTypes.cpp
@@ -5,6 +5,7 @@
 #include "GameplayTagContainer.h"
 #include "GameplayTagsModule.h"
 #include "GameplayEffectTypes.h"
+#include "GameplayAbility.h"
 
 //----------------------------------------------------------------------
 
