refactoring of some UAttributeSet methods.
Added way to remove gameplay effects based on ActiveGameplayEffectQuery
removed some dead code, minor cleanup.

diff --git a/Engine/Source/Runtime/GameplayAbilities/Private/Abilities/GameplayAbility.cpp b/Engine/Source/Runtime/GameplayAbilities/Private/Abilities/GameplayAbility.cpp
--- a/Engine/Source/Runtime/GameplayAbilities/Private/Abilities/GameplayAbility.cpp
+++ b/Engine/Source/Runtime/GameplayAbilities/Private/Abilities/GameplayAbility.cpp
@@ -639,8 +639,6 @@ FActiveGameplayEffectHandle UGameplayAbility::K2_ApplyGameplayEffectToOwner(cons
 	check(CurrentActorInfo);
 	check(CurrentSpecHandle.IsValid());
 
-	//CurrentActivationInfo = GetCurrentAbilitySpec()->ActivationInfo;
-
 	return ApplyGameplayEffectToOwner(CurrentSpecHandle, CurrentActorInfo, CurrentActivationInfo, GameplayEffect, GameplayEffectLevel);
 }
 
diff --git a/Engine/Source/Runtime/GameplayAbilities/Private/AbilitySystemComponent.cpp b/Engine/Source/Runtime/GameplayAbilities/Private/AbilitySystemComponent.cpp
--- a/Engine/Source/Runtime/GameplayAbilities/Private/AbilitySystemComponent.cpp
+++ b/Engine/Source/Runtime/GameplayAbilities/Private/AbilitySystemComponent.cpp
@@ -632,6 +632,11 @@ TArray<float> UAbilitySystemComponent::GetActiveEffectsDuration(const FActiveGam
 	return ActiveGameplayEffects.GetActiveEffectsDuration(Query);
 }
 
+void UAbilitySystemComponent::RemoveActiveEffects(const FActiveGameplayEffectQuery Query)
+{
+	return ActiveGameplayEffects.RemoveActiveEffects(Query);
+}
+
 void UAbilitySystemComponent::OnRestackGameplayEffects()
 {
 	ActiveGameplayEffects.RecalculateStacking();
diff --git a/Engine/Source/Runtime/GameplayAbilities/Private/AbilitySystemTestAttributeSet.cpp b/Engine/Source/Runtime/GameplayAbilities/Private/AbilitySystemTestAttributeSet.cpp
--- a/Engine/Source/Runtime/GameplayAbilities/Private/AbilitySystemTestAttributeSet.cpp
+++ b/Engine/Source/Runtime/GameplayAbilities/Private/AbilitySystemTestAttributeSet.cpp
@@ -23,7 +23,7 @@ UAbilitySystemTestAttributeSet::UAbilitySystemTestAttributeSet(const class FPost
 	NoStackAttribute = 0.f;
 }
 
-void UAbilitySystemTestAttributeSet::PreAttributeModify(struct FGameplayEffectModCallbackData &Data)
+bool UAbilitySystemTestAttributeSet::PreGameplayEffectExecute(struct FGameplayEffectModCallbackData &Data)
 {
 	static UProperty *HealthProperty = FindFieldChecked<UProperty>(UAbilitySystemTestAttributeSet::StaticClass(), GET_MEMBER_NAME_CHECKED(UAbilitySystemTestAttributeSet, Health));
 	static UProperty *DamageProperty = FindFieldChecked<UProperty>(UAbilitySystemTestAttributeSet::StaticClass(), GET_MEMBER_NAME_CHECKED(UAbilitySystemTestAttributeSet, Damage));
@@ -79,9 +79,11 @@ void UAbilitySystemTestAttributeSet::PreAttributeModify(struct FGameplayEffectMo
 
 		// At this point, the Magnitude of the applied damage may have been modified by us. We still do the translation to Health in UAbilitySystemTestAttributeSet::PostAttributeModify.
 	}
+
+	return true;
 }
 
-void UAbilitySystemTestAttributeSet::PostAttributeModify(const struct FGameplayEffectModCallbackData &Data)
+void UAbilitySystemTestAttributeSet::PostGameplayEffectExecute(const struct FGameplayEffectModCallbackData &Data)
 {
 	static UProperty *HealthProperty = FindFieldChecked<UProperty>(UAbilitySystemTestAttributeSet::StaticClass(), GET_MEMBER_NAME_CHECKED(UAbilitySystemTestAttributeSet, Health));
 	static UProperty *DamageProperty = FindFieldChecked<UProperty>(UAbilitySystemTestAttributeSet::StaticClass(), GET_MEMBER_NAME_CHECKED(UAbilitySystemTestAttributeSet, Damage));
diff --git a/Engine/Source/Runtime/GameplayAbilities/Private/AttributeSet.cpp b/Engine/Source/Runtime/GameplayAbilities/Private/AttributeSet.cpp
--- a/Engine/Source/Runtime/GameplayAbilities/Private/AttributeSet.cpp
+++ b/Engine/Source/Runtime/GameplayAbilities/Private/AttributeSet.cpp
@@ -12,6 +12,7 @@ void FGameplayAttribute::SetNumericValueChecked(const float NewValue, class UAtt
 {
 	UNumericProperty *NumericProperty = CastChecked<UNumericProperty>(Attribute);
 	void * ValuePtr = NumericProperty->ContainerPtrToValuePtr<void>(Dest);
+	Dest->PreAttributeSet(*this, NewValue);
 	NumericProperty->SetFloatingPointPropertyValue(ValuePtr, NewValue);
 }
 
@@ -52,36 +53,6 @@ void UAttributeSet::InitFromMetaDataTable(const UDataTable* DataTable)
 	PrintDebug();
 }
 
-void UAttributeSet::ClampFromMetaDataTable(const UDataTable *DataTable)
-{
-	// Looking up a datatable row everytime a value changes in order to apply clamping pretty much sucks.
-	static const FString Context = FString(TEXT("UAttribute::BindToMetaDataTable"));
-
-	UE_LOG(LogActorComponent, Warning, TEXT("UAttribute::ClampFromMetaDataTable. UAttribute %s from UDataTable %s"), *GetName(), *DataTable->GetName());
-
-	for( TFieldIterator<UProperty> It(GetClass(), EFieldIteratorFlags::IncludeSuper) ; It ; ++It )
-	{
-		UProperty *Property = *It;
-		UNumericProperty *NumericProperty = Cast<UNumericProperty>(Property);
-		if (NumericProperty)
-		{
-			FString RowNameStr = FString::Printf(TEXT("%s.%s"), *Property->GetOuter()->GetName(), *Property->GetName());
-			FAttributeMetaData * MetaData = DataTable->FindRow<FAttributeMetaData>(FName(*RowNameStr), Context, false);
-			if (MetaData)
-			{
-				float CurrentValue = NumericProperty->GetFloatingPointPropertyValue(this);
-				CurrentValue = FMath::Clamp(CurrentValue, MetaData->MinValue, MetaData->MaxValue);
-
-				UE_LOG(LogActorComponent, Warning, TEXT("   Found Row for %s. Clamping: %.2f <%.2f, %.2f>"), *RowNameStr, CurrentValue, MetaData->MinValue, MetaData->MaxValue );
-
-				NumericProperty->SetFloatingPointPropertyValue(NumericProperty->ContainerPtrToValuePtr<void>(this), CurrentValue);
-			}
-		}
-	}
-
-	PrintDebug();
-}
-
 UAbilitySystemComponent* UAttributeSet::GetOwningAbilitySystemComponent() const
 {
 	return UAbilitySystemGlobals::GetAbilitySystemComponentFromActor(CastChecked<AActor>(GetOuter()));
diff --git a/Engine/Source/Runtime/GameplayAbilities/Private/GameplayEffect.cpp b/Engine/Source/Runtime/GameplayAbilities/Private/GameplayEffect.cpp
--- a/Engine/Source/Runtime/GameplayAbilities/Private/GameplayEffect.cpp
+++ b/Engine/Source/Runtime/GameplayAbilities/Private/GameplayEffect.cpp
@@ -1350,9 +1350,6 @@ void FActiveGameplayEffectsContainer::ExecuteActiveEffectsFrom(const FGameplayEf
 				continue;
 			}
 
-			// Todo: Tags/application checks here - make sure we can still apply
-			InvokeGameplayCueExecute = true;
-
 			ABILITY_LOG_SCOPE(TEXT("Executing Attribute Mod %s"), *Mod.ToSimpleString());
 
 			// First, evaluate all of our data 
@@ -1364,8 +1361,16 @@ void FActiveGameplayEffectsContainer::ExecuteActiveEffectsFrom(const FGameplayEf
 			/** This should apply 'gameplay effect specific' rules, such as life steal, shields, etc */
 			Mod.Aggregator.Get()->PreEvaluate(ExecuteData);
 
-			/** This should apply 'gamewide' rules. Such as clamping Health to MaxHealth or granting +3 health for every point of strength, etc */
-			AttributeSet->PreAttributeModify(ExecuteData);
+			/** This should apply 'gamewide' rules. Such as clamping Health to MaxHealth or granting +3 health for every point of strength, etc 
+			 *	PreAttributeModify can return false to 'throw out' this modification.
+			 */
+			if (AttributeSet->PreGameplayEffectExecute(ExecuteData) == false)
+			{
+				continue;
+			}
+
+			// Todo: Tags/application checks here - make sure we can still apply
+			InvokeGameplayCueExecute = true;
 
 			// Do we have active GE's that are already modifying this?
 			FAggregatorRef *RefPtr = OngoingPropertyEffects.Find(Mod.Info.Attribute);
@@ -1391,7 +1396,7 @@ void FActiveGameplayEffectsContainer::ExecuteActiveEffectsFrom(const FGameplayEf
 			Mod.Aggregator.Get()->PostEvaluate(ExecuteData);
 
 			/** This should apply 'gamewide' rules. Such as clamping Health to MaxHealth or granting +3 health for every point of strength, etc */
-			AttributeSet->PostAttributeModify(ExecuteData);
+			AttributeSet->PostGameplayEffectExecute(ExecuteData);
 			
 		}
 		else if(Mod.Info.ModifierType == EGameplayMod::ActiveGE)
@@ -1597,6 +1602,8 @@ void FActiveGameplayEffectsContainer::OnPropertyAggregatorDirty(FAggregator* Agg
 	float NewValue = Aggregator->Evaluate().Magnitude;
 
 	InternalUpdateNumericalAttribute(Attribute, NewValue, nullptr);
+
+	
 }
 
 void FActiveGameplayEffectsContainer::InternalUpdateNumericalAttribute(FGameplayAttribute Attribute, float NewValue, const FGameplayEffectModCallbackData* ModData)
@@ -2051,12 +2058,9 @@ TArray<float> FActiveGameplayEffectsContainer::GetActiveEffectsTimeRemaining(con
 
 	for (const FActiveGameplayEffect& Effect : GameplayEffects)
 	{
-		if (Query.TagContainer)
+		if (!Query.Matches(Effect))
 		{
-			if (!Effect.Spec.Def->OwnedTagsContainer.MatchesAny(*Query.TagContainer, false))
-			{
-				continue;
-			}
+			continue;
 		}
 
 		float Elapsed = CurrentTime - Effect.StartWorldTime;
@@ -2077,12 +2081,9 @@ TArray<float> FActiveGameplayEffectsContainer::GetActiveEffectsDuration(const FA
 
 	for (const FActiveGameplayEffect &Effect : GameplayEffects)
 	{
-		if (Query.TagContainer)
+		if (!Query.Matches(Effect))
 		{
-			if (!Effect.Spec.Def->OwnedTagsContainer.MatchesAny(*Query.TagContainer, false))
-			{
-				continue;
-			}
+			continue;
 		}
 
 		ReturnList.Add(Effect.GetDuration());
@@ -2092,6 +2093,19 @@ TArray<float> FActiveGameplayEffectsContainer::GetActiveEffectsDuration(const FA
 	return ReturnList;
 }
 
+void FActiveGameplayEffectsContainer::RemoveActiveEffects(const FActiveGameplayEffectQuery Query)
+{
+	for (int32 idx=0; idx < GameplayEffects.Num(); ++idx)
+	{
+		const FActiveGameplayEffect& Effect = GameplayEffects[idx];
+		if (Query.Matches(Effect))
+		{
+			InternalRemoveActiveGameplayEffect(idx);
+			idx--;
+		}
+	}
+}
+
 FOnGameplayEffectTagCountChanged& FActiveGameplayEffectsContainer::RegisterGameplayTagEvent(FGameplayTag Tag)
 {
 	return GameplayTagCountContainer.GameplayTagEventMap.FindOrAdd(Tag);
@@ -2346,4 +2360,34 @@ void FGameplayTagCountContainer::UpdateTagMap(const FGameplayTagContainer& Conta
 }
 
 
+// -----------------------------------------------------------------
+
+bool FActiveGameplayEffectQuery::Matches(const FActiveGameplayEffect& Effect) const
+{
+	if (TagContainer)
+	{
+		if (!Effect.Spec.Def->OwnedTagsContainer.MatchesAny(*TagContainer, false))
+		{
+			return false;
+		}
+	}	
+	
+	if (ModifyingAttribute.IsValid())
+	{
+		bool FailedModifyingAttributeCheck = true;
+		for (const FModifierSpec& Modifier : Effect.Spec.Modifiers)
+		{
+			if (Modifier.Info.Attribute == ModifyingAttribute && Modifier.Info.ModifierType == EGameplayMod::Attribute)
+			{
+				FailedModifyingAttributeCheck = false;
+				break;
+			}
+		}
+		if (FailedModifyingAttributeCheck)
+		{
+			return false;
+		}
+	}
 
+	return true;
+}
diff --git a/Engine/Source/Runtime/GameplayAbilities/Public/AbilitySystemComponent.h b/Engine/Source/Runtime/GameplayAbilities/Public/AbilitySystemComponent.h
--- a/Engine/Source/Runtime/GameplayAbilities/Public/AbilitySystemComponent.h
+++ b/Engine/Source/Runtime/GameplayAbilities/Public/AbilitySystemComponent.h
@@ -275,6 +275,9 @@ class GAMEPLAYABILITIES_API UAbilitySystemComponent : public UActorComponent, pu
 
 	TArray<float> GetActiveEffectsDuration(const FActiveGameplayEffectQuery Query) const;
 
+	/** Removes all active effects that match given query */
+	void RemoveActiveEffects(const FActiveGameplayEffectQuery Query);
+
 	void OnRestackGameplayEffects();
 
 	// --------------------------------------------
diff --git a/Engine/Source/Runtime/GameplayAbilities/Public/AbilitySystemTestAttributeSet.h b/Engine/Source/Runtime/GameplayAbilities/Public/AbilitySystemTestAttributeSet.h
--- a/Engine/Source/Runtime/GameplayAbilities/Public/AbilitySystemTestAttributeSet.h
+++ b/Engine/Source/Runtime/GameplayAbilities/Public/AbilitySystemTestAttributeSet.h
@@ -69,6 +69,6 @@ class GAMEPLAYABILITIES_API UAbilitySystemTestAttributeSet : public UAttributeSe
 	UPROPERTY(EditAnywhere, BlueprintReadWrite, Category = "AttributeTest")
 	mutable float	NoStackAttribute;
 
-	virtual void PreAttributeModify(struct FGameplayEffectModCallbackData &Data) override;
-	virtual void PostAttributeModify(const struct FGameplayEffectModCallbackData &Data) override;
+	virtual bool PreGameplayEffectExecute(struct FGameplayEffectModCallbackData &Data) override;
+	virtual void PostGameplayEffectExecute(const struct FGameplayEffectModCallbackData &Data) override;
 };
diff --git a/Engine/Source/Runtime/GameplayAbilities/Public/AttributeSet.h b/Engine/Source/Runtime/GameplayAbilities/Public/AttributeSet.h
--- a/Engine/Source/Runtime/GameplayAbilities/Public/AttributeSet.h
+++ b/Engine/Source/Runtime/GameplayAbilities/Public/AttributeSet.h
@@ -56,6 +56,11 @@ struct GAMEPLAYABILITIES_API FGameplayAttribute
 		return Attribute->GetName();
 	}
 
+	bool IsValid() const
+	{
+		return Attribute != nullptr;
+	}
+
 private:
 	friend class FAttributePropertyDetails;
 
@@ -75,15 +80,33 @@ public:
 		return true;
 	}
 
-	/** Called just before modifying the value of an attribute. AttributeSet can make additional modifications here. */
-	virtual void PreAttributeModify(struct FGameplayEffectModCallbackData &Data) { }
+	/**
+	 *	Called just before modifying the value of an attribute. AttributeSet can make additional modifications here. Return true to continue, or false to throw out the modification.
+	 *	Note this is only called during an 'execute'. E.g., a modification to the 'base value' of an attribute. It is not called during an application of a GameplayEffect, such as a 5 ssecond +10 movement speed buff.
+	 */	
+	virtual bool PreGameplayEffectExecute(struct FGameplayEffectModCallbackData &Data) { return true; }
+	
 	
-	/** Called just after modifying the value of an attribute. No more changes can be made */
-	virtual void PostAttributeModify(const struct FGameplayEffectModCallbackData &Data) { }
+	/**
+	 *	Called just before a GameplayEffect is executed to modify the base value of an attribute. No more changes can be made.
+	 *	Note this is only called during an 'execute'. E.g., a modification to the 'base value' of an attribute. It is not called during an application of a GameplayEffect, such as a 5 ssecond +10 movement speed buff.
+	 */
+	virtual void PostGameplayEffectExecute(const struct FGameplayEffectModCallbackData &Data) { }
+
+	/**
+	 *	An "On Aggregator Change" type of event could go here, and that could be called when active gamepaly effects are added or removed to an attribute aggregator.
+	 *	It is difficult to give all the information in these cases though - aggregators can change for many reasons: being added, being removed, being modified, having a modifier change, immunity, stacking rules, etc.
+	 */
 
-	void InitFromMetaDataTable(const UDataTable* DataTable);
 
-	void ClampFromMetaDataTable(const UDataTable *DataTable);
+	/**
+	 *	Called just before any modification happens to an attribute. This is lower level than PreAttributeModify/PostAttribute modify.
+	 *	There is no additional context provided here since anything can trigger this. Executed effects, duration based effects, effects being removed, immunity being applied, stacking rules changing, etc.
+	 *	This function is meant to enforce things like "Health = Clamp(Health, 0, MaxHealth)" and NOT things like "trigger this extra thing if damage is applied, etc".
+	 */
+	virtual void PreAttributeSet(const FGameplayAttribute& Attribute, float NewValue) { }
+
+	void InitFromMetaDataTable(const UDataTable* DataTable);
 
 	UAbilitySystemComponent* GetOwningAbilitySystemComponent() const;
 
diff --git a/Engine/Source/Runtime/GameplayAbilities/Public/GameplayEffect.h b/Engine/Source/Runtime/GameplayAbilities/Public/GameplayEffect.h
--- a/Engine/Source/Runtime/GameplayAbilities/Public/GameplayEffect.h
+++ b/Engine/Source/Runtime/GameplayAbilities/Public/GameplayEffect.h
@@ -995,9 +995,12 @@ struct FActiveGameplayEffect : public FFastArraySerializerItem
 	}
 };
 
-/** Generic querying data structure for active GameplayEffects. Lets us ask things like: ** 
- *		-Give me duration/magnitude of active gameplay effects with these tags
- *		-
+/**
+ *	Generic querying data structure for active GameplayEffects. Lets us ask things like:
+ *		Give me duration/magnitude of active gameplay effects with these tags
+ *		Give me handles to all activate gameplay effects modifying this attribute.
+ *		
+ *	Any requirements specified in the query are required: must meet "all" not "one".
  */
 USTRUCT()
 struct FActiveGameplayEffectQuery
@@ -1009,12 +1012,18 @@ struct FActiveGameplayEffectQuery
 	{
 	}
 
-	FActiveGameplayEffectQuery(const FGameplayTagContainer * InTagContainer)
+	FActiveGameplayEffectQuery(const FGameplayTagContainer* InTagContainer, FGameplayAttribute InModifyingAttribute = FGameplayAttribute())
 		: TagContainer(InTagContainer)
+		, ModifyingAttribute(InModifyingAttribute)
 	{
 	}
 
-	const FGameplayTagContainer *	TagContainer;
+	bool Matches(const FActiveGameplayEffect& Effect) const;
+
+	const FGameplayTagContainer* TagContainer;
+
+	/** Matches on GameplayEffects which modify given attribute */
+	FGameplayAttribute ModifyingAttribute;
 };
 
 /**
@@ -1111,6 +1120,8 @@ struct FActiveGameplayEffectsContainer : public FFastArraySerializer
 
 	TArray<float> GetActiveEffectsDuration(const FActiveGameplayEffectQuery Query) const;
 
+	void RemoveActiveEffects(const FActiveGameplayEffectQuery Query);
+
 	int32 GetGameStateTime() const;
 
 	float GetWorldTime() const;
