rename UAttributeset::PreAttributeSet to ::PreAttributeChange

diff --git a/Engine/Source/Runtime/GameplayAbilities/Private/AttributeSet.cpp b/Engine/Source/Runtime/GameplayAbilities/Private/AttributeSet.cpp
--- a/Engine/Source/Runtime/GameplayAbilities/Private/AttributeSet.cpp
+++ b/Engine/Source/Runtime/GameplayAbilities/Private/AttributeSet.cpp
@@ -12,7 +12,7 @@ void FGameplayAttribute::SetNumericValueChecked(const float NewValue, class UAtt
 {
 	UNumericProperty *NumericProperty = CastChecked<UNumericProperty>(Attribute);
 	void * ValuePtr = NumericProperty->ContainerPtrToValuePtr<void>(Dest);
-	Dest->PreAttributeSet(*this, NewValue);
+	Dest->PreAttributeChange(*this, NewValue);
 	NumericProperty->SetFloatingPointPropertyValue(ValuePtr, NewValue);
 }
 
diff --git a/Engine/Source/Runtime/GameplayAbilities/Public/AttributeSet.h b/Engine/Source/Runtime/GameplayAbilities/Public/AttributeSet.h
--- a/Engine/Source/Runtime/GameplayAbilities/Public/AttributeSet.h
+++ b/Engine/Source/Runtime/GameplayAbilities/Public/AttributeSet.h
@@ -104,7 +104,7 @@ public:
 	 *	There is no additional context provided here since anything can trigger this. Executed effects, duration based effects, effects being removed, immunity being applied, stacking rules changing, etc.
 	 *	This function is meant to enforce things like "Health = Clamp(Health, 0, MaxHealth)" and NOT things like "trigger this extra thing if damage is applied, etc".
 	 */
-	virtual void PreAttributeSet(const FGameplayAttribute& Attribute, float NewValue) { }
+	virtual void PreAttributeChange(const FGameplayAttribute& Attribute, float NewValue) { }
 
 	void InitFromMetaDataTable(const UDataTable* DataTable);
 
