fix prediction key problem with target data

diff --git a/Engine/Source/Runtime/GameplayAbilities/Private/GameplayPrediction.cpp b/Engine/Source/Runtime/GameplayAbilities/Private/GameplayPrediction.cpp
--- a/Engine/Source/Runtime/GameplayAbilities/Private/GameplayPrediction.cpp
+++ b/Engine/Source/Runtime/GameplayAbilities/Private/GameplayPrediction.cpp
@@ -221,6 +221,7 @@ FScopedPredictionWindow::FScopedPredictionWindow(UGameplayAbility* GameplayAbili
 		{
 			ActivationInfo.SetPredictionKey(Owner->ScopedPedictionKey);
 			ActivationInfo.ActivationMode = EGameplayAbilityActivationMode::Predicting;
+			ScopedPredictionKey = Owner->ScopedPedictionKey;
 		}
 		else
 		{
