Don't replicate bActive on AbilitySystemComponent: clients need to control this predictively.

diff --git a/Engine/Source/Runtime/GameplayAbilities/Private/AbilitySystemComponent.cpp b/Engine/Source/Runtime/GameplayAbilities/Private/AbilitySystemComponent.cpp
--- a/Engine/Source/Runtime/GameplayAbilities/Private/AbilitySystemComponent.cpp
+++ b/Engine/Source/Runtime/GameplayAbilities/Private/AbilitySystemComponent.cpp
@@ -668,7 +668,8 @@ void UAbilitySystemComponent::TickingTaskEnded(UAbilityTask* Task)
 
 void UAbilitySystemComponent::GetLifetimeReplicatedProps(TArray< FLifetimeProperty > & OutLifetimeProps) const
 {
-	Super::GetLifetimeReplicatedProps(OutLifetimeProps);
+	// Intentionally not calling super: We do not want to replicate bActive which controls ticking. We sometimes need to tick on client predictively.
+	
 
 	DOREPLIFETIME(UAbilitySystemComponent, SpawnedAttributes);
 	DOREPLIFETIME(UAbilitySystemComponent, ActiveGameplayEffects);
