GroundTrace and ActorPlacement improvements so that placed objects find open ground whenever possible.

diff --git a/Engine/Source/Runtime/GameplayAbilities/Private/Abilities/GameplayAbilityTargetActor_GroundTrace.cpp b/Engine/Source/Runtime/GameplayAbilities/Private/Abilities/GameplayAbilityTargetActor_GroundTrace.cpp
--- a/Engine/Source/Runtime/GameplayAbilities/Private/Abilities/GameplayAbilityTargetActor_GroundTrace.cpp
+++ b/Engine/Source/Runtime/GameplayAbilities/Private/Abilities/GameplayAbilityTargetActor_GroundTrace.cpp
@@ -13,8 +13,114 @@
 // --------------------------------------------------------------------------------------------------------------------------------------------------------
 
 AGameplayAbilityTargetActor_GroundTrace::AGameplayAbilityTargetActor_GroundTrace(const class FPostConstructInitializeProperties& PCIP)
-	: Super(PCIP)
+: Super(PCIP)
 {
+	CollisionHeightOffset = 0.0f;
+}
+
+void AGameplayAbilityTargetActor_GroundTrace::StartTargeting(UGameplayAbility* InAbility)
+{
+	//CollisionShape starts as "line," which is correct as a default for our purposes
+	if (CollisionRadius > 0.0f)
+	{
+		//Check CollisionHeight here because the shape code doesn't. It is used as half-height
+		if ((CollisionHeight * 0.5f) > CollisionRadius)
+		{
+			CollisionShape = FCollisionShape::MakeCapsule(CollisionRadius, CollisionHeight * 0.5f);
+			CollisionHeightOffset = CollisionHeight * 0.5f;
+		}
+		else
+		{
+			CollisionShape = FCollisionShape::MakeSphere(CollisionRadius);
+			CollisionHeight = 0.0f;
+			CollisionHeightOffset = CollisionRadius;
+		}
+	}
+	Super::StartTargeting(InAbility);
+}
+
+bool AGameplayAbilityTargetActor_GroundTrace::AdjustCollisionResultForShape(const FVector OriginalStartPoint, const FVector OriginalEndPoint, const FCollisionQueryParams Params, FHitResult& OutHitResult) const
+{
+	UWorld *ThisWorld = GetWorld();
+	//Pull back toward player to find a better spot, accounting for the width of our object
+	FVector Movement = (OriginalEndPoint - OriginalStartPoint);
+	FVector MovementDirection = Movement.SafeNormal();
+	float MovementMagnitude2D = Movement.Size2D();
+
+	if (bDebug)
+	{
+		if (CollisionHeight > 0.0f)
+		{
+			DrawDebugCapsule(ThisWorld, OriginalEndPoint, CollisionHeight * 0.5f, CollisionRadius, FQuat::Identity, FColor::Black);
+		}
+		else
+		{
+			DrawDebugSphere(ThisWorld, OriginalEndPoint, CollisionRadius, 8, FColor::Black);
+		}
+	}
+
+	if ((MovementMagnitude2D < (CollisionRadius * 2.0f)) || (CollisionRadius <= 1.0f))
+	{
+		return false;		//Bad case!
+	}
+
+	float IncrementSize = FMath::Clamp<float>(CollisionRadius * 0.5f, 25.0f, 250.0f);
+	float LerpIncrement = IncrementSize / MovementMagnitude2D;
+	FHitResult LocalResult;
+	FVector TraceStart;
+	FVector TraceEnd;
+	//This needs to ramp up - the first few increments should be small, then we should start moving in larger steps.
+	for (float LerpValue = CollisionRadius / MovementMagnitude2D; LerpValue < 1.0f; LerpValue += LerpIncrement)
+	{
+		TraceEnd = TraceStart = OriginalEndPoint - (LerpValue * Movement);
+		TraceEnd.Z -= 99999.0f;
+		ThisWorld->SweepSingle(LocalResult, TraceStart, TraceEnd, FQuat::Identity, ECC_WorldStatic, CollisionShape, Params);
+		if (!LocalResult.bStartPenetrating)
+		{
+			if (!LocalResult.bBlockingHit)
+			{
+				//This is probably off the map and should not be considered valid. This should not happen in a non-debug map.
+				if (bDebug)
+				{
+					if (CollisionHeight > 0.0f)
+					{
+						DrawDebugCapsule(ThisWorld, LocalResult.Location, CollisionHeight * 0.5f, CollisionRadius, FQuat::Identity, FColor::Yellow);
+					}
+					else
+					{
+						DrawDebugSphere(ThisWorld, LocalResult.Location, CollisionRadius, 8, FColor::Yellow);
+					}
+				}
+				continue;
+				//LocalResult.Location = TraceStart;
+			}
+			if (bDebug)
+			{
+				if (CollisionHeight > 0.0f)
+				{
+					DrawDebugCapsule(ThisWorld, LocalResult.Location, CollisionHeight * 0.5f, CollisionRadius, FQuat::Identity, FColor::Green);
+				}
+				else
+				{
+					DrawDebugSphere(ThisWorld, LocalResult.Location, CollisionRadius, 8, FColor::Green);
+				}
+			}
+			OutHitResult = LocalResult;
+			return true;
+		}
+		if (bDebug)
+		{
+			if (CollisionHeight > 0.0f)
+			{
+				DrawDebugCapsule(ThisWorld, TraceStart, CollisionHeight * 0.5f, CollisionRadius, FQuat::Identity, FColor::Red);
+			}
+			else
+			{
+				DrawDebugSphere(ThisWorld, TraceStart, CollisionRadius, 8, FColor::Red);
+			}
+		}
+	}
+	return false;
 }
 
 FHitResult AGameplayAbilityTargetActor_GroundTrace::PerformTrace(AActor* InSourceActor) const
@@ -52,6 +158,9 @@ FHitResult AGameplayAbilityTargetActor_GroundTrace::PerformTrace(AActor* InSourc
 	// ------------------------------------------------------
 
 	FHitResult ReturnHitResult;
+	FVector LineTraceStart = TraceStart;
+	FVector LineTraceEnd = TraceEnd;
+	//Use a line trace initially to see where the player is actually pointing
 	InSourceActor->GetWorld()->LineTraceSingle(ReturnHitResult, TraceStart, TraceEnd, ECC_WorldStatic, Params);
 	//Default to end of trace line if we don't hit anything.
 	if (!ReturnHitResult.bBlockingHit)
@@ -62,15 +171,22 @@ FHitResult AGameplayAbilityTargetActor_GroundTrace::PerformTrace(AActor* InSourc
 	//Second trace, straight down. Consider using InSourceActor->GetWorld()->NavigationSystem->ProjectPointToNavigation() instead of tracing in the case of movement abilities (flag/bool).
 	TraceStart = ReturnHitResult.Location - AimDirection;		//Pull back very slightly to avoid scraping down walls
 	TraceEnd = TraceStart;
-	TraceEnd[2] -= 999999.0f;
-	if (ReturnHitResult.Actor.IsValid())
-	{
-		Params.AddIgnoredActor(ReturnHitResult.Actor.Get());
-	}
+	TraceStart.Z += CollisionHeightOffset;
+	TraceEnd.Z -= 99999.0f;
 	InSourceActor->GetWorld()->LineTraceSingle(ReturnHitResult, TraceStart, TraceEnd, ECC_WorldStatic, Params);
-	if (!ReturnHitResult.bBlockingHit)
+	//if (!ReturnHitResult.bBlockingHit) then our endpoint may be off the map. Hopefully this is only possible in debug maps.
+
+	//Use collision shape to find a valid ground spot, if appropriate
+	if (CollisionShape.ShapeType != ECollisionShape::Line)
 	{
-		ReturnHitResult.Location = TraceEnd;
+		ReturnHitResult.Location.Z += CollisionHeightOffset;		//Rise up out of the ground
+		TraceStart = InSourceActor->GetActorLocation();
+		TraceEnd = ReturnHitResult.Location;
+		TraceStart.Z += CollisionHeightOffset;
+		if (AdjustCollisionResultForShape(TraceStart, TraceEnd, Params, ReturnHitResult))
+		{
+			ReturnHitResult.Location.Z -= CollisionHeightOffset;	//Undo the artificial height adjustment
+		}
 	}
 
 	if (AActor* LocalReticleActor = ReticleActor.Get())
diff --git a/Engine/Source/Runtime/GameplayAbilities/Private/Abilities/Tasks/AbilityTask_SpawnActor.cpp b/Engine/Source/Runtime/GameplayAbilities/Private/Abilities/Tasks/AbilityTask_SpawnActor.cpp
--- a/Engine/Source/Runtime/GameplayAbilities/Private/Abilities/Tasks/AbilityTask_SpawnActor.cpp
+++ b/Engine/Source/Runtime/GameplayAbilities/Private/Abilities/Tasks/AbilityTask_SpawnActor.cpp
@@ -11,14 +11,16 @@ UAbilityTask_SpawnActor::UAbilityTask_SpawnActor(const class FPostConstructIniti
 	
 }
 
-UAbilityTask_SpawnActor* UAbilityTask_SpawnActor::SpawnActor(UObject* WorldContextObject, TSubclassOf<AActor> InClass)
+UAbilityTask_SpawnActor* UAbilityTask_SpawnActor::SpawnActor(UObject* WorldContextObject, FGameplayAbilityTargetDataHandle TargetData, TSubclassOf<AActor> InClass)
 {
-	return NewTask<UAbilityTask_SpawnActor>(WorldContextObject);
+	auto MyObj = NewTask<UAbilityTask_SpawnActor>(WorldContextObject);
+	MyObj->CachedTargetDataHandle = TargetData;
+	return MyObj;
 }
 
 // ---------------------------------------------------------------------------------------
 
-bool UAbilityTask_SpawnActor::BeginSpawningActor(UObject* WorldContextObject, TSubclassOf<AActor> InClass, AActor*& SpawnedActor)
+bool UAbilityTask_SpawnActor::BeginSpawningActor(UObject* WorldContextObject, FGameplayAbilityTargetDataHandle TargetData, TSubclassOf<AActor> InClass, AActor*& SpawnedActor)
 {
 	if (Ability.IsValid() && Ability.Get()->GetCurrentActorInfo()->IsNetAuthority())
 	{
@@ -29,11 +31,24 @@ bool UAbilityTask_SpawnActor::BeginSpawningActor(UObject* WorldContextObject, TS
 	return (SpawnedActor != nullptr);
 }
 
-void UAbilityTask_SpawnActor::FinishSpawningActor(UObject* WorldContextObject, AActor* SpawnedActor)
+void UAbilityTask_SpawnActor::FinishSpawningActor(UObject* WorldContextObject, FGameplayAbilityTargetDataHandle TargetData, AActor* SpawnedActor)
 {
 	if (SpawnedActor)
 	{
-		const FTransform SpawnTransform = AbilitySystemComponent->GetOwner()->GetTransform();
+		FTransform SpawnTransform = AbilitySystemComponent->GetOwner()->GetTransform();
+
+		if (FGameplayAbilityTargetData* LocationData = CachedTargetDataHandle.Get(0))		//Hardcode to use data 0. It's OK if data isn't useful/valid.
+		{
+			//Set location. Rotation is unaffected.
+			if (LocationData->HasHitResult())
+			{
+				SpawnTransform.SetLocation(LocationData->GetHitResult()->Location);
+			}
+			else if (LocationData->HasEndPoint())
+			{
+				SpawnTransform.SetLocation(LocationData->GetEndPoint());
+			}
+		}
 
 		SpawnedActor->FinishSpawning(SpawnTransform);
 
diff --git a/Engine/Source/Runtime/GameplayAbilities/Public/Abilities/GameplayAbilityTargetActor_ActorPlacement.cpp b/Engine/Source/Runtime/GameplayAbilities/Public/Abilities/GameplayAbilityTargetActor_ActorPlacement.cpp
--- a/Engine/Source/Runtime/GameplayAbilities/Public/Abilities/GameplayAbilityTargetActor_ActorPlacement.cpp
+++ b/Engine/Source/Runtime/GameplayAbilities/Public/Abilities/GameplayAbilityTargetActor_ActorPlacement.cpp
@@ -29,15 +29,20 @@ void AGameplayAbilityTargetActor_ActorPlacement::EndPlay(const EEndPlayReason::T
 void AGameplayAbilityTargetActor_ActorPlacement::StartTargeting(UGameplayAbility* InAbility)
 {
 	Super::StartTargeting(InAbility);
+	if (AActor *TurretActor = GetWorld()->SpawnActor(TurretClass))
+	{
+		TurretReticleActor = GetWorld()->SpawnActor<AGameplayAbilityWorldReticle_ActorVisualization>();
+		TurretReticleActor->InitializeReticleTurretInformation(TurretActor, TurretMaterial);
+		GetWorld()->DestroyActor(TurretActor);
+	}
 	if (AGameplayAbilityWorldReticle* CachedReticleActor = ReticleActor.Get())
 	{
-		if (AActor *TurretActor = GetWorld()->SpawnActor(TurretClass))
-		{
-			TurretReticleActor = GetWorld()->SpawnActor<AGameplayAbilityWorldReticle_ActorVisualization>();
-			TurretReticleActor->AttachRootComponentToActor(CachedReticleActor);
-			TurretReticleActor->InitializeReticleTurretInformation(TurretActor, TurretMaterial);
-			GetWorld()->DestroyActor(TurretActor);
-		}
+		TurretReticleActor->AttachRootComponentToActor(CachedReticleActor);
+	}
+	else
+	{
+		ReticleActor = TurretReticleActor;
+		TurretReticleActor = nullptr;
 	}
 }
 
diff --git a/Engine/Source/Runtime/GameplayAbilities/Public/Abilities/GameplayAbilityTargetActor_GroundTrace.h b/Engine/Source/Runtime/GameplayAbilities/Public/Abilities/GameplayAbilityTargetActor_GroundTrace.h
--- a/Engine/Source/Runtime/GameplayAbilities/Public/Abilities/GameplayAbilityTargetActor_GroundTrace.h
+++ b/Engine/Source/Runtime/GameplayAbilities/Public/Abilities/GameplayAbilityTargetActor_GroundTrace.h
@@ -11,6 +11,23 @@ class GAMEPLAYABILITIES_API AGameplayAbilityTargetActor_GroundTrace : public AGa
 {
 	GENERATED_UCLASS_BODY()
 
+	virtual void StartTargeting(UGameplayAbility* InAbility) override;
+
+	/** Radius for a sphere or capsule. */
+	UPROPERTY(BlueprintReadWrite, EditAnywhere, meta = (ExposeOnSpawn = true), Category = Targeting)
+	float CollisionRadius;
+
+	/** Height for a capsule. Implicitly indicates a capsule is desired if this is greater than zero. */
+	UPROPERTY(BlueprintReadWrite, EditAnywhere, meta = (ExposeOnSpawn = true), Category = Targeting)
+	float CollisionHeight;
+
+	//TODO: Consider adding the ability to rotate the capsule with a supplied quaternion.
+
 protected:
 	virtual FHitResult PerformTrace(AActor* InSourceActor) const override;
+
+	bool AdjustCollisionResultForShape(const FVector OriginalStartPoint, const FVector OriginalEndPoint, const FCollisionQueryParams Params, FHitResult& OutHitResult) const;
+
+	FCollisionShape CollisionShape;
+	float CollisionHeightOffset;		//When tracing, give this much extra height to avoid start-in-ground problems. Dealing with thick placement actors while standing near walls may be trickier.
 };
\ No newline at end of file
diff --git a/Engine/Source/Runtime/GameplayAbilities/Public/Abilities/Tasks/AbilityTask_SpawnActor.h b/Engine/Source/Runtime/GameplayAbilities/Public/Abilities/Tasks/AbilityTask_SpawnActor.h
--- a/Engine/Source/Runtime/GameplayAbilities/Public/Abilities/Tasks/AbilityTask_SpawnActor.h
+++ b/Engine/Source/Runtime/GameplayAbilities/Public/Abilities/Tasks/AbilityTask_SpawnActor.h
@@ -30,14 +30,14 @@ class UAbilityTask_SpawnActor: public UAbilityTask
 	
 	/** Spawn new Actor on the network authority (server) */
 	UFUNCTION(BlueprintCallable, meta=(HidePin = "WorldContextObject", DefaultToSelf = "WorldContextObject", BlueprintInternalUseOnly = "true"), Category="Ability|Tasks")
-	static UAbilityTask_SpawnActor* SpawnActor(UObject* WorldContextObject, TSubclassOf<AActor> Class);
+	static UAbilityTask_SpawnActor* SpawnActor(UObject* WorldContextObject, FGameplayAbilityTargetDataHandle TargetData, TSubclassOf<AActor> Class);
 
 	UFUNCTION(BlueprintCallable, meta = (HidePin = "WorldContextObject", DefaultToSelf = "WorldContextObject", BlueprintInternalUseOnly = "true"), Category = "Abilities")
-	bool BeginSpawningActor(UObject* WorldContextObject, TSubclassOf<AActor> Class, AActor*& SpawnedActor);
+	bool BeginSpawningActor(UObject* WorldContextObject, FGameplayAbilityTargetDataHandle TargetData, TSubclassOf<AActor> Class, AActor*& SpawnedActor);
 
 	UFUNCTION(BlueprintCallable, meta = (HidePin = "WorldContextObject", DefaultToSelf = "WorldContextObject", BlueprintInternalUseOnly = "true"), Category = "Abilities")
-	void FinishSpawningActor(UObject* WorldContextObject, AActor* SpawnedActor);
+	void FinishSpawningActor(UObject* WorldContextObject, FGameplayAbilityTargetDataHandle TargetData, AActor* SpawnedActor);
 
 protected:
-	
+	FGameplayAbilityTargetDataHandle CachedTargetDataHandle;
 };
\ No newline at end of file
