Make gameplay ability spec IsActive call return true it the ability has been activated more times than it has ended regardless of the instancing rule.

diff --git a/Engine/Source/Runtime/GameplayAbilities/Private/AbilitySystemComponent_Abilities.cpp b/Engine/Source/Runtime/GameplayAbilities/Private/AbilitySystemComponent_Abilities.cpp
--- a/Engine/Source/Runtime/GameplayAbilities/Private/AbilitySystemComponent_Abilities.cpp
+++ b/Engine/Source/Runtime/GameplayAbilities/Private/AbilitySystemComponent_Abilities.cpp
@@ -227,7 +227,7 @@ void UAbilitySystemComponent::NotifyAbilityEnded(FGameplayAbilitySpecHandle Hand
 		AnimatingAbility = NULL;
 	}
 
-	Spec->ActiveInstanceCount--;
+	Spec->ActiveCount--;
 	
 	/** If this is instanced per execution, mark pending kill and remove it from our instanced lists if we are the authority */
 	if (Ability->GetInstancingPolicy() == EGameplayAbilityInstancingPolicy::InstancedPerExecution)
@@ -354,7 +354,7 @@ bool UAbilitySystemComponent::TryActivateAbility(FGameplayAbilitySpecHandle Hand
 		ABILITY_LOG(Warning, TEXT("TryActivateAbility called when ability was already active. NetMode: %d. Ability: %s"), (int32)NetMode, *Ability->GetName());
 	}
 	
-	Spec->ActiveInstanceCount++;
+	Spec->ActiveCount++;
 
 	// Setup a fresh ActivationInfo for this AbilitySpec.
 	Spec->ActivationInfo = FGameplayAbilityActivationInfo(ActorInfo->Actor.Get(), InPredictionKey);
diff --git a/Engine/Source/Runtime/GameplayAbilities/Private/GameplayAbilityTypes.cpp b/Engine/Source/Runtime/GameplayAbilities/Private/GameplayAbilityTypes.cpp
--- a/Engine/Source/Runtime/GameplayAbilities/Private/GameplayAbilityTypes.cpp
+++ b/Engine/Source/Runtime/GameplayAbilities/Private/GameplayAbilityTypes.cpp
@@ -96,10 +96,5 @@ void FGameplayAbilityActivationInfo::SetPredictionStale()
 
 bool FGameplayAbilitySpec::IsActive() const
 {
-	if (Ability && Ability->GetInstancingPolicy() == EGameplayAbilityInstancingPolicy::NonInstanced)
-	{
-		return true;
-	}
-
-	return ActiveInstanceCount > 0;
+	return ActiveCount > 0;
 }
\ No newline at end of file
diff --git a/Engine/Source/Runtime/GameplayAbilities/Public/Abilities/GameplayAbilityTypes.h b/Engine/Source/Runtime/GameplayAbilities/Public/Abilities/GameplayAbilityTypes.h
--- a/Engine/Source/Runtime/GameplayAbilities/Public/Abilities/GameplayAbilityTypes.h
+++ b/Engine/Source/Runtime/GameplayAbilities/Public/Abilities/GameplayAbilityTypes.h
@@ -235,13 +235,13 @@ struct GAMEPLAYABILITIES_API FGameplayAbilitySpec
 	GENERATED_USTRUCT_BODY()
 
 	FGameplayAbilitySpec()
-	: Ability(nullptr), Level(1), InputID(INDEX_NONE), InputPressed(false), ActiveInstanceCount(0)
+	: Ability(nullptr), Level(1), InputID(INDEX_NONE), InputPressed(false), ActiveCount(0)
 	{
 		
 	}
 
 	FGameplayAbilitySpec(UGameplayAbility* InAbility, int32 InLevel=1, int32 InInputID=INDEX_NONE)
-		: Ability(InAbility), Level(InLevel), InputID(InInputID), InputPressed(false), ActiveInstanceCount(0)
+		: Ability(InAbility), Level(InLevel), InputID(InInputID), InputPressed(false), ActiveCount(0)
 	{
 		Handle.GenerateNewHandle();
 	}
@@ -266,9 +266,9 @@ struct GAMEPLAYABILITIES_API FGameplayAbilitySpec
 	UPROPERTY(NotReplicated)
 	bool	InputPressed;
 
-	/** The number of currently active instances of this ability */
+	/** A count of the number of times this ability has been activated minus the number of times it has been ended. For instanced abilities this will be the number of currently active instances. */
 	UPROPERTY()
-	uint8	ActiveInstanceCount;
+	uint8	ActiveCount;
 
 	/** Activation state of this ability. This is not replicated since it needs to be overwritten locally on clients during prediction. */
 	UPROPERTY(NotReplicated)
