One more prediction bug + comments

diff --git a/Engine/Source/Runtime/GameplayAbilities/Private/GameplayPrediction.cpp b/Engine/Source/Runtime/GameplayAbilities/Private/GameplayPrediction.cpp
--- a/Engine/Source/Runtime/GameplayAbilities/Private/GameplayPrediction.cpp
+++ b/Engine/Source/Runtime/GameplayAbilities/Private/GameplayPrediction.cpp
@@ -190,6 +190,9 @@ void FPredictionKeyDelegates::AddDependancy(FPredictionKey::KeyType ThisKey, FPr
 
 FScopedPredictionWindow::FScopedPredictionWindow(UAbilitySystemComponent* AbilitySystemComponent, FPredictionKey InPredictionKey)
 {
+	// This is used to set an already generated predictiot key as the current scoped prediction key.
+	// Should be called on the server for logical scopes where a given key is valid. E.g, "client gave me this key, we both are going to run Foo()".
+
 	Owner = AbilitySystemComponent;
 	Owner->ScopedPedictionKey = InPredictionKey;
 	ClearScopedPredictionKey = true;
@@ -198,6 +201,10 @@ FScopedPredictionWindow::FScopedPredictionWindow(UAbilitySystemComponent* Abilit
 
 FScopedPredictionWindow::FScopedPredictionWindow(UGameplayAbility* GameplayAbilityInstance)
 {
+	// "Gets" the current/"best" prediction key and sets it as the prediciton key for this given abilities CurrentActivationInfo.
+	//	-On the server, it means we look for the prediction key that was given to us and either stored on the ActivationIfno or on the ScopedPredictionKey.
+	//  -On the client, we generate a new prediction  key if we don't already have one.
+
 	Ability = GameplayAbilityInstance;
 	check(Ability.IsValid());
 
@@ -208,23 +215,30 @@ FScopedPredictionWindow::FScopedPredictionWindow(UGameplayAbility* GameplayAbili
 
 	FGameplayAbilityActivationInfo& ActivationInfo = Ability->GetCurrentActivationInfoRef();
 
-
 	if (ActivationInfo.ActivationMode == EGameplayAbilityActivationMode::Authority)
 	{
-		ActivationInfo.SetPredictionKey(Owner->ScopedPedictionKey);
+		// We are the server, if we already have a valid prediction key in our ActivationInfo, just use it.
+		if (ActivationInfo.GetPredictionKey().IsValidForMorePrediction() == false)
+		{
+			// If we don't, then we fall back to the ScopedPredictionKey that would have been explcitly set by another ScopedPredictionKey that used the other constructor.
+			ActivationInfo.SetPredictionKey(Owner->ScopedPedictionKey);
+		}
 	}
 	else
 	{
+		// We are the client, first save off our current activation mode (E.g, if we were previously 'Confirmed' we need to back to 'Predicting' then back to 'Confirmed' after this scope.
 		ClientPrevActivationMode = static_cast<int8>(ActivationInfo.ActivationMode);
 
 		if (Owner->ScopedPedictionKey.IsValidKey())
 		{
+			// We already have a scoped prediction key set locally, so use that.
 			ActivationInfo.SetPredictionKey(Owner->ScopedPedictionKey);
 			ActivationInfo.ActivationMode = EGameplayAbilityActivationMode::Predicting;
 			ScopedPredictionKey = Owner->ScopedPedictionKey;
 		}
 		else
 		{
+			// We don't have a valid key, so generate a new one.
 			ActivationInfo.GenerateNewPredictionKey();
 			ScopedPredictionKey = ActivationInfo.GetPredictionKeyForNewAction();
 		}
