Comment + dont replicate activation count for now until prediction issues sorted out.

diff --git a/Engine/Source/Runtime/GameplayAbilities/Public/Abilities/GameplayAbilityTypes.h b/Engine/Source/Runtime/GameplayAbilities/Public/Abilities/GameplayAbilityTypes.h
--- a/Engine/Source/Runtime/GameplayAbilities/Public/Abilities/GameplayAbilityTypes.h
+++ b/Engine/Source/Runtime/GameplayAbilities/Public/Abilities/GameplayAbilityTypes.h
@@ -266,8 +266,8 @@ struct GAMEPLAYABILITIES_API FGameplayAbilitySpec
 	UPROPERTY(NotReplicated)
 	bool	InputPressed;
 
-	/** A count of the number of times this ability has been activated minus the number of times it has been ended. For instanced abilities this will be the number of currently active instances. */
-	UPROPERTY()
+	/** A count of the number of times this ability has been activated minus the number of times it has been ended. For instanced abilities this will be the number of currently active instances. Can't replicate until prediction accurately handles this.*/
+	UPROPERTY(NotReplicated)
 	uint8	ActiveCount;
 
 	/** Activation state of this ability. This is not replicated since it needs to be overwritten locally on clients during prediction. */
