damage scaling fix

diff --git a/Engine/Source/Runtime/GameplayAbilities/Private/GameplayEffect.cpp b/Engine/Source/Runtime/GameplayAbilities/Private/GameplayEffect.cpp
--- a/Engine/Source/Runtime/GameplayAbilities/Private/GameplayEffect.cpp
+++ b/Engine/Source/Runtime/GameplayAbilities/Private/GameplayEffect.cpp
@@ -2183,7 +2183,7 @@ void FGameplayEffectLevelSpec::RegisterLevelDependancy(TWeakPtr<FAggregator> Own
 {
 	if (Source.IsValid() && Attribute.GetUProperty() != NULL)
 	{
-		UAbilitySystemComponent * SourceComponent = Source->FindComponentByClass<UAbilitySystemComponent>();
+		UAbilitySystemComponent* SourceComponent = UAbilitySystemGlobals::GetAbilitySystemComponentFromActor(Source.Get());
 		if (SourceComponent)
 		{
 			SourceComponent->AddDependancyToAttribute(Attribute, OwningAggregator);
diff --git a/Engine/Source/Runtime/GameplayAbilities/Public/GameplayEffect.h b/Engine/Source/Runtime/GameplayAbilities/Public/GameplayEffect.h
--- a/Engine/Source/Runtime/GameplayAbilities/Public/GameplayEffect.h
+++ b/Engine/Source/Runtime/GameplayAbilities/Public/GameplayEffect.h
@@ -401,7 +401,10 @@ struct FGameplayEffectLevelSpec
 	{
 		if (Def.Attribute.GetUProperty() != NULL)
 		{
+			// If given an attribute to base level on, then we throw out the passed in constant level.
 			Attribute = Def.Attribute;
+			ConstantLevel = INVALID_LEVEL;
+			CachedLevel = INVALID_LEVEL;
 		}
 
 		if (Def.TakeSnapshotOnInit)
