Handle null gameplayeffect parameters better

diff --git a/Engine/Source/Runtime/GameplayAbilities/Private/Abilities/GameplayAbility.cpp b/Engine/Source/Runtime/GameplayAbilities/Private/Abilities/GameplayAbility.cpp
--- a/Engine/Source/Runtime/GameplayAbilities/Private/Abilities/GameplayAbility.cpp
+++ b/Engine/Source/Runtime/GameplayAbilities/Private/Abilities/GameplayAbility.cpp
@@ -655,14 +655,24 @@ FActiveGameplayEffectHandle UGameplayAbility::ApplyGameplayEffectToOwner(const F
 
 FActiveGameplayEffectHandle UGameplayAbility::K2_ApplyGameplayEffectToTarget(FGameplayAbilityTargetDataHandle Target, const UGameplayEffect* GameplayEffect, int32 GameplayEffectLevel)
 {
+	if (GameplayEffect==nullptr)
+	{
+		ABILITY_LOG(Error, TEXT("K2_ApplyGameplayEffectToTarget called on ability %s with no GameplayEffect."), *GetName());
+		return FActiveGameplayEffectHandle();
+	}
+
 	return ApplyGameplayEffectToTarget(CurrentSpecHandle, CurrentActorInfo, CurrentActivationInfo, Target, GameplayEffect, GameplayEffectLevel);
 }
 
 FActiveGameplayEffectHandle UGameplayAbility::ApplyGameplayEffectToTarget(const FGameplayAbilitySpecHandle Handle, const FGameplayAbilityActorInfo* ActorInfo, const FGameplayAbilityActivationInfo ActivationInfo, FGameplayAbilityTargetDataHandle Target, const UGameplayEffect* GameplayEffect, int32 GameplayEffectLevel)
 {
-	FActiveGameplayEffectHandle EffectHandle;
-
+	if (GameplayEffect == nullptr)
+	{
+		ABILITY_LOG(Error, TEXT("K2_ApplyGameplayEffectToTarget called on ability %s with no GameplayEffect."), *GetName());
+		return FActiveGameplayEffectHandle();
+	}
 
+	FActiveGameplayEffectHandle EffectHandle;
 	if (ActivationInfo.ActivationMode == EGameplayAbilityActivationMode::Authority || ActivationInfo.ActivationMode == EGameplayAbilityActivationMode::Predicting)
 	{
 		for (auto Data : Target.Data)
