Add way to get instigator actor/transform from GameplayCue parmaeter.

diff --git a/Engine/Source/Runtime/GameplayAbilities/Private/AbilitySystemBlueprintLibrary.cpp b/Engine/Source/Runtime/GameplayAbilities/Private/AbilitySystemBlueprintLibrary.cpp
--- a/Engine/Source/Runtime/GameplayAbilities/Private/AbilitySystemBlueprintLibrary.cpp
+++ b/Engine/Source/Runtime/GameplayAbilities/Private/AbilitySystemBlueprintLibrary.cpp
@@ -225,3 +225,20 @@ bool UAbilitySystemBlueprintLibrary::HasHitResult(FGameplayCueParameters Paramet
 {
 	return Parameters.InstigatorContext.HitResult.IsValid();
 }
+
+AActor*	UAbilitySystemBlueprintLibrary::GetInstigatorActor(FGameplayCueParameters Parameters)
+{
+	return Parameters.InstigatorContext.GetOriginalInstigator();
+}
+
+FTransform UAbilitySystemBlueprintLibrary::GetInstigatorTransform(FGameplayCueParameters Parameters)
+{
+	AActor* InstigatorActor = GetInstigatorActor(Parameters);
+	if (InstigatorActor)
+	{
+		return InstigatorActor->GetTransform();
+	}
+
+	ABILITY_LOG(Warning, TEXT("UAbilitySystemBlueprintLibrary::GetInstigatorTransform called on GameplayCue with no valid instigator"));
+	return FTransform::Identity;
+}
\ No newline at end of file
diff --git a/Engine/Source/Runtime/GameplayAbilities/Public/AbilitySystemBlueprintLibrary.h b/Engine/Source/Runtime/GameplayAbilities/Public/AbilitySystemBlueprintLibrary.h
--- a/Engine/Source/Runtime/GameplayAbilities/Public/AbilitySystemBlueprintLibrary.h
+++ b/Engine/Source/Runtime/GameplayAbilities/Public/AbilitySystemBlueprintLibrary.h
@@ -77,6 +77,14 @@ class UAbilitySystemBlueprintLibrary : public UBlueprintFunctionLibrary
 	UFUNCTION(BlueprintPure, Category = "Ability|GameplayCue")
 	static bool HasHitResult(FGameplayCueParameters Parameters);
 
+	/** Gets the instigating actor (Pawn/Avatar) of the GameplayCue */
+	UFUNCTION(BlueprintPure, Category = "Ability|GameplayCue")
+	static AActor*	GetInstigatorActor(FGameplayCueParameters Parameters);
+
+	/** Gets instigating world location */
+	UFUNCTION(BlueprintPure, Category = "Ability|GameplayCue")
+	static FTransform GetInstigatorTransform(FGameplayCueParameters Parameters);
+
 
 	// -------------------------------------------------------------------------------
 
