Fix crash if ReticleClas is null
remove AutoFire and TimeUntilAutoFire options

diff --git a/Engine/Source/Runtime/GameplayAbilities/Private/Abilities/GameplayAbilityTargetActor_Trace.cpp b/Engine/Source/Runtime/GameplayAbilities/Private/Abilities/GameplayAbilityTargetActor_Trace.cpp
--- a/Engine/Source/Runtime/GameplayAbilities/Private/Abilities/GameplayAbilityTargetActor_Trace.cpp
+++ b/Engine/Source/Runtime/GameplayAbilities/Private/Abilities/GameplayAbilityTargetActor_Trace.cpp
@@ -42,11 +42,15 @@ void AGameplayAbilityTargetActor_Trace::StartTargeting(UGameplayAbility* InAbili
 {
 	Super::StartTargeting(InAbility);
 	SourceActor = InAbility->GetCurrentActorInfo()->Actor.Get();
-	
-	ReticleActor = GetWorld()->SpawnActor<AGameplayAbilityWorldReticle>(ReticleClass, GetActorLocation(), GetActorRotation());
-	if (AGameplayAbilityWorldReticle* CachedReticleActor = ReticleActor.Get())
+
+	if (ReticleClass)
 	{
-		CachedReticleActor->InitializeReticle(this);
+		AGameplayAbilityWorldReticle* SpawnedReticleActor = GetWorld()->SpawnActor<AGameplayAbilityWorldReticle>(ReticleClass, GetActorLocation(), GetActorRotation());
+		if (SpawnedReticleActor)
+		{
+			SpawnedReticleActor->InitializeReticle(this);
+			ReticleActor = SpawnedReticleActor;
+		}
 	}
 }
 
@@ -65,16 +69,6 @@ void AGameplayAbilityTargetActor_Trace::Tick(float DeltaSeconds)
 		}
 
 		SetActorLocationAndRotation(EndPoint, SourceActor->GetActorRotation());
-
-		if (ShouldProduceTargetData() && bAutoFire)
-		{
-			if (TimeUntilAutoFire <= 0.0f)
-			{
-				ConfirmTargeting();
-				bAutoFire = false;
-			}
-			TimeUntilAutoFire -= DeltaSeconds;
-		}
 	}
 }
 
@@ -84,7 +78,6 @@ void AGameplayAbilityTargetActor_Trace::ConfirmTargetingAndContinue()
 	if (SourceActor)
 	{
 		bDebug = false;
-		bAutoFire = false;
 		FGameplayAbilityTargetDataHandle Handle = MakeTargetData(PerformTrace(SourceActor));
 		TargetDataReadyDelegate.Broadcast(Handle);
 	}
diff --git a/Engine/Source/Runtime/GameplayAbilities/Public/Abilities/GameplayAbilityTargetActor_Trace.h b/Engine/Source/Runtime/GameplayAbilities/Public/Abilities/GameplayAbilityTargetActor_Trace.h
--- a/Engine/Source/Runtime/GameplayAbilities/Public/Abilities/GameplayAbilityTargetActor_Trace.h
+++ b/Engine/Source/Runtime/GameplayAbilities/Public/Abilities/GameplayAbilityTargetActor_Trace.h
@@ -26,12 +26,6 @@ public:
 	UPROPERTY(BlueprintReadWrite, EditAnywhere, meta = (ExposeOnSpawn = true), Category = Projectile)
 	float MaxRange;
 
-	UPROPERTY(BlueprintReadWrite, EditAnywhere, meta = (ExposeOnSpawn = true), Category = Projectile)
-	bool bAutoFire;				//Automatically confirm after the timer runs out. Should also have support for an auto-cancel feature.
-
-	UPROPERTY(BlueprintReadWrite, EditAnywhere, meta = (ExposeOnSpawn = true), Category = Projectile)
-	float TimeUntilAutoFire;
-
 	UPROPERTY(BlueprintReadWrite, EditAnywhere, meta = (ExposeOnSpawn = true), Category = Projectile)
 	TSubclassOf<AGameplayAbilityWorldReticle> ReticleClass;		//Using a special class for replication purposes
 
