Add Rate to PlayMontageAndWait

diff --git a/Engine/Source/Runtime/GameplayAbilities/Private/Abilities/Tasks/AbilityTask_PlayMontageAndWait.cpp b/Engine/Source/Runtime/GameplayAbilities/Private/Abilities/Tasks/AbilityTask_PlayMontageAndWait.cpp
--- a/Engine/Source/Runtime/GameplayAbilities/Private/Abilities/Tasks/AbilityTask_PlayMontageAndWait.cpp
+++ b/Engine/Source/Runtime/GameplayAbilities/Private/Abilities/Tasks/AbilityTask_PlayMontageAndWait.cpp
@@ -5,7 +5,7 @@
 UAbilityTask_PlayMontageAndWait::UAbilityTask_PlayMontageAndWait(const class FPostConstructInitializeProperties& PCIP)
 	: Super(PCIP)
 {
-
+	Rate = 1.f;
 }
 
 void UAbilityTask_PlayMontageAndWait::OnMontageEnded(UAnimMontage* Montage, bool bInterrupted)
@@ -22,10 +22,11 @@ void UAbilityTask_PlayMontageAndWait::OnMontageEnded(UAnimMontage* Montage, bool
 	EndTask();
 }
 
-UAbilityTask_PlayMontageAndWait* UAbilityTask_PlayMontageAndWait::CreatePlayMontageAndWaitProxy(UObject* WorldContextObject, UAnimMontage *MontageToPlay)
+UAbilityTask_PlayMontageAndWait* UAbilityTask_PlayMontageAndWait::CreatePlayMontageAndWaitProxy(UObject* WorldContextObject, UAnimMontage *MontageToPlay, float Rate)
 {
 	UAbilityTask_PlayMontageAndWait* MyObj = NewTask<UAbilityTask_PlayMontageAndWait>(WorldContextObject);
 	MyObj->MontageToPlay = MontageToPlay;
+	MyObj->Rate = Rate;
 
 	return MyObj;
 }
@@ -37,7 +38,7 @@ void UAbilityTask_PlayMontageAndWait::Activate()
 		const FGameplayAbilityActorInfo* ActorInfo = Ability->GetCurrentActorInfo();
 		if (ActorInfo->AnimInstance.IsValid())
 		{
-			ActorInfo->AnimInstance->Montage_Play(MontageToPlay, 1.f);
+			ActorInfo->AnimInstance->Montage_Play(MontageToPlay, Rate);
 			if (MontageToPlay)
 			{
 				FOnMontageEnded EndDelegate;
diff --git a/Engine/Source/Runtime/GameplayAbilities/Public/Abilities/Tasks/AbilityTask_PlayMontageAndWait.h b/Engine/Source/Runtime/GameplayAbilities/Public/Abilities/Tasks/AbilityTask_PlayMontageAndWait.h
--- a/Engine/Source/Runtime/GameplayAbilities/Public/Abilities/Tasks/AbilityTask_PlayMontageAndWait.h
+++ b/Engine/Source/Runtime/GameplayAbilities/Public/Abilities/Tasks/AbilityTask_PlayMontageAndWait.h
@@ -20,7 +20,7 @@ class UAbilityTask_PlayMontageAndWait : public UAbilityTask
 	void OnMontageEnded(UAnimMontage* Montage, bool bInterrupted);
 
 	UFUNCTION(BlueprintCallable, Category="Ability|Tasks", meta = (FriendlyName="PlayMontageAndWait", HidePin = "WorldContextObject", DefaultToSelf = "WorldContextObject", BlueprintInternalUseOnly = "TRUE"))
-	static UAbilityTask_PlayMontageAndWait* CreatePlayMontageAndWaitProxy(UObject* WorldContextObject, UAnimMontage *MontageToPlay);
+	static UAbilityTask_PlayMontageAndWait* CreatePlayMontageAndWaitProxy(UObject* WorldContextObject, UAnimMontage *MontageToPlay, float Rate = 1.f);
 
 	virtual void Activate() override;
 
@@ -31,5 +31,5 @@ private:
 	virtual void OnDestroy(bool AbilityEnded) override;
 
 	UAnimMontage* MontageToPlay;
-	
+	float Rate;	
 };
\ No newline at end of file
