#UE4 - Fixed cases where an ability could potentially underflow ActiveCount by calling EndAbility multiple times

diff --git a/Engine/Source/Runtime/GameplayAbilities/Private/Abilities/GameplayAbility.cpp b/Engine/Source/Runtime/GameplayAbilities/Private/Abilities/GameplayAbility.cpp
--- a/Engine/Source/Runtime/GameplayAbilities/Private/Abilities/GameplayAbility.cpp
+++ b/Engine/Source/Runtime/GameplayAbilities/Private/Abilities/GameplayAbility.cpp
@@ -254,8 +254,6 @@ void UGameplayAbility::CancelAbility(const FGameplayAbilitySpecHandle Handle, co
 
 void UGameplayAbility::EndAbility(const FGameplayAbilitySpecHandle Handle, const FGameplayAbilityActorInfo* ActorInfo, const FGameplayAbilityActivationInfo ActivationInfo)
 {
-	bIsActive = false;
-
 	// Tell all our tasks that we are finished and they should cleanup
 	for (TWeakObjectPtr<UAbilityTask> Task : ActiveTasks)
 	{
@@ -266,14 +264,19 @@ void UGameplayAbility::EndAbility(const FGameplayAbilitySpecHandle Handle, const
 	}
 	ActiveTasks.Reset();	// Empty the array but dont resize memory, since this object is probably going to be destroyed very soon anyways.
 
-	// Unblock ability tags we previously blocked
-	if (ActorInfo && ActorInfo->AbilitySystemComponent.IsValid())
+	if (bIsActive)
 	{
-		ActorInfo->AbilitySystemComponent->UnBlockAbilitiesWithTags(BlockAbilitiesWithTag);
-	}
+		bIsActive = false;
+
+		// Unblock ability tags we previously blocked
+		if (ActorInfo && ActorInfo->AbilitySystemComponent.IsValid())
+		{
+			ActorInfo->AbilitySystemComponent->UnBlockAbilitiesWithTags(BlockAbilitiesWithTag);
+		}
 
-	// Tell owning AbilitySystemComponent that we ended so it can do stuff (including MarkPendingKill us)
-	ActorInfo->AbilitySystemComponent->NotifyAbilityEnded(Handle, this);
+		// Tell owning AbilitySystemComponent that we ended so it can do stuff (including MarkPendingKill us)
+		ActorInfo->AbilitySystemComponent->NotifyAbilityEnded(Handle, this);
+	}
 }
 
 void UGameplayAbility::ActivateAbility(const FGameplayAbilitySpecHandle Handle, const FGameplayAbilityActorInfo* ActorInfo, const FGameplayAbilityActivationInfo ActivationInfo)
@@ -302,10 +305,11 @@ void UGameplayAbility::PreActivate(const FGameplayAbilitySpecHandle Handle, cons
 {
 	UAbilitySystemComponent* Comp = ActorInfo->AbilitySystemComponent.Get();
 
+	bIsActive = true;
+
 	Comp->CancelAbilitiesWithTags(CancelAbilitiesWithTag, Handle, ActorInfo, ActivationInfo, this);
 	Comp->BlockAbilitiesWithTags(BlockAbilitiesWithTag);
 
-	bIsActive = true;
 	Comp->NotifyAbilityActivated(Handle, this);
 }
 
