rename CancelAbilityWithTags to CancelAbility, take in WIthTags and WithoutTags as parameters. Intended to be used to cancel abilities on pawn death, etc.

diff --git a/Engine/Source/Runtime/GameplayAbilities/Private/Abilities/GameplayAbility.cpp b/Engine/Source/Runtime/GameplayAbilities/Private/Abilities/GameplayAbility.cpp
--- a/Engine/Source/Runtime/GameplayAbilities/Private/Abilities/GameplayAbility.cpp
+++ b/Engine/Source/Runtime/GameplayAbilities/Private/Abilities/GameplayAbility.cpp
@@ -307,7 +307,7 @@ void UGameplayAbility::PreActivate(const FGameplayAbilitySpecHandle Handle, cons
 
 	bIsActive = true;
 
-	Comp->CancelAbilitiesWithTags(CancelAbilitiesWithTag, Handle, ActorInfo, ActivationInfo, this);
+	Comp->CancelAbilities(&CancelAbilitiesWithTag, nullptr, this);
 	Comp->BlockAbilitiesWithTags(BlockAbilitiesWithTag);
 
 	Comp->NotifyAbilityActivated(Handle, this);
diff --git a/Engine/Source/Runtime/GameplayAbilities/Private/AbilitySystemComponent_Abilities.cpp b/Engine/Source/Runtime/GameplayAbilities/Private/AbilitySystemComponent_Abilities.cpp
--- a/Engine/Source/Runtime/GameplayAbilities/Private/AbilitySystemComponent_Abilities.cpp
+++ b/Engine/Source/Runtime/GameplayAbilities/Private/AbilitySystemComponent_Abilities.cpp
@@ -251,11 +251,17 @@ void UAbilitySystemComponent::NotifyAbilityEnded(FGameplayAbilitySpecHandle Hand
 	}
 }
 
-void UAbilitySystemComponent::CancelAbilitiesWithTags(const FGameplayTagContainer Tags, const FGameplayAbilitySpecHandle Handle, const FGameplayAbilityActorInfo* ActorInfo, const FGameplayAbilityActivationInfo ActivationInfo, UGameplayAbility* Ignore)
+void UAbilitySystemComponent::CancelAbilities(const FGameplayTagContainer* WithTags, const FGameplayTagContainer* WithoutTags, UGameplayAbility* Ignore)
 {
+	FGameplayAbilityActorInfo* ActorInfo = AbilityActorInfo.Get();
+	FGameplayAbilityActivationInfo ActivationInfo;
+
 	for (FGameplayAbilitySpec& Spec : ActivatableAbilities)
 	{
-		if (Spec.IsActive() && Spec.Ability && Spec.Ability->AbilityTags.MatchesAny(Tags, false))
+		bool WithTagPass = (!WithTags || Spec.Ability->AbilityTags.MatchesAny(*WithTags, false));
+		bool WithoutTagPass = (!WithoutTags || Spec.Ability->AbilityTags.MatchesAll(*WithoutTags, false));
+
+		if (Spec.IsActive() && Spec.Ability && WithTagPass && WithoutTagPass)
 		{
 			if (Spec.Ability->GetInstancingPolicy() == EGameplayAbilityInstancingPolicy::InstancedPerExecution)
 			{
diff --git a/Engine/Source/Runtime/GameplayAbilities/Public/AbilitySystemComponent.h b/Engine/Source/Runtime/GameplayAbilities/Public/AbilitySystemComponent.h
--- a/Engine/Source/Runtime/GameplayAbilities/Public/AbilitySystemComponent.h
+++ b/Engine/Source/Runtime/GameplayAbilities/Public/AbilitySystemComponent.h
@@ -107,36 +107,7 @@ class GAMEPLAYABILITIES_API UAbilitySystemComponent : public UActorComponent, pu
 	
 	virtual void GetSubobjectsWithStableNamesForNetworking(TArray<UObject*>& Objs) override;
 
-	/**
-	 *	Prediction Keys
-	 *		PredictionKey is a simple system to allow clients to locally predict interactions within the AbilityComponent while having
-	 *		a mechanism for merging results from the server.
-	 *		
-	 *		There are two main cases:
-	 *			-Client predicts something and is wrong. (Mispredictions). He must cleanup stuff that he created or changed.
-	 *			-Client predicts something and is right. 
-	 *			
-	 *			The Client predicts something and asks the server. "Can I do X? PredictionKey=Y".
-	 *			Server can reply immediately and directly with yes or no, and include the PredictionKey with that answer.
-	 *			
-	 *			If No, the client can just immediately cleanup whatever he predicted. This is straightforward.
-	 *			
-	 *			If Yes, there is a small window where the Yes RPC is received but the actor/component properties have not been updated.
-	 *			The client's cannot immediately cleanup his prediction work: he needs to wait for property replication.
-	 *			
-	 *			This is where ReplicatedPredictionKey comes in. 
-	 *			
-	 *			When the server allows a predicted action, he sets ReplicatedPredictionKey to this value. The client then looks
-	 *			at this replicated value. In OnRep_PredictionKey, the client goes through and cleans up all stuff that was predicted
-	 *			and is no longer needed.
-	 *			
-	 *			-Client calls GetNextPredictionKey() at the start of a predicted action. This gives him a key to use for this 'run'.
-	 *			-Client can predict whatever he wants as long as he registers an 'undo' function with GetPredictionKeyDelegate().
-	 *			-This delegate is invoked if the client is denied his prediction, or if he is right and property replication catches up.
-	 *			
-	 *	 
-	 */
-
+	/** PredictionKeys, see more info in GameplayPrediction.h */
 	UPROPERTY(ReplicatedUsing=OnRep_PredictionKey)
 	FPredictionKey	ReplicatedPredictionKey;
 
@@ -365,7 +336,7 @@ class GAMEPLAYABILITIES_API UAbilitySystemComponent : public UActorComponent, pu
 
 	UGameplayAbility* CreateNewInstanceOfAbility(FGameplayAbilitySpec& Spec, UGameplayAbility* Ability);
 
-	void CancelAbilitiesWithTags(const FGameplayTagContainer Tags, const FGameplayAbilitySpecHandle Handle, const FGameplayAbilityActorInfo* ActorInfo, const FGameplayAbilityActivationInfo ActivationInfo, UGameplayAbility* Ignore);
+	void CancelAbilities(const FGameplayTagContainer* WithTags=nullptr, const FGameplayTagContainer* WithoutTags=nullptr, UGameplayAbility* Ignore=nullptr);
 
 	void BlockAbilitiesWithTags(const FGameplayTagContainer Tags);
 	void UnBlockAbilitiesWithTags(const FGameplayTagContainer Tags);
@@ -518,7 +489,7 @@ class GAMEPLAYABILITIES_API UAbilitySystemComponent : public UActorComponent, pu
 	// -----------------------------------------------------------------------------
 
 	/**
-	 *	While these appear to be state, this is actually synchronization events w/ some payload data
+	 *	While these appear to be state, these are actually synchronization events w/ some payload data
 	 */
 
 	UFUNCTION(Server, reliable, WithValidation)
@@ -529,7 +500,6 @@ class GAMEPLAYABILITIES_API UAbilitySystemComponent : public UActorComponent, pu
 
 	UFUNCTION(Server, reliable, WithValidation)
 	void ServerSetReplicatedTargetDataCancelled();
-	
 
 	void ConsumeAbilityConfirmCancel();
 
@@ -578,6 +548,7 @@ private:
 	UPROPERTY(ReplicatedUsing=OnRep_GameplayEffects)
 	FActiveGameplayCueContainer	ActiveGameplayCues;
 
+	/** Abilities with these tags are not able to be activated */
 	FGameplayTagCountContainer BlockedAbilityTags;
 
 	UFUNCTION()
