support for simulated ability tasks

diff --git a/Engine/Source/Runtime/GameplayAbilities/Private/Abilities/Tasks/AbilityTask.cpp b/Engine/Source/Runtime/GameplayAbilities/Private/Abilities/Tasks/AbilityTask.cpp
--- a/Engine/Source/Runtime/GameplayAbilities/Private/Abilities/Tasks/AbilityTask.cpp
+++ b/Engine/Source/Runtime/GameplayAbilities/Private/Abilities/Tasks/AbilityTask.cpp
@@ -7,6 +7,8 @@ UAbilityTask::UAbilityTask(const class FPostConstructInitializeProperties& PCIP)
 	: Super(PCIP)
 {
 	bTickingTask = false;
+	bSimulatedTask = false;
+	bIsSimulating = false;
 }
 
 void UAbilityTask::Activate()
@@ -21,12 +23,19 @@ void UAbilityTask::InitTask(UGameplayAbility* InAbility)
 	InAbility->TaskStarted(this);
 
 	// If this task requires ticking, register it with AbilitySystemComponent
-	if (bTickingTask && AbilitySystemComponent.IsValid())
+	if (AbilitySystemComponent.IsValid())
 	{
-		AbilitySystemComponent->TickingTaskStarted(this);
+		AbilitySystemComponent->TaskStarted(this);
 	}
 }
 
+void UAbilityTask::InitSimulatedTask(UAbilitySystemComponent* InAbilitySystemComponent)
+{
+	check(InAbilitySystemComponent);
+	AbilitySystemComponent = InAbilitySystemComponent;
+	bIsSimulating = true;
+}
+
 UWorld* UAbilityTask::GetWorld() const
 {
 	if (AbilitySystemComponent.IsValid())
@@ -47,6 +56,10 @@ AActor* UAbilityTask::GetActor() const
 			return Info->Actor.Get();
 		}
 	}
+	else if (AbilitySystemComponent.IsValid())
+	{
+		return AbilitySystemComponent->AbilityActorInfo->Actor.Get();
+	}
 
 	return nullptr;	
 }
@@ -85,9 +98,9 @@ void UAbilityTask::OnDestroy(bool AbilityIsEnding)
 	ensure(!IsPendingKill());
 
 	// If this task required ticking, unregister it with AbilitySystemComponent
-	if (bTickingTask && AbilitySystemComponent.IsValid())
+	if (AbilitySystemComponent.IsValid())
 	{
-		AbilitySystemComponent->TickingTaskEnded(this);
+		AbilitySystemComponent->TaskEnded(this);
 	}
 
 	// Remove ourselves from the owning Ability's task list, if the ability isn't ending
diff --git a/Engine/Source/Runtime/GameplayAbilities/Private/Abilities/Tasks/AbilityTask_MoveToLocation.cpp b/Engine/Source/Runtime/GameplayAbilities/Private/Abilities/Tasks/AbilityTask_MoveToLocation.cpp
--- a/Engine/Source/Runtime/GameplayAbilities/Private/Abilities/Tasks/AbilityTask_MoveToLocation.cpp
+++ b/Engine/Source/Runtime/GameplayAbilities/Private/Abilities/Tasks/AbilityTask_MoveToLocation.cpp
@@ -6,6 +6,8 @@ UAbilityTask_MoveToLocation::UAbilityTask_MoveToLocation(const class FPostConstr
 : Super(PCIP)
 {
 	bTickingTask = true;
+	bSimulatedTask = true;
+	bIsFinished = false;
 }
 
 UAbilityTask_MoveToLocation* UAbilityTask_MoveToLocation::MoveToLocation(class UObject* WorldContextObject, FVector Location, float Duration, UCurveFloat* OptionalInterpolationCurve)
@@ -14,7 +16,7 @@ UAbilityTask_MoveToLocation* UAbilityTask_MoveToLocation::MoveToLocation(class U
 
 	MyObj->StartLocation = MyObj->GetActor()->GetActorLocation();
 	MyObj->TargetLocation = Location;
-	MyObj->DurationOfMovement = FMath::Max(Duration, 0.001f);		//Avoid negative or divide-by-zero cases
+	MyObj->DurationOfMovement = FMath::Max(Duration, 0.001f);		// Avoid negative or divide-by-zero cases
 	MyObj->TimeMoveStarted = MyObj->GetWorld()->GetTimeSeconds();
 	MyObj->TimeMoveWillEnd = MyObj->TimeMoveStarted + MyObj->DurationOfMovement;
 	MyObj->LerpCurve = OptionalInterpolationCurve;
@@ -24,11 +26,25 @@ UAbilityTask_MoveToLocation* UAbilityTask_MoveToLocation::MoveToLocation(class U
 
 void UAbilityTask_MoveToLocation::Activate()
 {
+
+}
+
+void UAbilityTask_MoveToLocation::InitSimulatedTask(UAbilitySystemComponent* InAbilitySystemComponent)
+{
+	Super::InitSimulatedTask(InAbilitySystemComponent);
+
+	TimeMoveStarted = GetWorld()->GetTimeSeconds();
+	TimeMoveWillEnd = TimeMoveStarted + DurationOfMovement;
 }
 
 //TODO: This is still an awful way to do this and we should scrap this task or do it right.
 void UAbilityTask_MoveToLocation::TickTask(float DeltaTime)
 {
+	if (bIsFinished)
+	{
+		return;
+	}
+
 	Super::TickTask(DeltaTime);
 	AActor* MyActor = GetActor();
 	if (MyActor)
@@ -37,18 +53,34 @@ void UAbilityTask_MoveToLocation::TickTask(float DeltaTime)
 
 		if (CurrentTime >= TimeMoveWillEnd)
 		{
-			MyActor->SetActorLocation(TargetLocation);
-			OnTargetLocationReached.Broadcast();
-			EndTask();
+			bIsFinished = true;
+
+			// Teleport in attempt to find a valid collision spot
+			MyActor->TeleportTo(TargetLocation, MyActor->GetActorRotation());
+			if (!bIsSimulating)
+			{
+				MyActor->ForceNetUpdate();
+				OnTargetLocationReached.Broadcast();
+				EndTask();
+			}
 		}
 		else
 		{
 			float MoveFraction = (CurrentTime - TimeMoveStarted) / DurationOfMovement;
-			if (LerpCurve.IsValid())
+			if (LerpCurve)
 			{
-				MoveFraction = LerpCurve.Get()->GetFloatValue(MoveFraction);
+				MoveFraction = LerpCurve->GetFloatValue(MoveFraction);
 			}
+
 			MyActor->SetActorLocation(FMath::Lerp<FVector, float>(StartLocation, TargetLocation, MoveFraction));
 		}
 	}
 }
+
+void UAbilityTask_MoveToLocation::GetLifetimeReplicatedProps(TArray< FLifetimeProperty > & OutLifetimeProps) const
+{
+	DOREPLIFETIME(UAbilityTask_MoveToLocation, StartLocation);
+	DOREPLIFETIME(UAbilityTask_MoveToLocation, TargetLocation);
+	DOREPLIFETIME(UAbilityTask_MoveToLocation, DurationOfMovement);
+	DOREPLIFETIME(UAbilityTask_MoveToLocation, LerpCurve);
+}
\ No newline at end of file
diff --git a/Engine/Source/Runtime/GameplayAbilities/Private/AbilitySystemComponent.cpp b/Engine/Source/Runtime/GameplayAbilities/Private/AbilitySystemComponent.cpp
--- a/Engine/Source/Runtime/GameplayAbilities/Private/AbilitySystemComponent.cpp
+++ b/Engine/Source/Runtime/GameplayAbilities/Private/AbilitySystemComponent.cpp
@@ -644,23 +644,40 @@ void UAbilitySystemComponent::OnRestackGameplayEffects()
 
 // ---------------------------------------------------------------------------------------
 
-void UAbilitySystemComponent::TickingTaskStarted(UAbilityTask* NewTask)
+void UAbilitySystemComponent::TaskStarted(UAbilityTask* NewTask)
 {
-	// If this is our first ticking task, set this component as active so it begins ticking
-	if (TickingTasks.Num() == 0)
+	if (NewTask->bTickingTask)
 	{
-		SetActive(true);
+		// If this is our first ticking task, set this component as active so it begins ticking
+		if (TickingTasks.Num() == 0)
+		{
+			SetActive(true);
+		}
+		check(TickingTasks.Contains(NewTask) == false);
+		TickingTasks.Add(NewTask);
+	}
+	if (NewTask->bSimulatedTask)
+	{
+		check(SimulatedTasks.Contains(NewTask) == false);
+		SimulatedTasks.Add(NewTask);
 	}
-	TickingTasks.AddUnique(NewTask);
 }
 
-void UAbilitySystemComponent::TickingTaskEnded(UAbilityTask* Task)
+void UAbilitySystemComponent::TaskEnded(UAbilityTask* Task)
 {
-	// If we are removing our last ticking task, set this component as inactive so it stops ticking
-	TickingTasks.RemoveSingleSwap(Task);
-	if (TickingTasks.Num() == 0)
+	if (Task->bTickingTask)
+	{
+		// If we are removing our last ticking task, set this component as inactive so it stops ticking
+		TickingTasks.RemoveSingleSwap(Task);
+		if (TickingTasks.Num() == 0)
+		{
+			SetActive(false);
+		}
+	}
+
+	if (Task->bSimulatedTask)
 	{
-		SetActive(false);
+		SimulatedTasks.RemoveSingleSwap(Task);
 	}
 }
 
@@ -673,13 +690,12 @@ void UAbilitySystemComponent::GetLifetimeReplicatedProps(TArray< FLifetimeProper
 
 	DOREPLIFETIME(UAbilitySystemComponent, SpawnedAttributes);
 	DOREPLIFETIME(UAbilitySystemComponent, ActiveGameplayEffects);
-	DOREPLIFETIME(UAbilitySystemComponent, ActiveGameplayCues);
-	
+	DOREPLIFETIME(UAbilitySystemComponent, ActiveGameplayCues);	
 	DOREPLIFETIME(UAbilitySystemComponent, ActivatableAbilities);
-
 	DOREPLIFETIME(UAbilitySystemComponent, AbilityActor);
-
 	DOREPLIFETIME(UAbilitySystemComponent, ReplicatedPredictionKey);
+	
+	DOREPLIFETIME_CONDITION(UAbilitySystemComponent, SimulatedTasks, COND_SkipOwner);
 }
 
 bool UAbilitySystemComponent::ReplicateSubobjects(class UActorChannel *Channel, class FOutBunch *Bunch, FReplicationFlags *RepFlags)
@@ -702,6 +718,17 @@ bool UAbilitySystemComponent::ReplicateSubobjects(class UActorChannel *Channel,
 		}
 	}
 
+	if (!RepFlags->bNetOwner)
+	{
+		for (UAbilityTask* SimulatedTask : SimulatedTasks)
+		{
+			if (SimulatedTask && !SimulatedTask->HasAnyFlags(RF_PendingKill))
+			{
+				WroteSomething |= Channel->ReplicateSubobject(SimulatedTask, *Bunch, *RepFlags);
+			}
+		}
+	}
+
 	return WroteSomething;
 }
 
diff --git a/Engine/Source/Runtime/GameplayAbilities/Private/AbilitySystemComponent_Abilities.cpp b/Engine/Source/Runtime/GameplayAbilities/Private/AbilitySystemComponent_Abilities.cpp
--- a/Engine/Source/Runtime/GameplayAbilities/Private/AbilitySystemComponent_Abilities.cpp
+++ b/Engine/Source/Runtime/GameplayAbilities/Private/AbilitySystemComponent_Abilities.cpp
@@ -54,6 +54,7 @@ void UAbilitySystemComponent::TickComponent(float DeltaTime, enum ELevelTick Tic
 	// itself from the TickingTask list, we will make a copy of the tasks we want to service before ticking any
 
 	int32 NumTickingTasks = TickingTasks.Num();
+	int32 NumActuallyTicked = 0;
 	switch(NumTickingTasks)
 	{
 		case 0:
@@ -62,6 +63,7 @@ void UAbilitySystemComponent::TickComponent(float DeltaTime, enum ELevelTick Tic
 			if (TickingTasks[0].IsValid())
 			{
 				TickingTasks[0]->TickTask(DeltaTime);
+				NumActuallyTicked++;
 			}
 			break;
 		default:
@@ -72,11 +74,20 @@ void UAbilitySystemComponent::TickComponent(float DeltaTime, enum ELevelTick Tic
 					if (TickingTask.IsValid())
 					{
 						TickingTask->TickTask(DeltaTime);
+						NumActuallyTicked++;
 					}
 				}
 			}
 		break;
 	};
+	
+	// Stop ticking if no more active tasks
+	if (NumActuallyTicked == 0)
+	{
+		TickingTasks.SetNum(0, false);
+		SetActive(false);
+	}
+
 }
 
 void UAbilitySystemComponent::InitAbilityActorInfo(AActor* AvatarActor)
@@ -1016,6 +1027,24 @@ void UAbilitySystemComponent::ConsumeAbilityTargetData()
 	ReplicatedTargetData.Clear();
 }
 
+void UAbilitySystemComponent::OnRep_SimulatedTasks()
+{
+	for (UAbilityTask* SimulatedTask : SimulatedTasks)
+	{
+		// Temp check 
+		if (SimulatedTask && SimulatedTask->bTickingTask && TickingTasks.Contains(SimulatedTask) == false)
+		{
+			SimulatedTask->InitSimulatedTask(this);
+			if (TickingTasks.Num() == 0)
+			{
+				SetActive(true);
+			}
+
+			TickingTasks.Add(SimulatedTask);
+		}
+	}
+}
+
 
 #undef LOCTEXT_NAMESPACE
 
diff --git a/Engine/Source/Runtime/GameplayAbilities/Public/Abilities/Tasks/AbilityTask.h b/Engine/Source/Runtime/GameplayAbilities/Public/Abilities/Tasks/AbilityTask.h
--- a/Engine/Source/Runtime/GameplayAbilities/Public/Abilities/Tasks/AbilityTask.h
+++ b/Engine/Source/Runtime/GameplayAbilities/Public/Abilities/Tasks/AbilityTask.h
@@ -68,7 +68,9 @@ class GAMEPLAYABILITIES_API UAbilityTask : public UObject
 	virtual void Activate();
 
 	/** Initailizes the task with the owning GameplayAbility but does not actviate until Activate() is called */
-	virtual void InitTask(UGameplayAbility* InAbility);
+	void InitTask(UGameplayAbility* InAbility);
+
+	virtual void InitSimulatedTask(UAbilitySystemComponent* InAbilitySystemComponent);
 
 	/** Tick function for this task, if bTickingTask == true */
 	virtual void TickTask(float DeltaTime) {}
@@ -116,13 +118,24 @@ class GAMEPLAYABILITIES_API UAbilityTask : public UObject
 	UPROPERTY()
 	FName InstanceName;
 
+	bool IsSupportedForNetworking() const override
+	{
+		return bSimulatedTask;
+	}
+
+	/** If true, this task will receive TickTask calls from AbilitySystemComponent */
+	bool bTickingTask;
+
+	/** Should this task run on simulated clients? This should only be used in rare cases, such as movement tasks. Simulated Tasks do not broadcast their end delegates.  */
+	bool bSimulatedTask;
+
+	/** Am I actually running this as a simulated task. (This will be true on clients that simulating. This will be false on the server and the owning client) */
+	bool bIsSimulating;
+
 protected:	
 
 	/** End and CleanUp the task - may be called by the task itself or by the owning ability if the ability is ending. Do NOT call directly! Call EndTask() or AbilityEnded() */
 	virtual void OnDestroy(bool AbilityIsEnding);
-
-	/** If true, this task will receive TickTask calls from AbilitySystemComponent */
-	bool bTickingTask;
 };
 
 //For searching through lists of ability instances
diff --git a/Engine/Source/Runtime/GameplayAbilities/Public/Abilities/Tasks/AbilityTask_MoveToLocation.h b/Engine/Source/Runtime/GameplayAbilities/Public/Abilities/Tasks/AbilityTask_MoveToLocation.h
--- a/Engine/Source/Runtime/GameplayAbilities/Public/Abilities/Tasks/AbilityTask_MoveToLocation.h
+++ b/Engine/Source/Runtime/GameplayAbilities/Public/Abilities/Tasks/AbilityTask_MoveToLocation.h
@@ -10,8 +10,14 @@ class AActor;
 class UPrimitiveComponent;
 
 /**
-*		Move to a location, ignoring clipping, over a given length of time. Callback whenever we pass through an actor and when we reach our target location.
-*/
+ *	TODO:
+ *		-Implement replicated time so that this can work as a simulated task for Join In Prgorss clients.
+ */
+
+
+/**
+ *	Move to a location, ignoring clipping, over a given length of time. Ends when the TargetLocation is reached.
+ */
 UCLASS(MinimalAPI)
 class UAbilityTask_MoveToLocation : public UAbilityTask
 {
@@ -20,6 +26,8 @@ class UAbilityTask_MoveToLocation : public UAbilityTask
 	UPROPERTY(BlueprintAssignable)
 	FMoveToLocationDelegate		OnTargetLocationReached;
 
+	virtual void InitSimulatedTask(UAbilitySystemComponent* InAbilitySystemComponent) override;
+
 	/** Move to the specified location, using the curve (range 0 - 1) or linearly if no curve is specified */
 	UFUNCTION(BlueprintCallable, Category = "Ability|Tasks", meta = (HidePin = "WorldContextObject", DefaultToSelf = "WorldContextObject", BlueprintInternalUseOnly = "TRUE"))
 	static UAbilityTask_MoveToLocation* MoveToLocation(UObject* WorldContextObject, FVector Location, float Duration, UCurveFloat* OptionalInterpolationCurve);
@@ -31,10 +39,26 @@ class UAbilityTask_MoveToLocation : public UAbilityTask
 
 protected:
 
+	bool bIsFinished;
+
+	UPROPERTY(Replicated)
 	FVector StartLocation;
+
+	//FVector 
+	
+	UPROPERTY(Replicated)
 	FVector TargetLocation;
+
+	UPROPERTY(Replicated)
 	float DurationOfMovement;
+
+
 	float TimeMoveStarted;
+
 	float TimeMoveWillEnd;
-	TWeakObjectPtr<UCurveFloat> LerpCurve;
+
+	UPROPERTY(Replicated)
+	UCurveFloat* LerpCurve;
+
+	
 };
\ No newline at end of file
diff --git a/Engine/Source/Runtime/GameplayAbilities/Public/AbilitySystemComponent.h b/Engine/Source/Runtime/GameplayAbilities/Public/AbilitySystemComponent.h
--- a/Engine/Source/Runtime/GameplayAbilities/Public/AbilitySystemComponent.h
+++ b/Engine/Source/Runtime/GameplayAbilities/Public/AbilitySystemComponent.h
@@ -444,10 +444,11 @@ class GAMEPLAYABILITIES_API UAbilitySystemComponent : public UActorComponent, pu
 	// ----------------------------------------------------------------------------------------------------------------
 
 	/** Adds a UAbilityTask task to the list of tasks to be ticked */
-	void TickingTaskStarted(UAbilityTask* NewTask);
+	void TaskStarted(UAbilityTask* NewTask);
 
 	/** Removes a UAbilityTask task from the list of tasks to be ticked */
-	void TickingTaskEnded(UAbilityTask* Task);
+	void TaskEnded(UAbilityTask* Task);
+
 
 	/** There needs to be a concept of an animating ability. Only one may exist at a time. New requests can be queued up, overridden, or ignored. */
 	UPROPERTY()
@@ -516,6 +517,13 @@ class GAMEPLAYABILITIES_API UAbilitySystemComponent : public UActorComponent, pu
 	/** ReplicatedTargetData was 'cancelled' for this activation */
 	FAbilityConfirmOrCancel	ReplicatedTargetDataCancelledDelegate;
 
+	/** Tasks that run on simulated proxies */
+	UPROPERTY(ReplicatedUsing=OnRep_SimulatedTasks)
+	TArray<UAbilityTask*> SimulatedTasks;
+
+	UFUNCTION()
+	void OnRep_SimulatedTasks();
+
 private:
 
 	bool HasNetworkAuthorityToApplyGameplayEffect(const FModifierQualifier QualifierContext) const;
@@ -561,8 +569,7 @@ private:
 	// ---------------------------------------------
 
 	/** Array of currently active UAbilityTasks that require ticking */
-	TArray<TWeakObjectPtr<UAbilityTask> >	TickingTasks;
-	
+	TArray<TWeakObjectPtr<UAbilityTask> >	TickingTasks;	
 
 protected:
 
