comment

diff --git a/Engine/Source/Runtime/GameplayAbilities/Public/Abilities/Tasks/AbilityTask_MoveToLocation.h b/Engine/Source/Runtime/GameplayAbilities/Public/Abilities/Tasks/AbilityTask_MoveToLocation.h
--- a/Engine/Source/Runtime/GameplayAbilities/Public/Abilities/Tasks/AbilityTask_MoveToLocation.h
+++ b/Engine/Source/Runtime/GameplayAbilities/Public/Abilities/Tasks/AbilityTask_MoveToLocation.h
@@ -17,6 +17,8 @@ class UPrimitiveComponent;
 
 /**
  *	Move to a location, ignoring clipping, over a given length of time. Ends when the TargetLocation is reached.
+ *	This will RESET your character's current movement mode! If you wish to maintain PHYS_Flying or PHYS_Custom, you must
+ *	reset it on completion.!
  */
 UCLASS(MinimalAPI)
 class UAbilityTask_MoveToLocation : public UAbilityTask
