-AbilityTask_MoveToLocation forces characters into custom move mode
-remove TargetActor parameter from location based directed movement

diff --git a/Engine/Source/Runtime/GameplayAbilities/Private/Abilities/Tasks/AbilityTask_MoveToLocation.cpp b/Engine/Source/Runtime/GameplayAbilities/Private/Abilities/Tasks/AbilityTask_MoveToLocation.cpp
--- a/Engine/Source/Runtime/GameplayAbilities/Private/Abilities/Tasks/AbilityTask_MoveToLocation.cpp
+++ b/Engine/Source/Runtime/GameplayAbilities/Private/Abilities/Tasks/AbilityTask_MoveToLocation.cpp
@@ -1,6 +1,8 @@
 
 #include "AbilitySystemPrivatePCH.h"
 #include "Abilities/Tasks/AbilityTask_MoveToLocation.h"
+#include "GameFramework/Character.h"
+#include "GameFramework/CharacterMovementComponent.h"
 
 UAbilityTask_MoveToLocation::UAbilityTask_MoveToLocation(const class FPostConstructInitializeProperties& PCIP)
 : Super(PCIP)
@@ -49,6 +51,17 @@ void UAbilityTask_MoveToLocation::TickTask(float DeltaTime)
 	AActor* MyActor = GetActor();
 	if (MyActor)
 	{
+		ACharacter* MyCharacter = Cast<ACharacter>(MyActor);
+		if (MyCharacter)
+		{
+			UCharacterMovementComponent* CharMoveComp = Cast<UCharacterMovementComponent>(MyCharacter->GetMovementComponent());
+			if (CharMoveComp
+			{
+				CharMoveComp->SetMovementMode(MOVE_Custom, 0);
+			}
+		}
+
+
 		float CurrentTime = GetWorld()->GetTimeSeconds();
 
 		if (CurrentTime >= TimeMoveWillEnd)
