cis fix

diff --git a/Engine/Source/Runtime/GameplayAbilities/Private/Abilities/Tasks/AbilityTask_MoveToLocation.cpp b/Engine/Source/Runtime/GameplayAbilities/Private/Abilities/Tasks/AbilityTask_MoveToLocation.cpp
--- a/Engine/Source/Runtime/GameplayAbilities/Private/Abilities/Tasks/AbilityTask_MoveToLocation.cpp
+++ b/Engine/Source/Runtime/GameplayAbilities/Private/Abilities/Tasks/AbilityTask_MoveToLocation.cpp
@@ -55,7 +55,7 @@ void UAbilityTask_MoveToLocation::TickTask(float DeltaTime)
 		if (MyCharacter)
 		{
 			UCharacterMovementComponent* CharMoveComp = Cast<UCharacterMovementComponent>(MyCharacter->GetMovementComponent());
-			if (CharMoveComp
+			if (CharMoveComp)
 			{
 				CharMoveComp->SetMovementMode(MOVE_Custom, 0);
 			}
