No need to call SetIsReplicated, just set bReplicates

diff --git a/Engine/Source/Runtime/GameplayAbilities/Private/AbilitySystemComponent.cpp b/Engine/Source/Runtime/GameplayAbilities/Private/AbilitySystemComponent.cpp
--- a/Engine/Source/Runtime/GameplayAbilities/Private/AbilitySystemComponent.cpp
+++ b/Engine/Source/Runtime/GameplayAbilities/Private/AbilitySystemComponent.cpp
@@ -41,7 +41,7 @@ UAbilitySystemComponent::UAbilitySystemComponent(const class FPostConstructIniti
 	ActiveGameplayEffects.Owner = this;
 	ActiveGameplayCues.Owner = this;
 
-	SetIsReplicated(true);
+	bReplicates = true;
 
 	UserAbilityActivationInhibited = false;
 }
