Target filtering system created.

diff --git a/Engine/Source/Runtime/GameplayAbilities/Private/Abilities/GameplayAbilityTargetDataFilter.cpp b/Engine/Source/Runtime/GameplayAbilities/Private/Abilities/GameplayAbilityTargetDataFilter.cpp
new file mode 100644
--- /dev/null
+++ b/Engine/Source/Runtime/GameplayAbilities/Private/Abilities/GameplayAbilityTargetDataFilter.cpp
@@ -0,0 +1,17 @@
+// Copyright 1998-2013 Epic Games, Inc. All Rights Reserved.
+
+#include "AbilitySystemPrivatePCH.h"
+#include "GameplayAbilityTargetDataFilter.h"
+
+// --------------------------------------------------------------------------------------------------------------------------------------------------------
+//
+//	AGameplayAbilityTargetDataActorFilter
+//
+// --------------------------------------------------------------------------------------------------------------------------------------------------------
+
+/*
+AGameplayAbilityTargetDataActorFilter::AGameplayAbilityTargetDataActorFilter(const class FPostConstructInitializeProperties& PCIP)
+	: Super(PCIP)
+{
+}
+*/
\ No newline at end of file
diff --git a/Engine/Source/Runtime/GameplayAbilities/Private/AbilitySystemBlueprintLibrary.cpp b/Engine/Source/Runtime/GameplayAbilities/Private/AbilitySystemBlueprintLibrary.cpp
--- a/Engine/Source/Runtime/GameplayAbilities/Private/AbilitySystemBlueprintLibrary.cpp
+++ b/Engine/Source/Runtime/GameplayAbilities/Private/AbilitySystemBlueprintLibrary.cpp
@@ -50,17 +50,122 @@ FGameplayAbilityTargetDataHandle UAbilitySystemBlueprintLibrary::AbilityTargetDa
 	return Handle;
 }
 
-
 FGameplayAbilityTargetDataHandle UAbilitySystemBlueprintLibrary::AbilityTargetDataFromActor(AActor* Actor)
 {
 	// Construct TargetData
 	FGameplayAbilityTargetData_ActorArray*	NewData = new FGameplayAbilityTargetData_ActorArray();
 	NewData->TargetActorArray.Add(Actor);
-
 	FGameplayAbilityTargetDataHandle		Handle(NewData);
 	return Handle;
 }
 
+FGameplayAbilityTargetDataHandle UAbilitySystemBlueprintLibrary::AbilityTargetDataFromActorArray(TArray<TWeakObjectPtr<AActor>> ActorArray, bool OneTargetPerHandle)
+{
+	// Construct TargetData
+	if (OneTargetPerHandle)
+	{
+		FGameplayAbilityTargetDataHandle		Handle;
+		for (int32 i = 0; i < ActorArray.Num(); ++i)
+		{
+			if (ActorArray[i].IsValid())
+			{
+				FGameplayAbilityTargetDataHandle TempHandle = AbilityTargetDataFromActor(ActorArray[i].Get());
+				Handle.Append(&TempHandle);
+			}
+		}
+		return Handle;
+	}
+	else
+	{
+		FGameplayAbilityTargetData_ActorArray*	NewData = new FGameplayAbilityTargetData_ActorArray();
+		NewData->TargetActorArray = ActorArray;
+		FGameplayAbilityTargetDataHandle		Handle(NewData);
+		return Handle;
+	}
+}
+
+
+FGameplayAbilityTargetDataHandle UAbilitySystemBlueprintLibrary::FilterTargetData(FGameplayAbilityTargetDataHandle TargetDataHandle, FGameplayTargetDataFilterHandle FilterHandle)
+{
+	FGameplayAbilityTargetDataHandle ReturnDataHandle;
+	
+	//return ReturnDataHandle;
+	check(FilterHandle.Filter.IsValid());		//Return unfiltered data
+
+	for (int32 i = 0; TargetDataHandle.IsValid(i); ++i)
+	{
+		FGameplayAbilityTargetData* UnfilteredData = TargetDataHandle.Get(i);
+		check(UnfilteredData);
+		if (UnfilteredData->GetActors().Num() > 0)
+		{
+			//TArray<TWeakObjectPtr<AActor>> FilteredActors = FilterActorArray(UnfilteredData->GetActors(), Filter);
+			TArray<TWeakObjectPtr<AActor>> FilteredActors = UnfilteredData->GetActors().FilterByPredicate(*FilterHandle.Filter.Get());
+			if (FilteredActors.Num() > 0)
+			{
+				//Copy the data first, since we don't understand the internals of it
+				UScriptStruct* ScriptStruct = UnfilteredData->GetScriptStruct();
+				FGameplayAbilityTargetData* NewData = (FGameplayAbilityTargetData*)FMemory::Malloc(ScriptStruct->GetCppStructOps()->GetSize());
+				ScriptStruct->InitializeScriptStruct(NewData);
+				ScriptStruct->CopyScriptStruct(NewData, UnfilteredData);
+				ReturnDataHandle.Data.Add(TSharedPtr<FGameplayAbilityTargetData>(NewData));
+				if (FilteredActors.Num() < UnfilteredData->GetActors().Num())
+				{
+					//We have lost some, but not all, of our actors, so replace the array. This should only be possible with targeting types that permit actor-array setting.
+					if (!NewData->SetActors(FilteredActors))
+					{
+						//This is an error, though we could ignore it. We somehow filtered out part of a list, but the class doesn't support changing the list, so now it's all or nothing.
+						check(false);
+					}
+				}
+			}
+		}
+	}
+
+	return ReturnDataHandle;
+}
+
+FGameplayTargetDataFilterHandle UAbilitySystemBlueprintLibrary::MakeFilterHandle(FGameplayTargetDataFilter Filter)
+{
+	FGameplayTargetDataFilterHandle FilterHandle;
+	FGameplayTargetDataFilter* NewFilter = new FGameplayTargetDataFilter(Filter);
+	FilterHandle.Filter = TSharedPtr<FGameplayTargetDataFilter>(NewFilter);
+	return FilterHandle;
+}
+
+
+/*
+TArray<TWeakObjectPtr<AActor>> UGameplayAbility::FilterActorArray(TArray<TWeakObjectPtr<AActor>> ArrayToFilter, FGameplayAbilityTargetDataActorFilter Filter)
+{
+	TArray<TWeakObjectPtr<AActor>> ReturnArray;
+	int32 NewArraySize = 0;
+	for (int32 i = 0; i < ArrayToFilter.Num(); ++i)
+	{
+		if (ArrayToFilter[i].IsValid())
+		{
+			if (Filter.FilterPassesForActor(ArrayToFilter[i].Get()))
+			{
+				++NewArraySize;
+			}
+		}
+	}
+	if (NewArraySize)
+	{
+		ReturnArray.Reserve(NewArraySize);
+		for (int32 i = 0; i < ArrayToFilter.Num(); ++i)
+		{
+			if (ArrayToFilter[i].IsValid())
+			{
+				if (Filter.FilterPassesForActor(ArrayToFilter[i].Get()))
+				{
+					ReturnArray.Add(ArrayToFilter[i].Get());
+				}
+			}
+		}
+	}
+	return ReturnArray;
+}
+*/
+
 FGameplayAbilityTargetDataHandle UAbilitySystemBlueprintLibrary::AbilityTargetDataFromHitResult(FHitResult HitResult)
 {
 	// Construct TargetData
diff --git a/Engine/Source/Runtime/GameplayAbilities/Public/Abilities/GameplayAbility.h b/Engine/Source/Runtime/GameplayAbilities/Public/Abilities/GameplayAbility.h
--- a/Engine/Source/Runtime/GameplayAbilities/Public/Abilities/GameplayAbility.h
+++ b/Engine/Source/Runtime/GameplayAbilities/Public/Abilities/GameplayAbility.h
@@ -385,6 +385,8 @@ protected:
 		}
 	}
 
+	public:
+
 	UFUNCTION(BlueprintCallable, Category=Ability)
 	FGameplayAbilityActorInfo GetActorInfo() const;
 
@@ -395,7 +397,9 @@ protected:
 	USkeletalMeshComponent* GetOwningComponentFromActorInfo() const;
 
 	UFUNCTION(BlueprintCallable, Category=Ability)
-	FGameplayEffectSpecHandle GetOutgoingGameplayEffectSpec(UGameplayEffect* GameplayEffect, float Level=1.f) const;	
+	FGameplayEffectSpecHandle GetOutgoingGameplayEffectSpec(UGameplayEffect* GameplayEffect, float Level=1.f) const;
+
+	protected:
 
 	bool IsSupportedForNetworking() const override;
 
diff --git a/Engine/Source/Runtime/GameplayAbilities/Public/Abilities/GameplayAbilityTargetDataFilter.h b/Engine/Source/Runtime/GameplayAbilities/Public/Abilities/GameplayAbilityTargetDataFilter.h
new file mode 100644
--- /dev/null
+++ b/Engine/Source/Runtime/GameplayAbilities/Public/Abilities/GameplayAbilityTargetDataFilter.h
@@ -0,0 +1,35 @@
+// Copyright 1998-2014 Epic Games, Inc. All Rights Reserved.
+
+#pragma once
+
+#include "GameplayAbilityTargetDataFilter.generated.h"
+
+USTRUCT(BlueprintType)
+struct FGameplayTargetDataFilter
+{
+	GENERATED_USTRUCT_BODY()
+
+	virtual bool FilterPassesForActor(const AActor* ActorToBeFiltered) const
+	{
+		return true;
+	}
+
+	bool operator()(const TWeakObjectPtr<AActor> A) const
+	{
+		return FilterPassesForActor(A.Get());
+	}
+
+	bool operator()(const AActor* A) const
+	{
+		return FilterPassesForActor(A);
+	}
+};
+
+
+USTRUCT(BlueprintType)
+struct FGameplayTargetDataFilterHandle
+{
+	GENERATED_USTRUCT_BODY()
+
+	TSharedPtr<FGameplayTargetDataFilter> Filter;
+};
diff --git a/Engine/Source/Runtime/GameplayAbilities/Public/Abilities/GameplayAbilityTargetTypes.h b/Engine/Source/Runtime/GameplayAbilities/Public/Abilities/GameplayAbilityTargetTypes.h
--- a/Engine/Source/Runtime/GameplayAbilities/Public/Abilities/GameplayAbilityTargetTypes.h
+++ b/Engine/Source/Runtime/GameplayAbilities/Public/Abilities/GameplayAbilityTargetTypes.h
@@ -70,6 +70,12 @@ struct GAMEPLAYABILITIES_API FGameplayAbilityTargetData
 		return TArray<TWeakObjectPtr<AActor> >();
 	}
 
+	virtual bool SetActors(TArray<TWeakObjectPtr<AActor>> NewActorArray)
+	{
+		//By default, we don't keep this data, and therefore can't set it.
+		return false;
+	}
+
 	// -------------------------------------
 
 	virtual bool HasHitResult() const
@@ -238,6 +244,12 @@ struct GAMEPLAYABILITIES_API FGameplayAbilityTargetData_Radius : public FGamepla
 
 	virtual TArray<TWeakObjectPtr<AActor> >	GetActors() const { return Actors; }
 
+	virtual bool SetActors(TArray<TWeakObjectPtr<AActor>> NewActorArray) override
+	{
+		Actors = NewActorArray;
+		return true;
+	}
+
 	virtual bool HasOrigin() const { return true; }
 
 	virtual FTransform GetOrigin() const { return FTransform(Origin); }
@@ -264,6 +276,45 @@ struct TStructOpsTypeTraits<FGameplayAbilityTargetData_Radius> : public TStructO
 	};
 };
 
+/*
+USTRUCT(BlueprintType)
+struct FGameplayAbilityTargetDataActorFilter
+{
+	GENERATED_USTRUCT_BODY()
+
+	virtual bool FilterPassesForActor(const AActor* ActorToBeFiltered) const
+	{
+		return true;
+	}
+};
+
+USTRUCT(BlueprintType)
+struct FGameplayAbilityTargetDataActorFilterHandleBase
+{
+	GENERATED_USTRUCT_BODY()
+
+	TSharedPtr<FGameplayAbilityTargetDataActorFilter>	Filter;
+
+	bool operator()(const TWeakObjectPtr<AActor> A) const
+	{
+		if (Filter.IsValid())
+		{
+			return Filter.Get()->FilterPassesForActor(A.Get());
+		}
+		return true;
+	}
+
+	bool operator()(const AActor* A) const
+	{
+		if (Filter.IsValid())
+		{
+			return Filter.Get()->FilterPassesForActor(A);
+		}
+		return true;
+	}
+};
+*/
+
 USTRUCT(BlueprintType)
 struct GAMEPLAYABILITIES_API FGameplayAbilityTargetingLocationInfo
 {
@@ -315,7 +366,7 @@ public:
 
 	FGameplayAbilityTargetDataHandle MakeTargetDataHandleFromHitResult(TWeakObjectPtr<UGameplayAbility> Ability, FHitResult HitResult) const;
 
-	FGameplayAbilityTargetDataHandle MakeTargetDataHandleFromActors(TArray<TWeakObjectPtr<AActor> > TargetActors, bool OneActorPerHandle = false) const;
+	FGameplayAbilityTargetDataHandle MakeTargetDataHandleFromActors(TArray<TWeakObjectPtr<AActor>> TargetActors, bool OneActorPerHandle = false) const;
 
 	/** Type of location used - will determine what data is transmitted over the network and what fields are used when calculating position. */
 	UPROPERTY(BlueprintReadWrite, meta = (ExposeOnSpawn = true), Category = Targeting)
@@ -440,6 +491,12 @@ struct GAMEPLAYABILITIES_API FGameplayAbilityTargetData_ActorArray : public FGam
 		return TargetActorArray;
 	}
 
+	virtual bool SetActors(TArray<TWeakObjectPtr<AActor>> NewActorArray) override
+	{
+		TargetActorArray = NewActorArray;
+		return true;
+	}
+
 	// -------------------------------------
 
 	virtual bool HasOrigin() const override
@@ -624,6 +681,8 @@ struct GAMEPLAYABILITIES_API FGameplayAbilityTargetData_SingleTargetHit : public
 		return Actors;
 	}
 
+	// SetActors() will not work here because the actor "array" is drawn from the hit result data, and changing that doesn't make sense.
+
 	// -------------------------------------
 
 	virtual bool HasHitResult() const override
diff --git a/Engine/Source/Runtime/GameplayAbilities/Public/AbilitySystemBlueprintLibrary.h b/Engine/Source/Runtime/GameplayAbilities/Public/AbilitySystemBlueprintLibrary.h
--- a/Engine/Source/Runtime/GameplayAbilities/Public/AbilitySystemBlueprintLibrary.h
+++ b/Engine/Source/Runtime/GameplayAbilities/Public/AbilitySystemBlueprintLibrary.h
@@ -3,6 +3,7 @@
 #pragma once
 
 #include "Abilities/GameplayAbility.h"
+#include "Abilities/GameplayAbilityTargetDataFilter.h"
 #include "GameplayCueView.h"
 #include "AbilitySystemBlueprintLibrary.generated.h"
 
@@ -43,6 +44,15 @@ class UAbilitySystemBlueprintLibrary : public UBlueprintFunctionLibrary
 	UFUNCTION(BlueprintPure, Category = "Ability|TargetData")
 	static FGameplayAbilityTargetDataHandle	AbilityTargetDataFromActor(AActor* Actor);
 
+	UFUNCTION(BlueprintPure, Category = "Ability|TargetData")
+	static FGameplayAbilityTargetDataHandle	AbilityTargetDataFromActorArray(TArray<TWeakObjectPtr<AActor>> ActorArray, bool OneTargetPerHandle);
+
+	UFUNCTION(BlueprintPure, Category = "Ability|TargetData")
+	static FGameplayAbilityTargetDataHandle	FilterTargetData(FGameplayAbilityTargetDataHandle TargetDataHandle, FGameplayTargetDataFilterHandle ActorFilterClass);
+
+	UFUNCTION(BlueprintPure, Category = "Ability|TargetData")
+	static FGameplayTargetDataFilterHandle MakeFilterHandle(FGameplayTargetDataFilter Filter);
+
 	UFUNCTION(BlueprintPure, Category = "Ability|TargetData")
 	static TArray<AActor*> GetActorsFromTargetData(FGameplayAbilityTargetDataHandle TargetData, int32 Index);
 
