Fixed FGameplayAbilityTargetData_ActorArray::GetOrigin()

* 'FGameplayAbilityTargetData_ActorArray::GetOrigin()' was not returning the correct rotation.

diff --git a/Engine/Source/Runtime/GameplayAbilities/Public/Abilities/GameplayAbilityTargetTypes.h b/Engine/Source/Runtime/GameplayAbilities/Public/Abilities/GameplayAbilityTargetTypes.h
--- a/Engine/Source/Runtime/GameplayAbilities/Public/Abilities/GameplayAbilityTargetTypes.h
+++ b/Engine/Source/Runtime/GameplayAbilities/Public/Abilities/GameplayAbilityTargetTypes.h
@@ -513,8 +513,12 @@ struct GAMEPLAYABILITIES_API FGameplayAbilityTargetData_ActorArray : public FGam
 		{
 			if (TargetActorArray[i].IsValid())
 			{
-				ReturnTransform.SetRotation((TargetActorArray[i].Get()->GetActorLocation()).SafeNormal().Rotation().Quaternion());
-				break;
+				FVector Direction = (TargetActorArray[i].Get()->GetActorLocation() - ReturnTransform.GetLocation()).SafeNormal();
+				if (Direction.IsNormalized())
+				{
+					ReturnTransform.SetRotation(Direction.Rotation().Quaternion());
+					break;
+				}
 			}
 		}
 		return ReturnTransform;
