Fixed Clang compile error

diff --git a/Engine/Source/Runtime/GameplayAbilities/Public/Abilities/GameplayAbilityTargetDataFilter.h b/Engine/Source/Runtime/GameplayAbilities/Public/Abilities/GameplayAbilityTargetDataFilter.h
--- a/Engine/Source/Runtime/GameplayAbilities/Public/Abilities/GameplayAbilityTargetDataFilter.h
+++ b/Engine/Source/Runtime/GameplayAbilities/Public/Abilities/GameplayAbilityTargetDataFilter.h
@@ -9,6 +9,10 @@ struct FGameplayTargetDataFilter
 {
 	GENERATED_USTRUCT_BODY()
 
+	virtual ~FGameplayTargetDataFilter()
+	{
+	}
+
 	virtual bool FilterPassesForActor(const AActor* ActorToBeFiltered) const
 	{
 		return true;
