Fixed aiming disconnect between 2D and 3D reticles. Added some ray-sphere intersection to the base class so we can use this elsewhere, though maybe it should be slightly more generalized and go to a math library.

diff --git a/Engine/Source/Runtime/GameplayAbilities/Private/Abilities/GameplayAbilityTargetActor_SingleLineTrace.cpp b/Engine/Source/Runtime/GameplayAbilities/Private/Abilities/GameplayAbilityTargetActor_SingleLineTrace.cpp
--- a/Engine/Source/Runtime/GameplayAbilities/Private/Abilities/GameplayAbilityTargetActor_SingleLineTrace.cpp
+++ b/Engine/Source/Runtime/GameplayAbilities/Private/Abilities/GameplayAbilityTargetActor_SingleLineTrace.cpp
@@ -5,7 +5,7 @@
 #include "GameplayAbilityTargetActor_SingleLineTrace.h"
 #include "Engine/World.h"
 #include "Runtime/Engine/Public/Net/UnrealNetwork.h"
-
+#pragma optimize("",off)
 // --------------------------------------------------------------------------------------------------------------------------------------------------------
 //
 //	AGameplayAbilityTargetActor_SingleLineTrace
@@ -30,8 +30,9 @@ FHitResult AGameplayAbilityTargetActor_SingleLineTrace::PerformTrace(AActor* InS
 	Params.bTraceAsyncScene = true;
 	Params.AddIgnoredActors(ActorsToIgnore);
 
-	FVector AimDirection = InSourceActor->GetActorForwardVector();		//Default
-	
+	FVector TraceStart = InSourceActor->GetActorLocation();
+	FVector TraceEnd = TraceStart + (InSourceActor->GetActorForwardVector() * MaxRange);		//Default
+
 	if (OwningAbility)		//Server and launching client only
 	{
 		APlayerController* AimingPC = OwningAbility->GetCurrentActorInfo()->PlayerController.Get();
@@ -39,18 +40,26 @@ FHitResult AGameplayAbilityTargetActor_SingleLineTrace::PerformTrace(AActor* InS
 		FVector CamLoc;
 		FRotator CamRot;
 		AimingPC->GetPlayerViewPoint(CamLoc, CamRot);
-		AimDirection = CamRot.Vector();
-	}
+		FVector CamDir = CamRot.Vector();
+		FVector CamTarget = CamLoc + (CamDir * MaxRange);		//Straight, dumb aiming to a point that's reasonable though not exactly correct
 
-	FVector TraceStart = InSourceActor->GetActorLocation();
-	FVector TraceEnd = TraceStart + (AimDirection * MaxRange);
+		ClipCameraRayToAbilityRange(CamLoc, CamDir, TraceStart, MaxRange, CamTarget);
 
-	//If we're using a socket, adjust the starting location and aim direction after the end position has been found. This way we can still aim with the camera, then fire accurately from the socket.
-	if (OwningAbility)		//Server and launching client only
-	{
+		FHitResult TempHitResult;
+		InSourceActor->GetWorld()->LineTraceSingle(TempHitResult, CamLoc, CamTarget, ECC_WorldStatic, Params);
 		TraceStart = StartLocation.GetTargetingTransform().GetLocation();
+		if (TempHitResult.bBlockingHit && (FVector::DistSquared(TraceStart, TempHitResult.Location) <= (MaxRange * MaxRange)))
+		{
+			//We actually made a hit? Pull back.
+			TraceEnd = TempHitResult.Location;
+		}
+		else
+		{
+			//If we didn't make a hit, use the clipped location.
+			TraceEnd = CamTarget;
+		}
 	}
-	AimDirection = (TraceEnd - TraceStart).SafeNormal();
+	FVector AimDirection = (TraceStart - TraceEnd).SafeNormal();
 
 	// ------------------------------------------------------
 
@@ -65,5 +74,11 @@ FHitResult AGameplayAbilityTargetActor_SingleLineTrace::PerformTrace(AActor* InS
 	{
 		LocalReticleActor->SetActorLocation(ReturnHitResult.Location);
 	}
+
+	if (bDebug)
+	{
+		DrawDebugLine(GetWorld(), TraceStart, TraceEnd, FColor::Green);
+		DrawDebugSphere(GetWorld(), TraceEnd, 100.0f, 16, FColor::Green);
+	}
 	return ReturnHitResult;
 }
diff --git a/Engine/Source/Runtime/GameplayAbilities/Private/Abilities/GameplayAbilityTargetActor_Trace.cpp b/Engine/Source/Runtime/GameplayAbilities/Private/Abilities/GameplayAbilityTargetActor_Trace.cpp
--- a/Engine/Source/Runtime/GameplayAbilities/Private/Abilities/GameplayAbilityTargetActor_Trace.cpp
+++ b/Engine/Source/Runtime/GameplayAbilities/Private/Abilities/GameplayAbilityTargetActor_Trace.cpp
@@ -38,6 +38,25 @@ FGameplayAbilityTargetDataHandle AGameplayAbilityTargetActor_Trace::StaticGetTar
 	return MakeTargetData(PerformTrace(StaticSourceActor));
 }
 
+bool AGameplayAbilityTargetActor_Trace::ClipCameraRayToAbilityRange(FVector CameraLocation, FVector CameraDirection, FVector AbilityCenter, float AbilityRange, FVector& ClippedPosition)
+{
+	FVector CameraToCenter = AbilityCenter - CameraLocation;
+	float DotToCenter = FVector::DotProduct(CameraToCenter, CameraDirection);
+	if (DotToCenter >= 0)		//If this fails, we're pointed away from the center, but we might be inside the sphere and able to find a good exit point.
+	{
+		float DistanceSquared = CameraToCenter.SizeSquared() - (DotToCenter * DotToCenter);
+		float RadiusSquared = (AbilityRange * AbilityRange);
+		if (DistanceSquared <= RadiusSquared)
+		{
+			float DistanceFromCamera = FMath::Sqrt(RadiusSquared - DistanceSquared);
+			float DistanceAlongRay = DotToCenter + DistanceFromCamera;						//Subtracting instead of adding will get the other intersection point
+			ClippedPosition = CameraLocation + (DistanceAlongRay * CameraDirection);		//Cam aim point clipped to range sphere
+			return true;
+		}
+	}
+	return false;
+}
+
 void AGameplayAbilityTargetActor_Trace::StartTargeting(UGameplayAbility* InAbility)
 {
 	Super::StartTargeting(InAbility);
diff --git a/Engine/Source/Runtime/GameplayAbilities/Public/Abilities/GameplayAbilityTargetActor_Trace.h b/Engine/Source/Runtime/GameplayAbilities/Public/Abilities/GameplayAbilityTargetActor_Trace.h
--- a/Engine/Source/Runtime/GameplayAbilities/Public/Abilities/GameplayAbilityTargetActor_Trace.h
+++ b/Engine/Source/Runtime/GameplayAbilities/Public/Abilities/GameplayAbilityTargetActor_Trace.h
@@ -17,6 +17,8 @@ public:
 
 	virtual FGameplayAbilityTargetDataHandle StaticGetTargetData(UWorld * World, const FGameplayAbilityActorInfo* ActorInfo, FGameplayAbilityActivationInfo ActivationInfo) const override;
 
+	static bool ClipCameraRayToAbilityRange(FVector CameraLocation, FVector CameraDirection, FVector AbilityCenter, float AbilityRange, FVector& ClippedPosition);
+
 	virtual void StartTargeting(UGameplayAbility* Ability) override;
 
 	virtual void ConfirmTargetingAndContinue() override;
