#UE4 - Changed GetTargetDataEndPoint() to use Location instead of ImpactPoint or TraceEnd, since the latter two aren't serialized

diff --git a/Engine/Source/Runtime/GameplayAbilities/Private/AbilitySystemBlueprintLibrary.cpp b/Engine/Source/Runtime/GameplayAbilities/Private/AbilitySystemBlueprintLibrary.cpp
--- a/Engine/Source/Runtime/GameplayAbilities/Private/AbilitySystemBlueprintLibrary.cpp
+++ b/Engine/Source/Runtime/GameplayAbilities/Private/AbilitySystemBlueprintLibrary.cpp
@@ -288,14 +288,7 @@ FVector UAbilitySystemBlueprintLibrary::GetTargetDataEndPoint(FGameplayAbilityTa
 			const FHitResult* HitResultPtr = Data->GetHitResult();
 			if (HitResultPtr)
 			{
-				if (HitResultPtr->Component.IsValid())
-				{
-					return HitResultPtr->ImpactPoint;
-				}
-				else
-				{
-					return HitResultPtr->TraceEnd;
-				}
+				return HitResultPtr->Location;
 			}
 			else if (Data->HasEndPoint())
 			{
