Replicated montage support

diff --git a/Engine/Source/Runtime/GameplayAbilities/Private/Abilities/GameplayAbility.cpp b/Engine/Source/Runtime/GameplayAbilities/Private/Abilities/GameplayAbility.cpp
--- a/Engine/Source/Runtime/GameplayAbilities/Private/Abilities/GameplayAbility.cpp
+++ b/Engine/Source/Runtime/GameplayAbilities/Private/Abilities/GameplayAbility.cpp
@@ -481,14 +481,28 @@ void UGameplayAbility::K2_EndAbility()
 
 // --------------------------------------------------------------------
 
-void UGameplayAbility::MontageBranchPoint_AbilityDecisionStop(const FGameplayAbilityActorInfo* ActorInfo) const
+void UGameplayAbility::MontageJumpToSection(FName SectionName)
 {
+	check(CurrentActorInfo);
+
+	CurrentActorInfo->AbilitySystemComponent->LocalAnimMontageInfo.AnimMontage;
 
+	if (CurrentActorInfo->AbilitySystemComponent->IsAnimatingAbility(this))
+	{
+		CurrentActorInfo->AbilitySystemComponent->CurrentMontageJumpToSection(SectionName);
+	}
 }
 
-void UGameplayAbility::MontageBranchPoint_AbilityDecisionStart(const FGameplayAbilityActorInfo* ActorInfo) const
+
+void UGameplayAbility::MontageStop()
 {
+	check(CurrentActorInfo);
 
+	// We should only stop the current montage if we are the animating ability
+	if (CurrentActorInfo->AbilitySystemComponent->IsAnimatingAbility(this))
+	{
+		CurrentActorInfo->AbilitySystemComponent->CurrentMontageStop();
+	}
 }
 
 // --------------------------------------------------------------------
diff --git a/Engine/Source/Runtime/GameplayAbilities/Private/Abilities/Tasks/AbilityTask_PlayMontageAndWait.cpp b/Engine/Source/Runtime/GameplayAbilities/Private/Abilities/Tasks/AbilityTask_PlayMontageAndWait.cpp
--- a/Engine/Source/Runtime/GameplayAbilities/Private/Abilities/Tasks/AbilityTask_PlayMontageAndWait.cpp
+++ b/Engine/Source/Runtime/GameplayAbilities/Private/Abilities/Tasks/AbilityTask_PlayMontageAndWait.cpp
@@ -33,13 +33,12 @@ UAbilityTask_PlayMontageAndWait* UAbilityTask_PlayMontageAndWait::CreatePlayMont
 
 void UAbilityTask_PlayMontageAndWait::Activate()
 {
-	if (Ability.IsValid())
+	if (AbilitySystemComponent.IsValid() && Ability.IsValid())
 	{
 		const FGameplayAbilityActorInfo* ActorInfo = Ability->GetCurrentActorInfo();
 		if (ActorInfo->AnimInstance.IsValid())
 		{
-			ActorInfo->AnimInstance->Montage_Play(MontageToPlay, Rate);
-			if (MontageToPlay)
+			if (AbilitySystemComponent->PlayMontage(Ability.Get(), Ability->GetCurrentActivationInfo(), MontageToPlay, Rate	) > 0.f)
 			{
 				FOnMontageEnded EndDelegate;
 				EndDelegate.BindUObject(this, &UAbilityTask_PlayMontageAndWait::OnMontageEnded);
diff --git a/Engine/Source/Runtime/GameplayAbilities/Private/AbilitySystemComponent.cpp b/Engine/Source/Runtime/GameplayAbilities/Private/AbilitySystemComponent.cpp
--- a/Engine/Source/Runtime/GameplayAbilities/Private/AbilitySystemComponent.cpp
+++ b/Engine/Source/Runtime/GameplayAbilities/Private/AbilitySystemComponent.cpp
@@ -651,7 +651,7 @@ void UAbilitySystemComponent::TaskStarted(UAbilityTask* NewTask)
 		// If this is our first ticking task, set this component as active so it begins ticking
 		if (TickingTasks.Num() == 0)
 		{
-			SetActive(true);
+			UpdateShouldTick();
 		}
 		check(TickingTasks.Contains(NewTask) == false);
 		TickingTasks.Add(NewTask);
@@ -671,7 +671,7 @@ void UAbilitySystemComponent::TaskEnded(UAbilityTask* Task)
 		TickingTasks.RemoveSingleSwap(Task);
 		if (TickingTasks.Num() == 0)
 		{
-			SetActive(false);
+			UpdateShouldTick();
 		}
 	}
 
@@ -694,6 +694,7 @@ void UAbilitySystemComponent::GetLifetimeReplicatedProps(TArray< FLifetimeProper
 	DOREPLIFETIME(UAbilitySystemComponent, ActivatableAbilities);
 	DOREPLIFETIME(UAbilitySystemComponent, AbilityActor);
 	DOREPLIFETIME(UAbilitySystemComponent, ReplicatedPredictionKey);
+	DOREPLIFETIME(UAbilitySystemComponent, RepAnimMontageInfo);
 	
 	DOREPLIFETIME_CONDITION(UAbilitySystemComponent, SimulatedTasks, COND_SkipOwner);
 }
diff --git a/Engine/Source/Runtime/GameplayAbilities/Private/AbilitySystemComponent_Abilities.cpp b/Engine/Source/Runtime/GameplayAbilities/Private/AbilitySystemComponent_Abilities.cpp
--- a/Engine/Source/Runtime/GameplayAbilities/Private/AbilitySystemComponent_Abilities.cpp
+++ b/Engine/Source/Runtime/GameplayAbilities/Private/AbilitySystemComponent_Abilities.cpp
@@ -48,6 +48,11 @@ void UAbilitySystemComponent::TickComponent(float DeltaTime, enum ELevelTick Tic
 
 	Super::TickComponent(DeltaTime, TickType, ThisTickFunction);
 
+	if (IsOwnerActorAuthoritative())
+	{
+		AnimMontage_UpdateReplicatedData();
+	}
+
 	// Because we have no control over what a task may do when it ticks, we must be careful.
 	// Ticking a task may kill the task right here. It could also potentially kill another task
 	// which was waiting on the original task to do something. Since when a tasks is killed, it removes
@@ -85,9 +90,23 @@ void UAbilitySystemComponent::TickComponent(float DeltaTime, enum ELevelTick Tic
 	if (NumActuallyTicked == 0)
 	{
 		TickingTasks.SetNum(0, false);
-		SetActive(false);
+		UpdateShouldTick();
 	}
+}
 
+void UAbilitySystemComponent::UpdateShouldTick()
+{
+	bool HasTickingTasks = (TickingTasks.Num() == 0);
+	bool HasReplicatedMontageInfoToUpdate = (IsOwnerActorAuthoritative() && RepAnimMontageInfo.IsStopped == false);
+
+	if (HasTickingTasks || HasReplicatedMontageInfoToUpdate)
+	{
+		SetActive(true);
+	}
+	else
+	{
+		SetActive(false);
+	}
 }
 
 void UAbilitySystemComponent::InitAbilityActorInfo(AActor* AvatarActor)
@@ -233,9 +252,9 @@ void UAbilitySystemComponent::NotifyAbilityEnded(FGameplayAbilitySpecHandle Hand
 	check(Ability);
 
 	// If AnimatingAbility ended, clear the pointer
-	if (AnimatingAbility == Ability)
+	if (LocalAnimMontageInfo.AnimatingAbility == Ability)
 	{
-		AnimatingAbility = NULL;
+		LocalAnimMontageInfo.AnimatingAbility = NULL;
 	}
 
 	Spec->ActiveCount--;
@@ -695,22 +714,6 @@ void UAbilitySystemComponent::TriggerAbilityFromGameplayEvent(FGameplayAbilitySp
 //								Input 
 // ----------------------------------------------------------------------------------------------------------------------------------------------------
 
-void UAbilitySystemComponent::MontageBranchPoint_AbilityDecisionStop()
-{
-	if (AnimatingAbility)
-	{
-		AnimatingAbility->MontageBranchPoint_AbilityDecisionStop(AbilityActorInfo.Get());
-	}
-}
-
-void UAbilitySystemComponent::MontageBranchPoint_AbilityDecisionStart()
-{
-	if (AnimatingAbility)
-	{
-		AnimatingAbility->MontageBranchPoint_AbilityDecisionStart(AbilityActorInfo.Get());
-	}
-}
-
 bool UAbilitySystemComponent::GetUserAbilityActivationInhibited() const
 {
 	return UserAbilityActivationInhibited;
@@ -1037,7 +1040,7 @@ void UAbilitySystemComponent::OnRep_SimulatedTasks()
 			SimulatedTask->InitSimulatedTask(this);
 			if (TickingTasks.Num() == 0)
 			{
-				SetActive(true);
+				UpdateShouldTick();
 			}
 
 			TickingTasks.Add(SimulatedTask);
@@ -1045,6 +1048,276 @@ void UAbilitySystemComponent::OnRep_SimulatedTasks()
 	}
 }
 
+// ---------------------------------------------------------------------------
+
+float UAbilitySystemComponent::PlayMontage(UGameplayAbility* InAnimatingAbility, FGameplayAbilityActivationInfo ActivationInfo, UAnimMontage* NewAnimMontage, float InPlayRate, FName StartSectionName)
+{
+	UAnimInstance* AnimInstance = AbilityActorInfo->AnimInstance.Get();
+
+	if (AnimInstance && NewAnimMontage)
+	{
+		float const Duration = AnimInstance->Montage_Play(NewAnimMontage, InPlayRate);
+		if (Duration > 0.f)
+		{
+			if (LocalAnimMontageInfo.AnimatingAbility && LocalAnimMontageInfo.AnimatingAbility != InAnimatingAbility)
+			{
+				// The ability that was previously animating will have already gotten the 'interrupted' callback.
+				// It may be a good idea to make this a global policy and 'cancel' the ability.
+				// 
+				// For now, we expect it to end itself when this happens.
+			}
+
+			LocalAnimMontageInfo.AnimMontage = NewAnimMontage;
+			LocalAnimMontageInfo.AnimatingAbility = InAnimatingAbility;
+			
+			// Replicate to non owners
+			if (IsOwnerActorAuthoritative())
+			{
+				// Those are static parameters, they are only set when the montage is played. They are not changed after that.
+				RepAnimMontageInfo.AnimMontage = NewAnimMontage;
+				RepAnimMontageInfo.ForcePlayBit = !bool(RepAnimMontageInfo.ForcePlayBit);
+
+				// Start at a given Section.
+				if (StartSectionName != NAME_None)
+				{
+					AnimInstance->Montage_JumpToSection(StartSectionName);
+				}
+
+				// Update parameters that change during Montage life time.
+				AnimMontage_UpdateReplicatedData();
+			}
+			else
+			{
+				// If this prediction key is rejected, we need to end the preview
+				FPredictionKey PredictionKey = ActivationInfo.GetPredictionKeyForNewAction();
+				if (PredictionKey.IsValidKey())
+				{
+					PredictionKey.NewRejectedDelegate().BindUObject(this, &UAbilitySystemComponent::OnPredictiveMontageRejected, NewAnimMontage);
+				}
+			}
+
+			return Duration;
+		}
+	}
+
+	return -1.f;
+}
+
+float UAbilitySystemComponent::PlayMontageSimulated(UAnimMontage* NewAnimMontage, float InPlayRate, FName StartSectionName)
+{
+	UAnimInstance* AnimInstance = AbilityActorInfo->AnimInstance.Get();
+
+	if (AnimInstance && NewAnimMontage)
+	{
+		float const Duration = AnimInstance->Montage_Play(NewAnimMontage, InPlayRate);
+		if (Duration > 0.f)
+		{
+			LocalAnimMontageInfo.AnimMontage = NewAnimMontage;
+			return Duration;
+		}
+	}
+
+	return -1.f;
+}
+
+void UAbilitySystemComponent::AnimMontage_UpdateReplicatedData()
+{
+	check(IsOwnerActorAuthoritative());
+
+	UAnimInstance* AnimInstance = AbilityActorInfo->AnimInstance.Get();
+	if (AnimInstance && LocalAnimMontageInfo.AnimMontage)
+	{
+		RepAnimMontageInfo.AnimMontage = LocalAnimMontageInfo.AnimMontage;
+		RepAnimMontageInfo.PlayRate = AnimInstance->Montage_GetPlayRate(LocalAnimMontageInfo.AnimMontage);
+		RepAnimMontageInfo.Position = AnimInstance->Montage_GetPosition(LocalAnimMontageInfo.AnimMontage);
+
+		// Compressed Flags
+		bool bIsStopped = AnimInstance->Montage_GetIsStopped(LocalAnimMontageInfo.AnimMontage);
+
+		if (RepAnimMontageInfo.IsStopped != bIsStopped)
+		{
+			// When this changes, we should update whether or not we should be ticking
+			UpdateShouldTick();
+		}
+			
+		RepAnimMontageInfo.IsStopped = bIsStopped;
+
+		// Replicate NextSectionID to keep it in sync.
+		// We actually replicate NextSectionID+1 on a BYTE to put INDEX_NONE in there.
+		int32 CurrentSectionID = LocalAnimMontageInfo.AnimMontage->GetSectionIndexFromPosition(RepAnimMontageInfo.Position);
+		if (CurrentSectionID != INDEX_NONE)
+		{
+			int32 NextSectionID = AnimInstance->Montage_GetNextSectionID(LocalAnimMontageInfo.AnimMontage, CurrentSectionID);
+			ensure(NextSectionID < (256 - 1));
+			RepAnimMontageInfo.NextSectionID = uint8(NextSectionID + 1);
+		}
+		else
+		{
+			RepAnimMontageInfo.NextSectionID = 0;
+		}
+	}
+}
+
+void UAbilitySystemComponent::OnPredictiveMontageRejected(UAnimMontage* PredictiveMontage)
+{
+	static const float MONTAGE_PREDICTION_REJECT_FADETIME = 0.25f;
+
+	UAnimInstance* AnimInstance = AbilityActorInfo->AnimInstance.Get();
+	if (AnimInstance && PredictiveMontage)
+	{
+		// If this montage is still palying: kill it
+		if (AnimInstance->Montage_IsPlaying(PredictiveMontage))
+		{
+			AnimInstance->Montage_Stop(MONTAGE_PREDICTION_REJECT_FADETIME, PredictiveMontage);
+		}
+	}
+}
+
+/**	Replicated Event for AnimMontages */
+void UAbilitySystemComponent::OnRep_ReplicatedAnimMontage()
+{
+	static const float MONTAGE_REP_POS_ERR_THRESH = 0.1f;
+
+	UAnimInstance* AnimInstance = AbilityActorInfo->AnimInstance.Get();
+	if (!AbilityActorInfo->IsLocallyControlled() && AnimInstance)
+	{
+		static const auto CVar = IConsoleManager::Get().FindTConsoleVariableDataInt(TEXT("net.Montage.Debug"));
+		bool DebugMontage = (CVar && CVar->GetValueOnGameThread() == 1);
+		if (DebugMontage)
+		{
+			ABILITY_LOG( Warning, TEXT("\n\nOnRep_ReplicatedAnimMontage, %s"), *GetNameSafe(this));
+			ABILITY_LOG( Warning, TEXT("\tAnimMontage: %s\n\tPlayRate: %f\n\tPosition: %f\n\tNextSectionID: %d\n\tIsStopped: %d\n\tForcePlayBit: %d"),
+				*GetNameSafe(RepAnimMontageInfo.AnimMontage), 
+				RepAnimMontageInfo.PlayRate, 
+				RepAnimMontageInfo.Position, 
+				RepAnimMontageInfo.NextSectionID, 
+				RepAnimMontageInfo.IsStopped, 
+				RepAnimMontageInfo.ForcePlayBit);
+			ABILITY_LOG( Warning, TEXT("\tLocalAnimMontageInfo.AnimMontage: %s\n\tPosition: %f"),
+				*GetNameSafe(LocalAnimMontageInfo.AnimMontage), AnimInstance->Montage_GetPosition(LocalAnimMontageInfo.AnimMontage));
+		}
+
+		if( RepAnimMontageInfo.AnimMontage )
+		{
+			// New Montage to play
+			bool ReplicatedPlayBit = bool(RepAnimMontageInfo.ForcePlayBit);
+			if ((LocalAnimMontageInfo.AnimMontage != RepAnimMontageInfo.AnimMontage) || (LocalAnimMontageInfo.PlayBit != ReplicatedPlayBit))
+			{
+				LocalAnimMontageInfo.PlayBit = ReplicatedPlayBit;
+				PlayMontageSimulated(RepAnimMontageInfo.AnimMontage, RepAnimMontageInfo.PlayRate);
+			}
+
+			// Play Rate has changed
+			if (AnimInstance->Montage_GetPlayRate(LocalAnimMontageInfo.AnimMontage) != RepAnimMontageInfo.PlayRate)
+			{
+				AnimInstance->Montage_SetPlayRate(LocalAnimMontageInfo.AnimMontage, RepAnimMontageInfo.PlayRate);
+			}
+
+			int32 RepSectionID = LocalAnimMontageInfo.AnimMontage->GetSectionIndexFromPosition(RepAnimMontageInfo.Position);
+			int32 RepNextSectionID = int32(RepAnimMontageInfo.NextSectionID) - 1;
+		
+			// And NextSectionID for the replicated SectionID.
+			if( RepSectionID != INDEX_NONE )
+			{
+				int32 NextSectionID = AnimInstance->Montage_GetNextSectionID(LocalAnimMontageInfo.AnimMontage, RepSectionID);
+
+				// If NextSectionID is different thant the replicated one, then set it.
+				if( NextSectionID != RepNextSectionID )
+				{
+					AnimInstance->Montage_SetNextSection(LocalAnimMontageInfo.AnimMontage->GetSectionName(RepSectionID), LocalAnimMontageInfo.AnimMontage->GetSectionName(RepNextSectionID));
+				}
+
+				// Make sure we haven't received that update too late and the client hasn't already jumped to another section. 
+				int32 CurrentSectionID = LocalAnimMontageInfo.AnimMontage->GetSectionIndexFromPosition(AnimInstance->Montage_GetPosition(LocalAnimMontageInfo.AnimMontage));
+				if( (CurrentSectionID != RepSectionID) && (CurrentSectionID != RepNextSectionID) )
+				{
+					// Client is in a wrong section, jump to replicated position.
+					AnimInstance->Montage_SetPosition(LocalAnimMontageInfo.AnimMontage, RepAnimMontageInfo.Position);
+				}
+			}
+
+			// Update Position. If error is too great, jump to replicated position.
+			float CurrentPosition = AnimInstance->Montage_GetPosition(LocalAnimMontageInfo.AnimMontage);
+			int32 CurrentSectionID = LocalAnimMontageInfo.AnimMontage->GetSectionIndexFromPosition(CurrentPosition);
+			// Only check threshold if we are located in the same section. Different sections require a bit more work as we could be jumping around the timeline.
+			if( (CurrentSectionID == RepSectionID) && (FMath::Abs(CurrentPosition - RepAnimMontageInfo.Position) > MONTAGE_REP_POS_ERR_THRESH) )
+			{
+				AnimInstance->Montage_SetPosition(LocalAnimMontageInfo.AnimMontage, RepAnimMontageInfo.Position);
+			}
+
+			// Compressed Flags
+			bool bIsStopped = AnimInstance->Montage_GetIsStopped(LocalAnimMontageInfo.AnimMontage);
+			bool ReplicatedIsStopped = bool(RepAnimMontageInfo.IsStopped);
+			if( ReplicatedIsStopped && !bIsStopped )
+			{
+				CurrentMontageStop();
+			}
+		}
+	}
+}
+
+void UAbilitySystemComponent::CurrentMontageStop()
+{
+	UAnimInstance* AnimInstance = AbilityActorInfo->AnimInstance.Get();
+	UAnimMontage* MontageToStop = LocalAnimMontageInfo.AnimMontage;
+	bool bShouldStopMontage = AnimInstance && MontageToStop && !AnimInstance->Montage_GetIsStopped(MontageToStop);
+
+	if (bShouldStopMontage)
+	{
+		AnimInstance->Montage_Stop(MontageToStop->BlendOutTime);
+
+		if (IsOwnerActorAuthoritative())
+		{
+			AnimMontage_UpdateReplicatedData();
+		}
+	}
+}
+
+void UAbilitySystemComponent::CurrentMontageJumpToSection(FName SectionName)
+{
+	UAnimInstance* AnimInstance = AbilityActorInfo->AnimInstance.Get();
+	if( (SectionName != NAME_None) && AnimInstance )
+	{
+		AnimInstance->Montage_JumpToSection(SectionName);
+		if (IsOwnerActorAuthoritative())
+		{
+			AnimMontage_UpdateReplicatedData();
+		}
+	}
+}
+
+/** Set Next Section Name, and update replication for Simulated Proxies if we are on the server. */
+void UAbilitySystemComponent::CurrentMontageSetNextSectionName(FName FromSectionName, FName ToSectionName)
+{
+	UAnimInstance* AnimInstance = AbilityActorInfo->AnimInstance.Get();
+	if( LocalAnimMontageInfo.AnimMontage && AnimInstance )
+	{
+		// Set Next Section Name. 
+		AnimInstance->Montage_SetNextSection(FromSectionName, ToSectionName);
+
+		// Update replicated version for Simulated Proxies if we are on the server.
+		if( IsOwnerActorAuthoritative() )
+		{
+			AnimMontage_UpdateReplicatedData();
+		}
+	}
+}
+
+UAnimMontage* UAbilitySystemComponent::GetCurrentMontage() const
+{
+	UAnimInstance* AnimInstance = AbilityActorInfo->AnimInstance.Get();
+	if (AnimInstance)
+	{
+		return AnimInstance->GetCurrentActiveMontage();
+	}
+
+	return nullptr;
+}
+
+bool UAbilitySystemComponent::IsAnimatingAbility(UGameplayAbility* InAbility) const
+{
+	return (LocalAnimMontageInfo.AnimatingAbility == InAbility);
+}
 
 #undef LOCTEXT_NAMESPACE
 
diff --git a/Engine/Source/Runtime/GameplayAbilities/Public/Abilities/GameplayAbility.h b/Engine/Source/Runtime/GameplayAbilities/Public/Abilities/GameplayAbility.h
--- a/Engine/Source/Runtime/GameplayAbilities/Public/Abilities/GameplayAbility.h
+++ b/Engine/Source/Runtime/GameplayAbilities/Public/Abilities/GameplayAbility.h
@@ -464,13 +464,15 @@ protected:
 
 	// ----------------------------------------------------------------------------------------------------------------
 	//
-	//	Animation callbacks (still WIP)
+	//	Animation
 	//
 	// ----------------------------------------------------------------------------------------------------------------
 
-	void MontageBranchPoint_AbilityDecisionStop(const FGameplayAbilityActorInfo* ActorInfo) const;
+	UFUNCTION(BlueprintCallable, Category="Ability|Animation")
+	void MontageJumpToSection(FName SectionName);
 
-	void MontageBranchPoint_AbilityDecisionStart(const FGameplayAbilityActorInfo* ActorInfo) const;
+	UFUNCTION(BlueprintCallable, Category="Ability|Animation")
+	void MontageStop();
 
 	UFUNCTION(BlueprintPure, Category = Ability, meta = (HidePin = "WorldContextObject", DefaultToSelf = "WorldContextObject"))
 	FGameplayAbilityTargetingLocationInfo MakeTargetLocationInfoFromOwnerActor();
diff --git a/Engine/Source/Runtime/GameplayAbilities/Public/Abilities/GameplayAbilityTypes.h b/Engine/Source/Runtime/GameplayAbilities/Public/Abilities/GameplayAbilityTypes.h
--- a/Engine/Source/Runtime/GameplayAbilities/Public/Abilities/GameplayAbilityTypes.h
+++ b/Engine/Source/Runtime/GameplayAbilities/Public/Abilities/GameplayAbilityTypes.h
@@ -295,6 +295,78 @@ struct GAMEPLAYABILITIES_API FGameplayAbilitySpec
 
 // ----------------------------------------------------
 
+
+/** Data about montages that is replicated to simulated clients */
+USTRUCT()
+struct GAMEPLAYABILITIES_API FGameplayAbilityRepAnimMontage
+{
+	GENERATED_USTRUCT_BODY()
+
+	/** AnimMontage ref */
+	UPROPERTY()
+	UAnimMontage* AnimMontage;
+
+	/** Play Rate */
+	UPROPERTY()
+	float PlayRate;
+
+	/** Montage position */
+	UPROPERTY()
+	float Position;
+
+	/** NextSectionID */
+	UPROPERTY()
+	uint8 NextSectionID;
+
+	/** Bit set when montage has been stopped. */
+	UPROPERTY()
+	uint8 IsStopped : 1;
+	
+	/** Bit flipped every time a new Montage is played. To trigger replication when the same montage is played again. */
+	UPROPERTY()
+	uint8 ForcePlayBit : 1;
+
+	FGameplayAbilityRepAnimMontage()
+		: AnimMontage(NULL),
+		PlayRate(0.f),
+		Position(0.f),
+		NextSectionID(0),
+		IsStopped(0),
+		ForcePlayBit(0)
+	{
+	}
+
+	UPROPERTY()
+	FPredictionKey PredictionKey;
+};
+
+/** Data about montages that were played locally (all montages in case of server. predictive montages in case of client). Never replicated directly. */
+USTRUCT()
+struct GAMEPLAYABILITIES_API FGameplayAbilityLocalAnimMontage
+{
+	GENERATED_USTRUCT_BODY()
+
+	FGameplayAbilityLocalAnimMontage()
+	: AnimMontage(nullptr), PlayBit(false), AnimatingAbility(nullptr)
+	{
+	}
+
+	UPROPERTY()
+	UAnimMontage* AnimMontage;
+
+	UPROPERTY()
+	bool PlayBit;
+
+	UPROPERTY()
+	FPredictionKey PredictionKey;
+
+	/** The ability, if any, that instigated this montage */
+	UPROPERTY()
+	UGameplayAbility* AnimatingAbility;
+};
+
+// ----------------------------------------------------
+
 USTRUCT(BlueprintType)
 struct GAMEPLAYABILITIES_API FGameplayEventData
 {
@@ -304,7 +376,6 @@ struct GAMEPLAYABILITIES_API FGameplayEventData
 	: Instigator(NULL)
 	, Target(NULL)
 	{
-		// do nothing
 	}
 	
 	UPROPERTY(EditAnywhere, BlueprintReadWrite, Category = GameplayAbilityTriggerPayload)
diff --git a/Engine/Source/Runtime/GameplayAbilities/Public/AbilitySystemComponent.h b/Engine/Source/Runtime/GameplayAbilities/Public/AbilitySystemComponent.h
--- a/Engine/Source/Runtime/GameplayAbilities/Public/AbilitySystemComponent.h
+++ b/Engine/Source/Runtime/GameplayAbilities/Public/AbilitySystemComponent.h
@@ -57,6 +57,9 @@ class GAMEPLAYABILITIES_API UAbilitySystemComponent : public UActorComponent, pu
 	virtual void InitializeComponent() override;
 	virtual void TickComponent(float DeltaTime, enum ELevelTick TickType, FActorComponentTickFunction *ThisTickFunction) override;
 
+	/** Checks to see if we should be active (ticking). Called after something changes that would cause us to tick or not tick. */
+	void UpdateShouldTick();
+
 	/** Finds existing AttributeSet */
 	template <class T >
 	const T*	GetSet()
@@ -449,14 +452,51 @@ class GAMEPLAYABILITIES_API UAbilitySystemComponent : public UActorComponent, pu
 	/** Removes a UAbilityTask task from the list of tasks to be ticked */
 	void TaskEnded(UAbilityTask* Task);
 
+	// ----------------------------------------------------------------------------------------------------------------
+	//	AnimMontage Support
+	//	
+	//	TODO:
+	//	-Continously update RepAnimMontageInfo on server for join in progress clients.
+	//	-Some missing functionality may still be needed (GetCurrentSectionTime, SetPlayRate, etc)	
+	// ----------------------------------------------------------------------------------------------------------------	
 
-	/** There needs to be a concept of an animating ability. Only one may exist at a time. New requests can be queued up, overridden, or ignored. */
-	UPROPERTY()
-	UGameplayAbility*	AnimatingAbility;
+	/** Plays a montage and handles replication and prediction based on passed in ability/activation info */
+	float PlayMontage(UGameplayAbility* AnimatingAbility, FGameplayAbilityActivationInfo ActivationInfo, UAnimMontage* Montage, float InPlayRate, FName StartSectionName = NAME_None);
+
+	/** Plays a montage without updating replication/prediction structures. Used by simulated proxies when replication tells them to play a montage. */
+	float PlayMontageSimulated(UAnimMontage* Montage, float InPlayRate, FName StartSectionName = NAME_None);
+
+	/** Stops whatever montage is currently playing. Expectation is caller should only be stoping it if they are the current animating ability (or have good reason not to check) */
+	void CurrentMontageStop();
+
+	/** Jumps current montage to given section. Expectation is caller should only be stoping it if they are the current animating ability (or have good reason not to check) */
+	void CurrentMontageJumpToSection(FName SectionName);
+
+	/** Sets current montages next section name. Expectation is caller should only be stoping it if they are the current animating ability (or have good reason not to check) */
+	void CurrentMontageSetNextSectionName(FName FromSectionName, FName ToSectionName);
 
-	void MontageBranchPoint_AbilityDecisionStop();
+	/** Returns true if the passed in ability is the current animating ability */
+	bool IsAnimatingAbility(UGameplayAbility* Ability) const;
 
-	void MontageBranchPoint_AbilityDecisionStart();
+	/** Returns montage that is currently playing */
+	UAnimMontage* GetCurrentMontage() const;
+
+	/** Called when a prediction key that played a montage is rejected */
+	void OnPredictiveMontageRejected(UAnimMontage* PredictiveMontage);
+
+	/** Copy LocalAnimMontageInfo into RepAnimMontageInfo */
+	void AnimMontage_UpdateReplicatedData();
+
+	/** Data structure for replicating montage info to simulated clients */
+	UPROPERTY(ReplicatedUsing=OnRep_ReplicatedAnimMontage)
+	FGameplayAbilityRepAnimMontage RepAnimMontageInfo;
+
+	/** Data structure for montages that were instigated locally (everything if server, predictive if client. replicated if simulated proxy) */
+	UPROPERTY()
+	FGameplayAbilityLocalAnimMontage LocalAnimMontageInfo;
+
+	UFUNCTION()
+	void OnRep_ReplicatedAnimMontage();
 
 	// -----------------------------------------------------------------------------
 
