build fixes

diff --git a/Engine/Source/Runtime/GameplayAbilities/Private/Abilities/GameplayAbilityTargetActor_SingleLineTrace.cpp b/Engine/Source/Runtime/GameplayAbilities/Private/Abilities/GameplayAbilityTargetActor_SingleLineTrace.cpp
--- a/Engine/Source/Runtime/GameplayAbilities/Private/Abilities/GameplayAbilityTargetActor_SingleLineTrace.cpp
+++ b/Engine/Source/Runtime/GameplayAbilities/Private/Abilities/GameplayAbilityTargetActor_SingleLineTrace.cpp
@@ -5,7 +5,7 @@
 #include "GameplayAbilityTargetActor_SingleLineTrace.h"
 #include "Engine/World.h"
 #include "Runtime/Engine/Public/Net/UnrealNetwork.h"
-#pragma optimize("",off)
+
 // --------------------------------------------------------------------------------------------------------------------------------------------------------
 //
 //	AGameplayAbilityTargetActor_SingleLineTrace
diff --git a/Engine/Source/Runtime/GameplayAbilities/Private/Abilities/Tasks/AbilityTask_PlayMontageAndWait.cpp b/Engine/Source/Runtime/GameplayAbilities/Private/Abilities/Tasks/AbilityTask_PlayMontageAndWait.cpp
--- a/Engine/Source/Runtime/GameplayAbilities/Private/Abilities/Tasks/AbilityTask_PlayMontageAndWait.cpp
+++ b/Engine/Source/Runtime/GameplayAbilities/Private/Abilities/Tasks/AbilityTask_PlayMontageAndWait.cpp
@@ -1,6 +1,7 @@
 
 #include "AbilitySystemPrivatePCH.h"
 #include "Abilities/Tasks/AbilityTask_PlayMontageAndWait.h"
+#include "AbilitySystemComponent.h"
 
 UAbilityTask_PlayMontageAndWait::UAbilityTask_PlayMontageAndWait(const class FPostConstructInitializeProperties& PCIP)
 	: Super(PCIP)
