Build fix

diff --git a/Engine/Source/Runtime/GameplayAbilities/Private/Abilities/GameplayAbility.cpp b/Engine/Source/Runtime/GameplayAbilities/Private/Abilities/GameplayAbility.cpp
--- a/Engine/Source/Runtime/GameplayAbilities/Private/Abilities/GameplayAbility.cpp
+++ b/Engine/Source/Runtime/GameplayAbilities/Private/Abilities/GameplayAbility.cpp
@@ -485,8 +485,6 @@ void UGameplayAbility::MontageJumpToSection(FName SectionName)
 {
 	check(CurrentActorInfo);
 
-	CurrentActorInfo->AbilitySystemComponent->LocalAnimMontageInfo.AnimMontage;
-
 	if (CurrentActorInfo->AbilitySystemComponent->IsAnimatingAbility(this))
 	{
 		CurrentActorInfo->AbilitySystemComponent->CurrentMontageJumpToSection(SectionName);
