force non replication for reticle actor in cases where we are playing on a listen server and shouldnt normally be running the targeting code on the server.

diff --git a/Engine/Source/Runtime/GameplayAbilities/Private/Abilities/GameplayAbilityTargetActor_Trace.cpp b/Engine/Source/Runtime/GameplayAbilities/Private/Abilities/GameplayAbilityTargetActor_Trace.cpp
--- a/Engine/Source/Runtime/GameplayAbilities/Private/Abilities/GameplayAbilityTargetActor_Trace.cpp
+++ b/Engine/Source/Runtime/GameplayAbilities/Private/Abilities/GameplayAbilityTargetActor_Trace.cpp
@@ -69,6 +69,16 @@ void AGameplayAbilityTargetActor_Trace::StartTargeting(UGameplayAbility* InAbili
 		{
 			SpawnedReticleActor->InitializeReticle(this);
 			ReticleActor = SpawnedReticleActor;
+
+			// This is to catch cases of playing on a listen server where we are using a replicated reticle actor.
+			// (In a client controlled player, this would only run on the client and therefor never replicate. If it runs
+			// on a listen server, the reticle actor may replicate. We want consistancy between client/listen server players.
+			// Just saying 'make the reticle actor non replicated' isnt a good answer, since we want to mix and match reticle
+			// actors and there may be other targeting types that want to replicate the same reticle actor class).
+			if (!ShouldProduceTargetDataOnServer)
+			{
+				SpawnedReticleActor->SetReplicates(false);
+			}
 		}
 	}
 }
