Added "self" feature to multi select targeting, added self-or-other filtering to engine-level filters, allowed engine-level filters to see the current ability, fixed potential bug in multi projectile spawn, tightened up some trace targeting actor code, and added the recently-improved center-screen aiming code to ground and actor-placement targeting.

diff --git a/Engine/Source/Runtime/GameplayAbilities/Private/Abilities/GameplayAbilityTargetActor_GroundTrace.cpp b/Engine/Source/Runtime/GameplayAbilities/Private/Abilities/GameplayAbilityTargetActor_GroundTrace.cpp
--- a/Engine/Source/Runtime/GameplayAbilities/Private/Abilities/GameplayAbilityTargetActor_GroundTrace.cpp
+++ b/Engine/Source/Runtime/GameplayAbilities/Private/Abilities/GameplayAbilityTargetActor_GroundTrace.cpp
@@ -133,27 +133,10 @@ FHitResult AGameplayAbilityTargetActor_GroundTrace::PerformTrace(AActor* InSourc
 	Params.bTraceAsyncScene = true;
 	Params.AddIgnoredActor(InSourceActor);
 
-	FVector AimDirection = InSourceActor->GetActorForwardVector();		//Default
-	
-	if (OwningAbility)		//Server and launching client only
-	{
-		APlayerController* AimingPC = OwningAbility->GetCurrentActorInfo()->PlayerController.Get();
-		check(AimingPC);
-		FVector CamLoc;
-		FRotator CamRot;
-		AimingPC->GetPlayerViewPoint(CamLoc, CamRot);
-		AimDirection = CamRot.Vector();
-	}
-
-	FVector TraceStart = InSourceActor->GetActorLocation();
-	FVector TraceEnd = TraceStart + (AimDirection * MaxRange);
-
-	//If we're using a socket, adjust the starting location and aim direction after the end position has been found. This way we can still aim with the camera, then fire accurately from the socket.
-	if (OwningAbility)		//Server and launching client only
-	{
-		TraceStart = StartLocation.GetTargetingTransform().GetLocation();
-	}
-	AimDirection = (TraceEnd - TraceStart).SafeNormal();
+	FVector TraceStart = StartLocation.GetTargetingTransform().GetLocation();// InSourceActor->GetActorLocation();
+	FVector TraceEnd = TraceStart + (InSourceActor->GetActorForwardVector() * MaxRange);		//Default
+	AimWithPlayerController(InSourceActor, Params, TraceStart, TraceEnd);		//Effective on server and launching client only
+	FVector AimDirection = (TraceStart - TraceEnd).SafeNormal();
 
 	// ------------------------------------------------------
 
diff --git a/Engine/Source/Runtime/GameplayAbilities/Private/Abilities/GameplayAbilityTargetActor_SingleLineTrace.cpp b/Engine/Source/Runtime/GameplayAbilities/Private/Abilities/GameplayAbilityTargetActor_SingleLineTrace.cpp
--- a/Engine/Source/Runtime/GameplayAbilities/Private/Abilities/GameplayAbilityTargetActor_SingleLineTrace.cpp
+++ b/Engine/Source/Runtime/GameplayAbilities/Private/Abilities/GameplayAbilityTargetActor_SingleLineTrace.cpp
@@ -30,35 +30,9 @@ FHitResult AGameplayAbilityTargetActor_SingleLineTrace::PerformTrace(AActor* InS
 	Params.bTraceAsyncScene = true;
 	Params.AddIgnoredActors(ActorsToIgnore);
 
-	FVector TraceStart = InSourceActor->GetActorLocation();
+	FVector TraceStart = StartLocation.GetTargetingTransform().GetLocation();// InSourceActor->GetActorLocation();
 	FVector TraceEnd = TraceStart + (InSourceActor->GetActorForwardVector() * MaxRange);		//Default
-
-	if (OwningAbility)		//Server and launching client only
-	{
-		APlayerController* AimingPC = OwningAbility->GetCurrentActorInfo()->PlayerController.Get();
-		check(AimingPC);
-		FVector CamLoc;
-		FRotator CamRot;
-		AimingPC->GetPlayerViewPoint(CamLoc, CamRot);
-		FVector CamDir = CamRot.Vector();
-		FVector CamTarget = CamLoc + (CamDir * MaxRange);		//Straight, dumb aiming to a point that's reasonable though not exactly correct
-
-		ClipCameraRayToAbilityRange(CamLoc, CamDir, TraceStart, MaxRange, CamTarget);
-
-		FHitResult TempHitResult;
-		InSourceActor->GetWorld()->LineTraceSingle(TempHitResult, CamLoc, CamTarget, ECC_WorldStatic, Params);
-		TraceStart = StartLocation.GetTargetingTransform().GetLocation();
-		if (TempHitResult.bBlockingHit && (FVector::DistSquared(TraceStart, TempHitResult.Location) <= (MaxRange * MaxRange)))
-		{
-			//We actually made a hit? Pull back.
-			TraceEnd = TempHitResult.Location;
-		}
-		else
-		{
-			//If we didn't make a hit, use the clipped location.
-			TraceEnd = CamTarget;
-		}
-	}
+	AimWithPlayerController(InSourceActor, Params, TraceStart, TraceEnd);		//Effective on server and launching client only
 	FVector AimDirection = (TraceStart - TraceEnd).SafeNormal();
 
 	// ------------------------------------------------------
diff --git a/Engine/Source/Runtime/GameplayAbilities/Private/Abilities/GameplayAbilityTargetActor_Trace.cpp b/Engine/Source/Runtime/GameplayAbilities/Private/Abilities/GameplayAbilityTargetActor_Trace.cpp
--- a/Engine/Source/Runtime/GameplayAbilities/Private/Abilities/GameplayAbilityTargetActor_Trace.cpp
+++ b/Engine/Source/Runtime/GameplayAbilities/Private/Abilities/GameplayAbilityTargetActor_Trace.cpp
@@ -38,6 +38,36 @@ FGameplayAbilityTargetDataHandle AGameplayAbilityTargetActor_Trace::StaticGetTar
 	return MakeTargetData(PerformTrace(StaticSourceActor));
 }
 
+void AGameplayAbilityTargetActor_Trace::AimWithPlayerController(AActor* InSourceActor, FCollisionQueryParams Params, FVector TraceStart, FVector& TraceEnd) const
+{
+	if (OwningAbility)		//Server and launching client only
+	{
+		APlayerController* AimingPC = OwningAbility->GetCurrentActorInfo()->PlayerController.Get();
+		check(AimingPC);
+		FVector CamLoc;
+		FRotator CamRot;
+		AimingPC->GetPlayerViewPoint(CamLoc, CamRot);
+		FVector CamDir = CamRot.Vector();
+		FVector CamTarget = CamLoc + (CamDir * MaxRange);		//Straight, dumb aiming to a point that's reasonable though not exactly correct
+
+		ClipCameraRayToAbilityRange(CamLoc, CamDir, TraceStart, MaxRange, CamTarget);
+
+		FHitResult TempHitResult;
+		InSourceActor->GetWorld()->LineTraceSingle(TempHitResult, CamLoc, CamTarget, ECC_WorldStatic, Params);
+		//TraceStart = ;
+		if (TempHitResult.bBlockingHit && (FVector::DistSquared(TraceStart, TempHitResult.Location) <= (MaxRange * MaxRange)))
+		{
+			//We actually made a hit? Pull back.
+			TraceEnd = TempHitResult.Location;
+		}
+		else
+		{
+			//If we didn't make a hit, use the clipped location.
+			TraceEnd = CamTarget;
+		}
+	}
+}
+
 bool AGameplayAbilityTargetActor_Trace::ClipCameraRayToAbilityRange(FVector CameraLocation, FVector CameraDirection, FVector AbilityCenter, float AbilityRange, FVector& ClippedPosition)
 {
 	FVector CameraToCenter = AbilityCenter - CameraLocation;
diff --git a/Engine/Source/Runtime/GameplayAbilities/Private/Abilities/GameplayAbilityTargetDataFilter.cpp b/Engine/Source/Runtime/GameplayAbilities/Private/Abilities/GameplayAbilityTargetDataFilter.cpp
--- a/Engine/Source/Runtime/GameplayAbilities/Private/Abilities/GameplayAbilityTargetDataFilter.cpp
+++ b/Engine/Source/Runtime/GameplayAbilities/Private/Abilities/GameplayAbilityTargetDataFilter.cpp
@@ -5,13 +5,6 @@
 
 // --------------------------------------------------------------------------------------------------------------------------------------------------------
 //
-//	AGameplayAbilityTargetDataActorFilter
+//	FGameplayTargetDataFilter
 //
 // --------------------------------------------------------------------------------------------------------------------------------------------------------
-
-/*
-AGameplayAbilityTargetDataActorFilter::AGameplayAbilityTargetDataActorFilter(const class FPostConstructInitializeProperties& PCIP)
-	: Super(PCIP)
-{
-}
-*/
\ No newline at end of file
diff --git a/Engine/Source/Runtime/GameplayAbilities/Private/AbilitySystemBlueprintLibrary.cpp b/Engine/Source/Runtime/GameplayAbilities/Private/AbilitySystemBlueprintLibrary.cpp
--- a/Engine/Source/Runtime/GameplayAbilities/Private/AbilitySystemBlueprintLibrary.cpp
+++ b/Engine/Source/Runtime/GameplayAbilities/Private/AbilitySystemBlueprintLibrary.cpp
@@ -124,10 +124,11 @@ FGameplayAbilityTargetDataHandle UAbilitySystemBlueprintLibrary::FilterTargetDat
 	return ReturnDataHandle;
 }
 
-FGameplayTargetDataFilterHandle UAbilitySystemBlueprintLibrary::MakeFilterHandle(FGameplayTargetDataFilter Filter)
+FGameplayTargetDataFilterHandle UAbilitySystemBlueprintLibrary::MakeFilterHandle(UGameplayAbility* WorldContextObject, FGameplayTargetDataFilter Filter)
 {
 	FGameplayTargetDataFilterHandle FilterHandle;
 	FGameplayTargetDataFilter* NewFilter = new FGameplayTargetDataFilter(Filter);
+	NewFilter->SourceAbility = WorldContextObject;
 	FilterHandle.Filter = TSharedPtr<FGameplayTargetDataFilter>(NewFilter);
 	return FilterHandle;
 }
diff --git a/Engine/Source/Runtime/GameplayAbilities/Public/Abilities/GameplayAbilityTargetActor.h b/Engine/Source/Runtime/GameplayAbilities/Public/Abilities/GameplayAbilityTargetActor.h
--- a/Engine/Source/Runtime/GameplayAbilities/Public/Abilities/GameplayAbilityTargetActor.h
+++ b/Engine/Source/Runtime/GameplayAbilities/Public/Abilities/GameplayAbilityTargetActor.h
@@ -4,6 +4,17 @@
 #include "GameplayAbility.h"
 #include "GameplayAbilityTargetActor.generated.h"
 
+UENUM(BlueprintType)
+namespace ETargetAbilitySelfSelection
+{
+	enum Type
+	{
+		TASS_Permit			UMETA(DisplayName = "Allow self-selection"),
+		TASS_Forbid			UMETA(DisplayName = "Forbid self-selection"),
+		TASS_Require 		UMETA(DisplayName = "Force self-selection (add to final data)")
+	};
+}
+
 /** TargetActors are spawned to assist with ability targeting. They are spawned by ability tasks and create/determine the outgoing targeting data passed from one task to another. */
 UCLASS(Blueprintable, abstract, notplaceable)
 class GAMEPLAYABILITIES_API AGameplayAbilityTargetActor : public AActor
diff --git a/Engine/Source/Runtime/GameplayAbilities/Public/Abilities/GameplayAbilityTargetActor_Trace.h b/Engine/Source/Runtime/GameplayAbilities/Public/Abilities/GameplayAbilityTargetActor_Trace.h
--- a/Engine/Source/Runtime/GameplayAbilities/Public/Abilities/GameplayAbilityTargetActor_Trace.h
+++ b/Engine/Source/Runtime/GameplayAbilities/Public/Abilities/GameplayAbilityTargetActor_Trace.h
@@ -17,6 +17,8 @@ public:
 
 	virtual FGameplayAbilityTargetDataHandle StaticGetTargetData(UWorld * World, const FGameplayAbilityActorInfo* ActorInfo, FGameplayAbilityActivationInfo ActivationInfo) const override;
 
+	void AimWithPlayerController(AActor* InSourceActor, FCollisionQueryParams Params, FVector TraceStart, FVector& TraceEnd) const;
+
 	static bool ClipCameraRayToAbilityRange(FVector CameraLocation, FVector CameraDirection, FVector AbilityCenter, float AbilityRange, FVector& ClippedPosition);
 
 	virtual void StartTargeting(UGameplayAbility* Ability) override;
diff --git a/Engine/Source/Runtime/GameplayAbilities/Public/Abilities/GameplayAbilityTargetDataFilter.h b/Engine/Source/Runtime/GameplayAbilities/Public/Abilities/GameplayAbilityTargetDataFilter.h
--- a/Engine/Source/Runtime/GameplayAbilities/Public/Abilities/GameplayAbilityTargetDataFilter.h
+++ b/Engine/Source/Runtime/GameplayAbilities/Public/Abilities/GameplayAbilityTargetDataFilter.h
@@ -4,6 +4,17 @@
 
 #include "GameplayAbilityTargetDataFilter.generated.h"
 
+UENUM(BlueprintType)
+namespace ETargetDataFilterSelf
+{
+	enum Type
+	{
+		TDFS_Any 			UMETA(DisplayName = "Allow self or others"),
+		TDFS_NoSelf 		UMETA(DisplayName = "Filter self out"),
+		TDFS_NoOthers		UMETA(DisplayName = "Filter others out")
+	};
+}
+
 USTRUCT(BlueprintType)
 struct FGameplayTargetDataFilter
 {
@@ -15,6 +26,26 @@ struct FGameplayTargetDataFilter
 
 	virtual bool FilterPassesForActor(const AActor* ActorToBeFiltered) const
 	{
+		check(SourceAbility.IsValid());
+		check(SourceAbility.Get()->GetOwningActorFromActorInfo());
+		switch (SelfFilter.GetValue())
+		{
+		case ETargetDataFilterSelf::Type::TDFS_NoOthers:
+			if (ActorToBeFiltered != SourceAbility.Get()->GetOwningActorFromActorInfo())
+			{
+				return false;
+			}
+			break;
+		case ETargetDataFilterSelf::Type::TDFS_NoSelf:
+			if (ActorToBeFiltered == SourceAbility.Get()->GetOwningActorFromActorInfo())
+			{
+				return false;
+			}
+			break;
+		case ETargetDataFilterSelf::Type::TDFS_Any:
+		default:
+			break;
+		}
 		return true;
 	}
 
@@ -27,6 +58,12 @@ struct FGameplayTargetDataFilter
 	{
 		return FilterPassesForActor(A);
 	}
+
+	/** Filled out while running */
+	TWeakObjectPtr<UGameplayAbility> SourceAbility;
+
+	UPROPERTY(BlueprintReadWrite, EditAnywhere, meta = (ExposeOnSpawn = true), Category = Filter)
+	TEnumAsByte<ETargetDataFilterSelf::Type> SelfFilter;
 };
 
 
diff --git a/Engine/Source/Runtime/GameplayAbilities/Public/Abilities/GameplayAbilityTargetTypes.h b/Engine/Source/Runtime/GameplayAbilities/Public/Abilities/GameplayAbilityTargetTypes.h
--- a/Engine/Source/Runtime/GameplayAbilities/Public/Abilities/GameplayAbilityTargetTypes.h
+++ b/Engine/Source/Runtime/GameplayAbilities/Public/Abilities/GameplayAbilityTargetTypes.h
@@ -171,6 +171,11 @@ struct GAMEPLAYABILITIES_API FGameplayAbilityTargetDataHandle
 		Data.Reset();
 	}
 
+	int32 Num() const
+	{
+		return Data.Num();
+	}
+
 	bool IsValid(int32 Index) const
 	{
 		return (Index < Data.Num() && Data[Index].IsValid());
diff --git a/Engine/Source/Runtime/GameplayAbilities/Public/AbilitySystemBlueprintLibrary.h b/Engine/Source/Runtime/GameplayAbilities/Public/AbilitySystemBlueprintLibrary.h
--- a/Engine/Source/Runtime/GameplayAbilities/Public/AbilitySystemBlueprintLibrary.h
+++ b/Engine/Source/Runtime/GameplayAbilities/Public/AbilitySystemBlueprintLibrary.h
@@ -50,8 +50,9 @@ class UAbilitySystemBlueprintLibrary : public UBlueprintFunctionLibrary
 	UFUNCTION(BlueprintPure, Category = "Ability|TargetData")
 	static FGameplayAbilityTargetDataHandle	FilterTargetData(FGameplayAbilityTargetDataHandle TargetDataHandle, FGameplayTargetDataFilterHandle ActorFilterClass);
 
-	UFUNCTION(BlueprintPure, Category = "Ability|TargetData")
-	static FGameplayTargetDataFilterHandle MakeFilterHandle(FGameplayTargetDataFilter Filter);
+	/** Create a handle for filtering target data */
+	UFUNCTION(BlueprintPure, meta = (HidePin = "WorldContextObject", DefaultToSelf = "WorldContextObject"), Category = "Filter")
+	static FGameplayTargetDataFilterHandle MakeFilterHandle(UGameplayAbility* WorldContextObject, FGameplayTargetDataFilter Filter);
 
 	UFUNCTION(BlueprintPure, Category = "Ability|TargetData")
 	static TArray<AActor*> GetActorsFromTargetData(FGameplayAbilityTargetDataHandle TargetData, int32 Index);
