GUBP fix.

diff --git a/Engine/Source/Runtime/GameplayAbilities/Public/Abilities/GameplayAbilityTargetDataFilter.h b/Engine/Source/Runtime/GameplayAbilities/Public/Abilities/GameplayAbilityTargetDataFilter.h
--- a/Engine/Source/Runtime/GameplayAbilities/Public/Abilities/GameplayAbilityTargetDataFilter.h
+++ b/Engine/Source/Runtime/GameplayAbilities/Public/Abilities/GameplayAbilityTargetDataFilter.h
@@ -2,6 +2,7 @@
 
 #pragma once
 
+#include "GameplayAbility.h"
 #include "GameplayAbilityTargetDataFilter.generated.h"
 
 UENUM(BlueprintType)
