Incoming/Outgoing modifiers must match attribute in all cases, not just in magnitude modifier case.

diff --git a/Engine/Source/Runtime/GameplayAbilities/Private/GameplayEffect.cpp b/Engine/Source/Runtime/GameplayAbilities/Private/GameplayEffect.cpp
--- a/Engine/Source/Runtime/GameplayAbilities/Private/GameplayEffect.cpp
+++ b/Engine/Source/Runtime/GameplayAbilities/Private/GameplayEffect.cpp
@@ -260,56 +260,45 @@ int32 FGameplayEffectSpec::ApplyModifier(const FModifierSpec &InMod, const FModi
 		return 0;
 	}
 
-	switch (InMod.Info.EffectType)
+	for (FModifierSpec &MyMod : Modifiers)
 	{
-		case EGameplayModEffect::Magnitude:
+		if (InMod.CanModifyModifier(MyMod, QualifierContext))
 		{
-			for (FModifierSpec &MyMod : Modifiers)
+			bool FinishedLoop = true;
+			NumApplied = 1;
+
+			switch (InMod.Info.EffectType)
 			{
-				if (InMod.CanModifyModifier(MyMod, QualifierContext))
-				{
+				case EGameplayModEffect::Magnitude:
 					InMod.ApplyModTo(MyMod, bApplyAsSnapshot);
-					NumApplied = 1;
-				}
-			}
-			break;
-		}
+					FinishedLoop = false;	// Keeep going: we need to go through the entire mod list
+					break;
 
-		// Fixme: Duration mods aren't checking that attributes match - do we care?
-		// eg - "I mod duration of all GEs that modify Health" would need to check to see if this mod modifies health before doing the stuff below.
-		// Or can we get by with just tags?
-
-		case EGameplayModEffect::Duration:
-		{
-			// don't modify infinite duration or instant unless we're overriding it
-			if (GetDuration() > 0.f || InMod.Info.ModifierOp == EGameplayModOp::Override)
-			{
-				Duration.Get()->ApplyMod(InMod.Info.ModifierOp, InMod.Aggregator, bApplyAsSnapshot);
-				NumApplied = 1;
+				case EGameplayModEffect::Duration:
+					// don't modify infinite duration or instant unless we're overriding it
+					if (GetDuration() > 0.f || InMod.Info.ModifierOp == EGameplayModOp::Override)
+					{
+						Duration.Get()->ApplyMod(InMod.Info.ModifierOp, InMod.Aggregator, bApplyAsSnapshot);
+					}
+					else
+					{
+						NumApplied = 0;
+					}
+					break;
+				case EGameplayModEffect::ChanceApplyTarget:
+					ChanceToApplyToTarget.Get()->ApplyMod(InMod.Info.ModifierOp, InMod.Aggregator, bApplyAsSnapshot);
+					break;
+				case EGameplayModEffect::ChanceExecuteEffect:
+					ChanceToExecuteOnGameplayEffect.Get()->ApplyMod(InMod.Info.ModifierOp, InMod.Aggregator, bApplyAsSnapshot);
+					break;
+
+				case EGameplayModEffect::LinkedGameplayEffect:
+					TargetEffectSpecs.Add(InMod.TargetEffectSpec.ToSharedRef());
+					break;
 			}
 
-			break;
-		}
-
-		case EGameplayModEffect::ChanceApplyTarget:
-		{
-			ChanceToApplyToTarget.Get()->ApplyMod(InMod.Info.ModifierOp, InMod.Aggregator, bApplyAsSnapshot);
-			NumApplied = 1;
-			break;
-		}
-
-		case EGameplayModEffect::ChanceExecuteEffect:
-		{
-			ChanceToExecuteOnGameplayEffect.Get()->ApplyMod(InMod.Info.ModifierOp, InMod.Aggregator, bApplyAsSnapshot);
-			NumApplied = 1;
-			break;
-		}
-
-		case EGameplayModEffect::LinkedGameplayEffect:
-		{
-			TargetEffectSpecs.Add(InMod.TargetEffectSpec.ToSharedRef());
-			NumApplied = 1;
-			break;
+			if (FinishedLoop)
+				break;
 		}
 	}
 
diff --git a/Engine/Source/Runtime/GameplayAbilities/Private/GameplayEffectTests.cpp b/Engine/Source/Runtime/GameplayAbilities/Private/GameplayEffectTests.cpp
--- a/Engine/Source/Runtime/GameplayAbilities/Private/GameplayEffectTests.cpp
+++ b/Engine/Source/Runtime/GameplayAbilities/Private/GameplayEffectTests.cpp
@@ -1612,6 +1612,7 @@ bool GameplayEffectsTest_BuffAppliesBuff(UWorld *World, FAutomationTestBase * Te
 		DummyBuffEffect->Modifiers[0].ModifierType = EGameplayMod::OutgoingGE;
 		DummyBuffEffect->Modifiers[0].EffectType = EGameplayModEffect::LinkedGameplayEffect;
 		DummyBuffEffect->Modifiers[0].RequiredTags.AddTag(IGameplayTagsModule::RequestGameplayTag(FName(TEXT("Damage.Buffable"))));
+		DummyBuffEffect->Modifiers[0].Attribute.SetUProperty(DamageProperty);
 		DummyBuffEffect->Modifiers[0].TargetEffect = BuffEffect;
 		DummyBuffEffect->Duration.SetValue(UGameplayEffect::INFINITE_DURATION);
 
@@ -1742,6 +1743,7 @@ bool GameplayEffectsTest_BuffIndirection(UWorld *World, FAutomationTestBase * Te
 		DummyBuffEffect->Modifiers.SetNum(1);
 		DummyBuffEffect->Modifiers[0].ModifierType = EGameplayMod::OutgoingGE;
 		DummyBuffEffect->Modifiers[0].EffectType = EGameplayModEffect::LinkedGameplayEffect;
+		DummyBuffEffect->Modifiers[0].Attribute.SetUProperty(DamageProperty);
 		DummyBuffEffect->Modifiers[0].TargetEffect = BuffEffect;
 		DummyBuffEffect->Modifiers[0].RequiredTags.AddTag(IGameplayTagsModule::RequestGameplayTag(FName(TEXT("Damage.Buffable"))));
 		DummyBuffEffect->Duration.SetValue(UGameplayEffect::INFINITE_DURATION);
@@ -1750,6 +1752,7 @@ bool GameplayEffectsTest_BuffIndirection(UWorld *World, FAutomationTestBase * Te
 		DummyBuffEffect2->Modifiers.SetNum(1);
 		DummyBuffEffect2->Modifiers[0].ModifierType = EGameplayMod::OutgoingGE;
 		DummyBuffEffect2->Modifiers[0].EffectType = EGameplayModEffect::LinkedGameplayEffect;
+		DummyBuffEffect2->Modifiers[0].Attribute.SetUProperty(DamageProperty);
 		DummyBuffEffect2->Modifiers[0].TargetEffect = DummyBuffEffect;
 		DummyBuffEffect2->Duration.SetValue(UGameplayEffect::INFINITE_DURATION);
 
@@ -5908,6 +5911,7 @@ bool GameplayEffectsTest_ModifyChanceToApplyToTarget(UWorld *World, FAutomationT
 		BaseEffect->Modifiers[0].Magnitude.SetValue(1.f);
 		BaseEffect->Modifiers[0].ModifierType = EGameplayMod::OutgoingGE;
 		BaseEffect->Modifiers[0].ModifierOp = EGameplayModOp::Additive;
+		BaseEffect->Modifiers[0].Attribute.SetUProperty(DamageProperty);
 		BaseEffect->Modifiers[0].EffectType = EGameplayModEffect::ChanceApplyTarget;
 		BaseEffect->Duration.SetValue(UGameplayEffect::INFINITE_DURATION);
 
@@ -6008,6 +6012,8 @@ bool GameplayEffectsTest_ModifyChanceToExecuteOnGE(UWorld *World, FAutomationTes
 		BaseEffect->Modifiers[0].ModifierType = EGameplayMod::ActiveGE;
 		BaseEffect->Modifiers[0].ModifierOp = EGameplayModOp::Additive;
 		BaseEffect->Modifiers[0].EffectType = EGameplayModEffect::ChanceExecuteEffect;
+		BaseEffect->Modifiers[0].ModifierOp = EGameplayModOp::Additive;
+		BaseEffect->Modifiers[0].Attribute.SetUProperty(DamageProperty);
 		BaseEffect->Duration.SetValue(UGameplayEffect::INFINITE_DURATION);
 
 		// Apply to self
