Code quality improvement - target data filters and handles cleaned up to use polymorphic structs with handles, major code duplication removed.

diff --git a/Engine/Source/Runtime/GameplayAbilities/Private/AbilitySystemBlueprintLibrary.cpp b/Engine/Source/Runtime/GameplayAbilities/Private/AbilitySystemBlueprintLibrary.cpp
--- a/Engine/Source/Runtime/GameplayAbilities/Private/AbilitySystemBlueprintLibrary.cpp
+++ b/Engine/Source/Runtime/GameplayAbilities/Private/AbilitySystemBlueprintLibrary.cpp
@@ -72,22 +72,17 @@ FGameplayAbilityTargetDataHandle UAbilitySystemBlueprintLibrary::AbilityTargetDa
 	}
 }
 
-
 FGameplayAbilityTargetDataHandle UAbilitySystemBlueprintLibrary::FilterTargetData(FGameplayAbilityTargetDataHandle TargetDataHandle, FGameplayTargetDataFilterHandle FilterHandle)
 {
 	FGameplayAbilityTargetDataHandle ReturnDataHandle;
 	
-	//return ReturnDataHandle;
-	check(FilterHandle.Filter.IsValid());		//Return unfiltered data
-
 	for (int32 i = 0; TargetDataHandle.IsValid(i); ++i)
 	{
 		FGameplayAbilityTargetData* UnfilteredData = TargetDataHandle.Get(i);
 		check(UnfilteredData);
 		if (UnfilteredData->GetActors().Num() > 0)
 		{
-			//TArray<TWeakObjectPtr<AActor>> FilteredActors = FilterActorArray(UnfilteredData->GetActors(), Filter);
-			TArray<TWeakObjectPtr<AActor>> FilteredActors = UnfilteredData->GetActors().FilterByPredicate(*FilterHandle.Filter.Get());
+			TArray<TWeakObjectPtr<AActor>> FilteredActors = UnfilteredData->GetActors().FilterByPredicate(FilterHandle);
 			if (FilteredActors.Num() > 0)
 			{
 				//Copy the data first, since we don't understand the internals of it
@@ -121,40 +116,6 @@ FGameplayTargetDataFilterHandle UAbilitySystemBlueprintLibrary::MakeFilterHandle
 	return FilterHandle;
 }
 
-
-/*
-TArray<TWeakObjectPtr<AActor>> UGameplayAbility::FilterActorArray(TArray<TWeakObjectPtr<AActor>> ArrayToFilter, FGameplayAbilityTargetDataActorFilter Filter)
-{
-	TArray<TWeakObjectPtr<AActor>> ReturnArray;
-	int32 NewArraySize = 0;
-	for (int32 i = 0; i < ArrayToFilter.Num(); ++i)
-	{
-		if (ArrayToFilter[i].IsValid())
-		{
-			if (Filter.FilterPassesForActor(ArrayToFilter[i].Get()))
-			{
-				++NewArraySize;
-			}
-		}
-	}
-	if (NewArraySize)
-	{
-		ReturnArray.Reserve(NewArraySize);
-		for (int32 i = 0; i < ArrayToFilter.Num(); ++i)
-		{
-			if (ArrayToFilter[i].IsValid())
-			{
-				if (Filter.FilterPassesForActor(ArrayToFilter[i].Get()))
-				{
-					ReturnArray.Add(ArrayToFilter[i].Get());
-				}
-			}
-		}
-	}
-	return ReturnArray;
-}
-*/
-
 FGameplayAbilityTargetDataHandle UAbilitySystemBlueprintLibrary::AbilityTargetDataFromHitResult(FHitResult HitResult)
 {
 	// Construct TargetData
diff --git a/Engine/Source/Runtime/GameplayAbilities/Public/Abilities/GameplayAbilityTargetDataFilter.h b/Engine/Source/Runtime/GameplayAbilities/Public/Abilities/GameplayAbilityTargetDataFilter.h
--- a/Engine/Source/Runtime/GameplayAbilities/Public/Abilities/GameplayAbilityTargetDataFilter.h
+++ b/Engine/Source/Runtime/GameplayAbilities/Public/Abilities/GameplayAbilityTargetDataFilter.h
@@ -17,7 +17,7 @@ namespace ETargetDataFilterSelf
 }
 
 USTRUCT(BlueprintType)
-struct FGameplayTargetDataFilter
+struct GAMEPLAYABILITIES_API FGameplayTargetDataFilter
 {
 	GENERATED_USTRUCT_BODY()
 
@@ -50,16 +50,6 @@ struct FGameplayTargetDataFilter
 		return true;
 	}
 
-	bool operator()(const TWeakObjectPtr<AActor> A) const
-	{
-		return FilterPassesForActor(A.Get());
-	}
-
-	bool operator()(const AActor* A) const
-	{
-		return FilterPassesForActor(A);
-	}
-
 	/** Filled out while running */
 	TWeakObjectPtr<UGameplayAbility> SourceAbility;
 
@@ -69,9 +59,36 @@ struct FGameplayTargetDataFilter
 
 
 USTRUCT(BlueprintType)
-struct FGameplayTargetDataFilterHandle
+struct GAMEPLAYABILITIES_API FGameplayTargetDataFilterHandle
 {
 	GENERATED_USTRUCT_BODY()
 
 	TSharedPtr<FGameplayTargetDataFilter> Filter;
+
+	bool FilterPassesForActor(const AActor* ActorToBeFiltered) const
+	{
+		if (!ActorToBeFiltered)
+		{
+			return false;
+		}
+		//Eventually, this might iterate through multiple filters. We'll need to decide how to designate OR versus AND functionality.
+		if (Filter.IsValid())
+		{
+			if (!Filter.Get()->FilterPassesForActor(ActorToBeFiltered))
+			{
+				return false;
+			}
+		}
+		return true;
+	}
+
+	bool operator()(const TWeakObjectPtr<AActor> ActorToBeFiltered) const
+	{
+		return FilterPassesForActor(ActorToBeFiltered.Get());
+	}
+
+	bool operator()(const AActor* ActorToBeFiltered) const
+	{
+		return FilterPassesForActor(ActorToBeFiltered);
+	}
 };
diff --git a/Engine/Source/Runtime/GameplayAbilities/Public/AbilitySystemBlueprintLibrary.h b/Engine/Source/Runtime/GameplayAbilities/Public/AbilitySystemBlueprintLibrary.h
--- a/Engine/Source/Runtime/GameplayAbilities/Public/AbilitySystemBlueprintLibrary.h
+++ b/Engine/Source/Runtime/GameplayAbilities/Public/AbilitySystemBlueprintLibrary.h
@@ -45,6 +45,7 @@ class UAbilitySystemBlueprintLibrary : public UBlueprintFunctionLibrary
 	UFUNCTION(BlueprintPure, Category = "Ability|TargetData")
 	static FGameplayAbilityTargetDataHandle	AbilityTargetDataFromActorArray(TArray<TWeakObjectPtr<AActor>> ActorArray, bool OneTargetPerHandle);
 
+	/** Create a new target data handle with filtration performed on the data */
 	UFUNCTION(BlueprintPure, Category = "Ability|TargetData")
 	static FGameplayAbilityTargetDataHandle	FilterTargetData(FGameplayAbilityTargetDataHandle TargetDataHandle, FGameplayTargetDataFilterHandle ActorFilterClass);
 
