Refactored world reticle to exist on TargetActor base class, as all subclasses are currently using it anyway. Added a parameter structure to world reticles, currently only supports scale field. World reticles themselves decide what to do with this information, and are notified that it is set by an event.

diff --git a/Engine/Source/Runtime/GameplayAbilities/Private/Abilities/GameplayAbilityTargetActor_Trace.cpp b/Engine/Source/Runtime/GameplayAbilities/Private/Abilities/GameplayAbilityTargetActor_Trace.cpp
--- a/Engine/Source/Runtime/GameplayAbilities/Private/Abilities/GameplayAbilityTargetActor_Trace.cpp
+++ b/Engine/Source/Runtime/GameplayAbilities/Private/Abilities/GameplayAbilityTargetActor_Trace.cpp
@@ -97,7 +97,7 @@ void AGameplayAbilityTargetActor_Trace::StartTargeting(UGameplayAbility* InAbili
 		AGameplayAbilityWorldReticle* SpawnedReticleActor = GetWorld()->SpawnActor<AGameplayAbilityWorldReticle>(ReticleClass, GetActorLocation(), GetActorRotation());
 		if (SpawnedReticleActor)
 		{
-			SpawnedReticleActor->InitializeReticle(this);
+			SpawnedReticleActor->InitializeReticle(this, ReticleParams);
 			ReticleActor = SpawnedReticleActor;
 
 			// This is to catch cases of playing on a listen server where we are using a replicated reticle actor.
diff --git a/Engine/Source/Runtime/GameplayAbilities/Private/Abilities/GameplayAbilityWorldReticle.cpp b/Engine/Source/Runtime/GameplayAbilities/Private/Abilities/GameplayAbilityWorldReticle.cpp
--- a/Engine/Source/Runtime/GameplayAbilities/Private/Abilities/GameplayAbilityWorldReticle.cpp
+++ b/Engine/Source/Runtime/GameplayAbilities/Private/Abilities/GameplayAbilityWorldReticle.cpp
@@ -17,12 +17,14 @@ AGameplayAbilityWorldReticle::AGameplayAbilityWorldReticle(const class FPostCons
 	PrimaryActorTick.TickGroup = TG_PrePhysics;
 }
 
-void AGameplayAbilityWorldReticle::InitializeReticle(AGameplayAbilityTargetActor* InTargetingActor)
+void AGameplayAbilityWorldReticle::InitializeReticle(AGameplayAbilityTargetActor* InTargetingActor, FWorldReticleParameters InParameters)
 {
 	check(InTargetingActor);
 	TargetingActor = InTargetingActor;
 	MasterPC = InTargetingActor->MasterPC;
 	AddTickPrerequisiteActor(TargetingActor);		//We want the reticle to tick after the targeting actor so that designers have the final say on the position
+	Parameters = InParameters;
+	OnParametersInitialized();
 }
 
 bool AGameplayAbilityWorldReticle::IsNetRelevantFor(class APlayerController* RealViewer, AActor* Viewer, const FVector& SrcLocation)
diff --git a/Engine/Source/Runtime/GameplayAbilities/Public/Abilities/GameplayAbilityTargetActor.h b/Engine/Source/Runtime/GameplayAbilities/Public/Abilities/GameplayAbilityTargetActor.h
--- a/Engine/Source/Runtime/GameplayAbilities/Public/Abilities/GameplayAbilityTargetActor.h
+++ b/Engine/Source/Runtime/GameplayAbilities/Public/Abilities/GameplayAbilityTargetActor.h
@@ -2,6 +2,7 @@
 #pragma once
 
 #include "GameplayAbility.h"
+#include "Abilities/GameplayAbilityWorldReticle.h"
 #include "GameplayAbilityTargetActor.generated.h"
 
 UENUM(BlueprintType)
@@ -81,6 +82,14 @@ public:
 	UPROPERTY(BlueprintReadOnly, Replicated, Category = Targeting)
 	AActor* SourceActor;
 
+	/** Parameters for world reticle. Usage of these parameters is dependent on the reticle. */
+	UPROPERTY(BlueprintReadWrite, EditAnywhere, meta = (ExposeOnSpawn = true), Category = Targeting)
+	FWorldReticleParameters ReticleParams;
+
+	/** Reticle that will appear on top of acquired targets. Reticles will be spawned/despawned as targets are acquired/lost. */
+	UPROPERTY(BlueprintReadWrite, EditAnywhere, meta = (ExposeOnSpawn = true), Category = Targeting)
+	TSubclassOf<AGameplayAbilityWorldReticle> ReticleClass;		//Using a special class for replication purposes.
+
 	/** Draw the debug information (if applicable) for this targeting actor. */
 	UPROPERTY(BlueprintReadWrite, EditAnywhere, Replicated, meta = (ExposeOnSpawn = true), Category = Targeting)
 	bool bDebug;
diff --git a/Engine/Source/Runtime/GameplayAbilities/Public/Abilities/GameplayAbilityTargetActor_Trace.h b/Engine/Source/Runtime/GameplayAbilities/Public/Abilities/GameplayAbilityTargetActor_Trace.h
--- a/Engine/Source/Runtime/GameplayAbilities/Public/Abilities/GameplayAbilityTargetActor_Trace.h
+++ b/Engine/Source/Runtime/GameplayAbilities/Public/Abilities/GameplayAbilityTargetActor_Trace.h
@@ -30,9 +30,6 @@ public:
 	UPROPERTY(BlueprintReadWrite, EditAnywhere, meta = (ExposeOnSpawn = true), Category = Projectile)
 	float MaxRange;
 
-	UPROPERTY(BlueprintReadWrite, EditAnywhere, meta = (ExposeOnSpawn = true), Category = Projectile)
-	TSubclassOf<AGameplayAbilityWorldReticle> ReticleClass;		//Using a special class for replication purposes
-
 protected:
 	virtual FHitResult PerformTrace(AActor* InSourceActor) const PURE_VIRTUAL(AGameplayAbilityTargetActor_Trace, return FHitResult(););
 
diff --git a/Engine/Source/Runtime/GameplayAbilities/Public/Abilities/GameplayAbilityWorldReticle.h b/Engine/Source/Runtime/GameplayAbilities/Public/Abilities/GameplayAbilityWorldReticle.h
--- a/Engine/Source/Runtime/GameplayAbilities/Public/Abilities/GameplayAbilityWorldReticle.h
+++ b/Engine/Source/Runtime/GameplayAbilities/Public/Abilities/GameplayAbilityWorldReticle.h
@@ -5,6 +5,21 @@
 
 class AGameplayAbilityTargetActor;
 
+USTRUCT(BlueprintType)
+struct FWorldReticleParameters
+{
+	GENERATED_USTRUCT_BODY()
+
+	//Use this so that we can't slip in new parameters without some actor knowing about it.
+	void Initialize(FVector InAOEScale)
+	{
+		AOEScale = InAOEScale;
+	}
+
+	UPROPERTY(BlueprintReadWrite, Category = Reticle)
+	FVector AOEScale;
+};
+
 /** Reticles allow targeting to be visualized. Tasks can spawn these. Artists/designers can create BPs for these. */
 UCLASS(Blueprintable, notplaceable)
 class GAMEPLAYABILITIES_API AGameplayAbilityWorldReticle : public AActor
@@ -20,7 +35,13 @@ public:
 
 	virtual bool IsNetRelevantFor(class APlayerController* RealViewer, AActor* Viewer, const FVector& SrcLocation) override;
 
-	void InitializeReticle(AGameplayAbilityTargetActor* InTargetingActor);
+	void InitializeReticle(AGameplayAbilityTargetActor* InTargetingActor, FWorldReticleParameters InParameters);
+
+	UFUNCTION(BlueprintImplementableEvent, Category = Reticle)
+	void OnParametersInitialized();
+
+	UPROPERTY(EditAnywhere, BlueprintReadWrite, meta = (ExposeOnSpawn = "true"), Category = "Reticle")
+	FWorldReticleParameters Parameters;
 
 protected:
 	/** This is used in the process of determining whether we should replicate to a specific client. */
