Fixed FGameplayAbilityTargetData_SingleTargetHit::GetEndPoint()

Changed 'FGameplayAbilityTargetData_SingleTargetHit::GetEndPoint()' to return 'HitResult.Location' instead of 'HitResult.ImpactPoint' since the impact point is not set if the trace wasn't blocked.

diff --git a/Engine/Source/Runtime/GameplayAbilities/Public/Abilities/GameplayAbilityTargetTypes.h b/Engine/Source/Runtime/GameplayAbilities/Public/Abilities/GameplayAbilityTargetTypes.h
--- a/Engine/Source/Runtime/GameplayAbilities/Public/Abilities/GameplayAbilityTargetTypes.h
+++ b/Engine/Source/Runtime/GameplayAbilities/Public/Abilities/GameplayAbilityTargetTypes.h
@@ -557,7 +557,7 @@ struct GAMEPLAYABILITIES_API FGameplayAbilityTargetData_SingleTargetHit : public
 
 	virtual FVector GetEndPoint() const override
 	{
-		return HitResult.ImpactPoint;
+		return HitResult.Location;
 	}
 
 	// -------------------------------------
