Fixed filtering tests to match new ability function naming.

diff --git a/Engine/Source/Runtime/GameplayAbilities/Public/Abilities/GameplayAbilityTargetDataFilter.h b/Engine/Source/Runtime/GameplayAbilities/Public/Abilities/GameplayAbilityTargetDataFilter.h
--- a/Engine/Source/Runtime/GameplayAbilities/Public/Abilities/GameplayAbilityTargetDataFilter.h
+++ b/Engine/Source/Runtime/GameplayAbilities/Public/Abilities/GameplayAbilityTargetDataFilter.h
@@ -28,17 +28,17 @@ struct GAMEPLAYABILITIES_API FGameplayTargetDataFilter
 	virtual bool FilterPassesForActor(const AActor* ActorToBeFiltered) const
 	{
 		check(SourceAbility.IsValid());
-		check(SourceAbility.Get()->GetOwningActorFromActorInfo());
+		check(SourceAbility.Get()->GetAvatarActorFromActorInfo());
 		switch (SelfFilter.GetValue())
 		{
 		case ETargetDataFilterSelf::Type::TDFS_NoOthers:
-			if (ActorToBeFiltered != SourceAbility.Get()->GetOwningActorFromActorInfo())
+			if (ActorToBeFiltered != SourceAbility.Get()->GetAvatarActorFromActorInfo())
 			{
 				return false;
 			}
 			break;
 		case ETargetDataFilterSelf::Type::TDFS_NoSelf:
-			if (ActorToBeFiltered == SourceAbility.Get()->GetOwningActorFromActorInfo())
+			if (ActorToBeFiltered == SourceAbility.Get()->GetAvatarActorFromActorInfo())
 			{
 				return false;
 			}
