#UE4 - Added TraceChannel parameter to GameplaAbilityTargetActor_Trace

diff --git a/Engine/Source/Runtime/GameplayAbilities/Private/Abilities/GameplayAbilityTargetActor_GroundTrace.cpp b/Engine/Source/Runtime/GameplayAbilities/Private/Abilities/GameplayAbilityTargetActor_GroundTrace.cpp
--- a/Engine/Source/Runtime/GameplayAbilities/Private/Abilities/GameplayAbilityTargetActor_GroundTrace.cpp
+++ b/Engine/Source/Runtime/GameplayAbilities/Private/Abilities/GameplayAbilityTargetActor_GroundTrace.cpp
@@ -144,7 +144,7 @@ FHitResult AGameplayAbilityTargetActor_GroundTrace::PerformTrace(AActor* InSourc
 	FVector LineTraceStart = TraceStart;
 	FVector LineTraceEnd = TraceEnd;
 	//Use a line trace initially to see where the player is actually pointing
-	InSourceActor->GetWorld()->LineTraceSingle(ReturnHitResult, TraceStart, TraceEnd, ECC_WorldStatic, Params);
+	InSourceActor->GetWorld()->LineTraceSingle(ReturnHitResult, TraceStart, TraceEnd, TraceChannel, Params);
 	//Default to end of trace line if we don't hit anything.
 	if (!ReturnHitResult.bBlockingHit)
 	{
@@ -156,7 +156,7 @@ FHitResult AGameplayAbilityTargetActor_GroundTrace::PerformTrace(AActor* InSourc
 	TraceEnd = TraceStart;
 	TraceStart.Z += CollisionHeightOffset;
 	TraceEnd.Z -= 99999.0f;
-	InSourceActor->GetWorld()->LineTraceSingle(ReturnHitResult, TraceStart, TraceEnd, ECC_WorldStatic, Params);
+	InSourceActor->GetWorld()->LineTraceSingle(ReturnHitResult, TraceStart, TraceEnd, TraceChannel, Params);
 	//if (!ReturnHitResult.bBlockingHit) then our endpoint may be off the map. Hopefully this is only possible in debug maps.
 
 	//Use collision shape to find a valid ground spot, if appropriate
diff --git a/Engine/Source/Runtime/GameplayAbilities/Private/Abilities/GameplayAbilityTargetActor_SingleLineTrace.cpp b/Engine/Source/Runtime/GameplayAbilities/Private/Abilities/GameplayAbilityTargetActor_SingleLineTrace.cpp
--- a/Engine/Source/Runtime/GameplayAbilities/Private/Abilities/GameplayAbilityTargetActor_SingleLineTrace.cpp
+++ b/Engine/Source/Runtime/GameplayAbilities/Private/Abilities/GameplayAbilityTargetActor_SingleLineTrace.cpp
@@ -38,7 +38,7 @@ FHitResult AGameplayAbilityTargetActor_SingleLineTrace::PerformTrace(AActor* InS
 	// ------------------------------------------------------
 
 	FHitResult ReturnHitResult;
-	InSourceActor->GetWorld()->LineTraceSingle(ReturnHitResult, TraceStart, TraceEnd, ECC_WorldStatic, Params);
+	InSourceActor->GetWorld()->LineTraceSingle(ReturnHitResult, TraceStart, TraceEnd, TraceChannel, Params);
 	//Default to end of trace line if we don't hit anything.
 	if (!ReturnHitResult.bBlockingHit)
 	{
diff --git a/Engine/Source/Runtime/GameplayAbilities/Private/Abilities/GameplayAbilityTargetActor_Trace.cpp b/Engine/Source/Runtime/GameplayAbilities/Private/Abilities/GameplayAbilityTargetActor_Trace.cpp
--- a/Engine/Source/Runtime/GameplayAbilities/Private/Abilities/GameplayAbilityTargetActor_Trace.cpp
+++ b/Engine/Source/Runtime/GameplayAbilities/Private/Abilities/GameplayAbilityTargetActor_Trace.cpp
@@ -20,6 +20,7 @@ AGameplayAbilityTargetActor_Trace::AGameplayAbilityTargetActor_Trace(const class
 	StaticTargetFunction = false;
 
 	MaxRange = 999999.0f;
+	TraceChannel = ECC_WorldStatic;
 }
 
 void AGameplayAbilityTargetActor_Trace::EndPlay(const EEndPlayReason::Type EndPlayReason)
@@ -53,7 +54,7 @@ void AGameplayAbilityTargetActor_Trace::AimWithPlayerController(AActor* InSource
 		ClipCameraRayToAbilityRange(CamLoc, CamDir, TraceStart, MaxRange, CamTarget);
 
 		FHitResult TempHitResult;
-		InSourceActor->GetWorld()->LineTraceSingle(TempHitResult, CamLoc, CamTarget, ECC_WorldStatic, Params);
+		InSourceActor->GetWorld()->LineTraceSingle(TempHitResult, CamLoc, CamTarget, TraceChannel, Params);
 		//TraceStart = ;
 		if (TempHitResult.bBlockingHit && (FVector::DistSquared(TraceStart, TempHitResult.Location) <= (MaxRange * MaxRange)))
 		{
diff --git a/Engine/Source/Runtime/GameplayAbilities/Public/Abilities/GameplayAbilityTargetActor_Trace.h b/Engine/Source/Runtime/GameplayAbilities/Public/Abilities/GameplayAbilityTargetActor_Trace.h
--- a/Engine/Source/Runtime/GameplayAbilities/Public/Abilities/GameplayAbilityTargetActor_Trace.h
+++ b/Engine/Source/Runtime/GameplayAbilities/Public/Abilities/GameplayAbilityTargetActor_Trace.h
@@ -30,6 +30,9 @@ public:
 	UPROPERTY(BlueprintReadWrite, EditAnywhere, meta = (ExposeOnSpawn = true), Category = Projectile)
 	float MaxRange;
 
+	UPROPERTY(BlueprintReadWrite, EditAnywhere, meta = (ExposeOnSpawn = true), Category = Trace)
+	TEnumAsByte<ECollisionChannel> TraceChannel;
+
 protected:
 	virtual FHitResult PerformTrace(AActor* InSourceActor) const PURE_VIRTUAL(AGameplayAbilityTargetActor_Trace, return FHitResult(););
 
