warning update
-specify with parameter if you want to manually look through component list when trying to find ability system component.

diff --git a/Engine/Source/Runtime/GameplayAbilities/Private/AbilitySystemGlobals.cpp b/Engine/Source/Runtime/GameplayAbilities/Private/AbilitySystemGlobals.cpp
--- a/Engine/Source/Runtime/GameplayAbilities/Private/AbilitySystemGlobals.cpp
+++ b/Engine/Source/Runtime/GameplayAbilities/Private/AbilitySystemGlobals.cpp
@@ -96,7 +96,7 @@ UAbilitySystemGlobals& UAbilitySystemGlobals::Get()
 }
 
 /** Helping function to avoid having to manually cast */
-UAbilitySystemComponent* UAbilitySystemGlobals::GetAbilitySystemComponentFromActor(AActor* Actor)
+UAbilitySystemComponent* UAbilitySystemGlobals::GetAbilitySystemComponentFromActor(AActor* Actor, bool LookForComponent)
 {
 	if (Actor == nullptr)
 	{
@@ -109,10 +109,15 @@ UAbilitySystemComponent* UAbilitySystemGlobals::GetAbilitySystemComponentFromAct
 		return ASI->GetAbilitySystemComponent();
 	}
 
-	/** This is slow and not desirable */
-	ABILITY_LOG(Warning, TEXT("GetAbilitySystemComponentFromActor called on Actor that is not IAbilitySystemInterface. This slow!"));
+	if (LookForComponent)
+	{
+		/** This is slow and not desirable */
+		ABILITY_LOG(Warning, TEXT("GetAbilitySystemComponentFromActor called on %s that is not IAbilitySystemInterface. This slow!"), *Actor->GetName());
+
+		return Actor->FindComponentByClass<UAbilitySystemComponent>();
+	}
 
-	return Actor->FindComponentByClass<UAbilitySystemComponent>();
+	return nullptr;
 }
 
 // --------------------------------------------------------------------
diff --git a/Engine/Source/Runtime/GameplayAbilities/Public/AbilitySystemGlobals.h b/Engine/Source/Runtime/GameplayAbilities/Public/AbilitySystemGlobals.h
--- a/Engine/Source/Runtime/GameplayAbilities/Public/AbilitySystemGlobals.h
+++ b/Engine/Source/Runtime/GameplayAbilities/Public/AbilitySystemGlobals.h
@@ -59,7 +59,7 @@ class GAMEPLAYABILITIES_API UAbilitySystemGlobals : public UObject
 
 	static UAbilitySystemGlobals& Get();
 
-	static UAbilitySystemComponent* GetAbilitySystemComponentFromActor(AActor* Actor);
+	static UAbilitySystemComponent* GetAbilitySystemComponentFromActor(AActor* Actor, bool LookForComponent=true);
 
 	/** Should allocate a project specific AbilityActorInfo struct. Caller is responsible for deallocation */
 	virtual FGameplayAbilityActorInfo * AllocAbilityActorInfo() const;
