OrionGameplayVolume

diff --git a/Engine/Source/Runtime/GameplayAbilities/Private/Abilities/Tasks/AbilityTask_SpawnActor.cpp b/Engine/Source/Runtime/GameplayAbilities/Private/Abilities/Tasks/AbilityTask_SpawnActor.cpp
--- a/Engine/Source/Runtime/GameplayAbilities/Private/Abilities/Tasks/AbilityTask_SpawnActor.cpp
+++ b/Engine/Source/Runtime/GameplayAbilities/Private/Abilities/Tasks/AbilityTask_SpawnActor.cpp
@@ -28,7 +28,13 @@ bool UAbilityTask_SpawnActor::BeginSpawningActor(UObject* WorldContextObject, FG
 		SpawnedActor = World->SpawnActorDeferred<AActor>(InClass, FVector::ZeroVector, FRotator::ZeroRotator, NULL, NULL, true);
 	}
 	
-	return (SpawnedActor != nullptr);
+	if (SpawnedActor == nullptr)
+	{
+		DidNotSpawn.Broadcast(nullptr);
+		return false;
+	}
+
+	return true;
 }
 
 void UAbilityTask_SpawnActor::FinishSpawningActor(UObject* WorldContextObject, FGameplayAbilityTargetDataHandle TargetData, AActor* SpawnedActor)
diff --git a/Engine/Source/Runtime/GameplayAbilities/Public/Abilities/Tasks/AbilityTask_SpawnActor.h b/Engine/Source/Runtime/GameplayAbilities/Public/Abilities/Tasks/AbilityTask_SpawnActor.h
--- a/Engine/Source/Runtime/GameplayAbilities/Public/Abilities/Tasks/AbilityTask_SpawnActor.h
+++ b/Engine/Source/Runtime/GameplayAbilities/Public/Abilities/Tasks/AbilityTask_SpawnActor.h
@@ -27,6 +27,10 @@ class UAbilityTask_SpawnActor: public UAbilityTask
 
 	UPROPERTY(BlueprintAssignable)
 	FSpawnActorDelegate	Success;
+
+	/** Called when we can't spawn: on clients or potentially on server if they fail to spawn (rare) */
+	UPROPERTY(BlueprintAssignable)
+	FSpawnActorDelegate	DidNotSpawn;
 	
 	/** Spawn new Actor on the network authority (server) */
 	UFUNCTION(BlueprintCallable, meta=(HidePin = "WorldContextObject", DefaultToSelf = "WorldContextObject", BlueprintInternalUseOnly = "true"), Category="Ability|Tasks")
