TargetActor_Radius uses filters.

diff --git a/Engine/Source/Runtime/GameplayAbilities/Private/Abilities/GameplayAbilityTargetActor_Radius.cpp b/Engine/Source/Runtime/GameplayAbilities/Private/Abilities/GameplayAbilityTargetActor_Radius.cpp
--- a/Engine/Source/Runtime/GameplayAbilities/Private/Abilities/GameplayAbilityTargetActor_Radius.cpp
+++ b/Engine/Source/Runtime/GameplayAbilities/Private/Abilities/GameplayAbilityTargetActor_Radius.cpp
@@ -55,26 +55,21 @@ TArray<TWeakObjectPtr<AActor> >	AGameplayAbilityTargetActor_Radius::PerformOverl
 	static FName RadiusTargetingOverlap = FName(TEXT("RadiusTargetingOverlap"));
 	bool bTraceComplex = false;
 	
-	TArray<TWeakObjectPtr<AActor> > ActorsToIgnore;
-	ActorsToIgnore.Add(SourceActor);
-
 	FCollisionQueryParams Params(RadiusTargetingOverlap, bTraceComplex);
 	Params.bReturnPhysicalMaterial = true;
 	Params.bTraceAsyncScene = true;
-	Params.AddIgnoredActors(ActorsToIgnore);
 
 	TArray<FOverlapResult> Overlaps;
 
 	SourceActor->GetWorld()->OverlapMulti(Overlaps, Origin, FQuat::Identity, ECC_Pawn, FCollisionShape::MakeSphere(Radius), Params);
 
-
-	TArray<TWeakObjectPtr<AActor> >	HitActors;
+	TArray<TWeakObjectPtr<AActor>>	HitActors;
 
 	for (int32 i=0; i < Overlaps.Num(); ++i)
 	{
-		// Temp, this needs to be way more robust
+		//Should this check to see if these pawns are in the AimTarget list?
 		APawn* PawnActor = Cast<APawn>(Overlaps[i].GetActor());
-		if (PawnActor)
+		if (PawnActor && FilterHandle.FilterPassesForActor(PawnActor))
 		{
 			HitActors.Add(PawnActor);
 		}
diff --git a/Engine/Source/Runtime/GameplayAbilities/Public/Abilities/GameplayAbilityTargetActor_Radius.h b/Engine/Source/Runtime/GameplayAbilities/Public/Abilities/GameplayAbilityTargetActor_Radius.h
--- a/Engine/Source/Runtime/GameplayAbilities/Public/Abilities/GameplayAbilityTargetActor_Radius.h
+++ b/Engine/Source/Runtime/GameplayAbilities/Public/Abilities/GameplayAbilityTargetActor_Radius.h
@@ -4,6 +4,7 @@
 
 #include "GameplayAbilityTargetActor.h"
 #include "GameplayAbilityWorldReticle.h"
+#include "GameplayAbilityTargetDataFilter.h"
 #include "GameplayAbilityTargetActor_Radius.generated.h"
 
 /** Selects everything within a given radius of the source actor. */
@@ -18,6 +19,11 @@ public:
 	
 	virtual void ConfirmTargetingAndContinue() override;
 
+	/** Filter to decide who can be selected. */
+	UPROPERTY(BlueprintReadWrite, EditAnywhere, meta = (ExposeOnSpawn = true), Category = Targeting)
+	FGameplayTargetDataFilterHandle FilterHandle;
+
+	/** Radius of target acquisition around the ability's start location. */
 	UPROPERTY(BlueprintReadWrite, EditAnywhere, meta = (ExposeOnSpawn = true), Category = Radius)
 	float Radius;
 
