Ground targeting will now decline confirmation without canceling if you try to target an invalid location.

diff --git a/Engine/Source/Runtime/GameplayAbilities/Private/Abilities/GameplayAbilityTargetActor.cpp b/Engine/Source/Runtime/GameplayAbilities/Private/Abilities/GameplayAbilityTargetActor.cpp
--- a/Engine/Source/Runtime/GameplayAbilities/Private/Abilities/GameplayAbilityTargetActor.cpp
+++ b/Engine/Source/Runtime/GameplayAbilities/Private/Abilities/GameplayAbilityTargetActor.cpp
@@ -34,18 +34,29 @@ void AGameplayAbilityTargetActor::StartTargeting(UGameplayAbility* Ability)
 	OwningAbility = Ability;
 }
 
+bool AGameplayAbilityTargetActor::IsConfirmTargetingAllowed()
+{
+	return true;
+}
+
 void AGameplayAbilityTargetActor::ConfirmTargetingAndContinue()
 {
 	check(ShouldProduceTargetData());
-	TargetDataReadyDelegate.Broadcast(FGameplayAbilityTargetDataHandle());
+	if (IsConfirmTargetingAllowed())
+	{
+		TargetDataReadyDelegate.Broadcast(FGameplayAbilityTargetDataHandle());
+	}
 }
 
 void AGameplayAbilityTargetActor::ConfirmTargeting()
 {
-	ConfirmTargetingAndContinue();
-	if (bDestroyOnConfirmation)
+	if (IsConfirmTargetingAllowed())
 	{
-		Destroy();
+		ConfirmTargetingAndContinue();
+		if (bDestroyOnConfirmation)
+		{
+			Destroy();
+		}
 	}
 }
 
diff --git a/Engine/Source/Runtime/GameplayAbilities/Private/Abilities/GameplayAbilityTargetActor_GroundTrace.cpp b/Engine/Source/Runtime/GameplayAbilities/Private/Abilities/GameplayAbilityTargetActor_GroundTrace.cpp
--- a/Engine/Source/Runtime/GameplayAbilities/Private/Abilities/GameplayAbilityTargetActor_GroundTrace.cpp
+++ b/Engine/Source/Runtime/GameplayAbilities/Private/Abilities/GameplayAbilityTargetActor_GroundTrace.cpp
@@ -74,7 +74,7 @@ bool AGameplayAbilityTargetActor_GroundTrace::AdjustCollisionResultForShape(cons
 	{
 		TraceEnd = TraceStart = OriginalEndPoint - (LerpValue * Movement);
 		TraceEnd.Z -= 99999.0f;
-		ThisWorld->SweepSingle(LocalResult, TraceStart, TraceEnd, FQuat::Identity, ECC_WorldStatic, CollisionShape, Params);
+		ThisWorld->SweepSingle(LocalResult, TraceStart, TraceEnd, FQuat::Identity, TraceChannel, CollisionShape, Params);
 		if (!LocalResult.bStartPenetrating)
 		{
 			if (!LocalResult.bBlockingHit)
@@ -123,7 +123,7 @@ bool AGameplayAbilityTargetActor_GroundTrace::AdjustCollisionResultForShape(cons
 	return false;
 }
 
-FHitResult AGameplayAbilityTargetActor_GroundTrace::PerformTrace(AActor* InSourceActor) const
+FHitResult AGameplayAbilityTargetActor_GroundTrace::PerformTrace(AActor* InSourceActor)
 {
 	static const FName LineTraceSingleName(TEXT("AGameplayAbilityTargetActor_GroundTrace"));
 	bool bTraceComplex = false;
@@ -159,6 +159,8 @@ FHitResult AGameplayAbilityTargetActor_GroundTrace::PerformTrace(AActor* InSourc
 	InSourceActor->GetWorld()->LineTraceSingle(ReturnHitResult, TraceStart, TraceEnd, TraceChannel, Params);
 	//if (!ReturnHitResult.bBlockingHit) then our endpoint may be off the map. Hopefully this is only possible in debug maps.
 
+	bLastTraceWasGood = true;		//So far, we're good. If we need a ground spot and can't find one, we'll come back.
+
 	//Use collision shape to find a valid ground spot, if appropriate
 	if (CollisionShape.ShapeType != ECollisionShape::Line)
 	{
@@ -166,7 +168,8 @@ FHitResult AGameplayAbilityTargetActor_GroundTrace::PerformTrace(AActor* InSourc
 		TraceStart = InSourceActor->GetActorLocation();
 		TraceEnd = ReturnHitResult.Location;
 		TraceStart.Z += CollisionHeightOffset;
-		if (AdjustCollisionResultForShape(TraceStart, TraceEnd, Params, ReturnHitResult))
+		bLastTraceWasGood = AdjustCollisionResultForShape(TraceStart, TraceEnd, Params, ReturnHitResult);
+		if (bLastTraceWasGood)
 		{
 			ReturnHitResult.Location.Z -= CollisionHeightOffset;	//Undo the artificial height adjustment
 		}
@@ -174,8 +177,14 @@ FHitResult AGameplayAbilityTargetActor_GroundTrace::PerformTrace(AActor* InSourc
 
 	if (AActor* LocalReticleActor = ReticleActor.Get())
 	{
+		//TODO Special (for now) functionality: We should tell the reticle to turn red or something, indicating this isn't a valid location
 		LocalReticleActor->SetActorLocation(ReturnHitResult.Location);
 	}
 
 	return ReturnHitResult;
 }
+
+bool AGameplayAbilityTargetActor_GroundTrace::IsConfirmTargetingAllowed()
+{
+	return bLastTraceWasGood;
+}
diff --git a/Engine/Source/Runtime/GameplayAbilities/Private/Abilities/GameplayAbilityTargetActor_SingleLineTrace.cpp b/Engine/Source/Runtime/GameplayAbilities/Private/Abilities/GameplayAbilityTargetActor_SingleLineTrace.cpp
--- a/Engine/Source/Runtime/GameplayAbilities/Private/Abilities/GameplayAbilityTargetActor_SingleLineTrace.cpp
+++ b/Engine/Source/Runtime/GameplayAbilities/Private/Abilities/GameplayAbilityTargetActor_SingleLineTrace.cpp
@@ -17,7 +17,7 @@ AGameplayAbilityTargetActor_SingleLineTrace::AGameplayAbilityTargetActor_SingleL
 {
 }
 
-FHitResult AGameplayAbilityTargetActor_SingleLineTrace::PerformTrace(AActor* InSourceActor) const
+FHitResult AGameplayAbilityTargetActor_SingleLineTrace::PerformTrace(AActor* InSourceActor)
 {
 	static const FName LineTraceSingleName(TEXT("AGameplayAbilityTargetActor_SingleLineTrace"));
 	bool bTraceComplex = false;
diff --git a/Engine/Source/Runtime/GameplayAbilities/Private/Abilities/GameplayAbilityTargetActor_Trace.cpp b/Engine/Source/Runtime/GameplayAbilities/Private/Abilities/GameplayAbilityTargetActor_Trace.cpp
--- a/Engine/Source/Runtime/GameplayAbilities/Private/Abilities/GameplayAbilityTargetActor_Trace.cpp
+++ b/Engine/Source/Runtime/GameplayAbilities/Private/Abilities/GameplayAbilityTargetActor_Trace.cpp
@@ -35,8 +35,10 @@ void AGameplayAbilityTargetActor_Trace::EndPlay(const EEndPlayReason::Type EndPl
 
 FGameplayAbilityTargetDataHandle AGameplayAbilityTargetActor_Trace::StaticGetTargetData(UWorld * World, const FGameplayAbilityActorInfo* ActorInfo, FGameplayAbilityActivationInfo ActivationInfo) const
 {
+	check(false);		//This should never actually be called, and if it is, it will require a const version of PerformTrace()
 	AActor* StaticSourceActor = ActorInfo->AvatarActor.Get();
-	return MakeTargetData(PerformTrace(StaticSourceActor));
+	//return MakeTargetData(PerformTrace(StaticSourceActor));		//Old way, requires PerformTrace to be a const call
+	return MakeTargetData(FHitResult());
 }
 
 void AGameplayAbilityTargetActor_Trace::AimWithPlayerController(AActor* InSourceActor, FCollisionQueryParams Params, FVector TraceStart, FVector& TraceEnd) const
diff --git a/Engine/Source/Runtime/GameplayAbilities/Private/AbilitySystemComponent_Abilities.cpp b/Engine/Source/Runtime/GameplayAbilities/Private/AbilitySystemComponent_Abilities.cpp
--- a/Engine/Source/Runtime/GameplayAbilities/Private/AbilitySystemComponent_Abilities.cpp
+++ b/Engine/Source/Runtime/GameplayAbilities/Private/AbilitySystemComponent_Abilities.cpp
@@ -975,15 +975,22 @@ void UAbilitySystemComponent::ServerInputRelease_Implementation(FGameplayAbility
 
 void UAbilitySystemComponent::TargetConfirm()
 {
+	TArray<AGameplayAbilityTargetActor*> LeftoverTargetActors;
 	for (AGameplayAbilityTargetActor* TargetActor : SpawnedTargetActors)
 	{
 		if (TargetActor)
 		{
-			TargetActor->ConfirmTargeting();
+			if (TargetActor->IsConfirmTargetingAllowed())
+			{
+				TargetActor->ConfirmTargeting();
+			}
+			else
+			{
+				LeftoverTargetActors.Add(TargetActor);
+			}
 		}
 	}
-
-	SpawnedTargetActors.Empty();
+	SpawnedTargetActors = LeftoverTargetActors;		//These actors declined to confirm targeting, so keep contact with them.
 }
 
 void UAbilitySystemComponent::TargetCancel()
diff --git a/Engine/Source/Runtime/GameplayAbilities/Public/Abilities/GameplayAbilityTargetActor.h b/Engine/Source/Runtime/GameplayAbilities/Public/Abilities/GameplayAbilityTargetActor.h
--- a/Engine/Source/Runtime/GameplayAbilities/Public/Abilities/GameplayAbilityTargetActor.h
+++ b/Engine/Source/Runtime/GameplayAbilities/Public/Abilities/GameplayAbilityTargetActor.h
@@ -42,10 +42,12 @@ public:
 	/** Initialize and begin targeting logic  */
 	virtual void StartTargeting(UGameplayAbility* Ability);
 
+	virtual bool IsConfirmTargetingAllowed();
+
 	/** Requesting targeting data, but not necessarily stopping/destroying the task. Useful for external target data requests. */
 	virtual void ConfirmTargetingAndContinue();
 
-	/** Outside code is saying 'stop and just give me what you have' */
+	/** Outside code is saying 'stop and just give me what you have.' Returns true if the ability accepts this and can be forgotten. */
 	UFUNCTION()
 	virtual void ConfirmTargeting();
 
diff --git a/Engine/Source/Runtime/GameplayAbilities/Public/Abilities/GameplayAbilityTargetActor_GroundTrace.h b/Engine/Source/Runtime/GameplayAbilities/Public/Abilities/GameplayAbilityTargetActor_GroundTrace.h
--- a/Engine/Source/Runtime/GameplayAbilities/Public/Abilities/GameplayAbilityTargetActor_GroundTrace.h
+++ b/Engine/Source/Runtime/GameplayAbilities/Public/Abilities/GameplayAbilityTargetActor_GroundTrace.h
@@ -13,6 +13,10 @@ class GAMEPLAYABILITIES_API AGameplayAbilityTargetActor_GroundTrace : public AGa
 
 	virtual void StartTargeting(UGameplayAbility* InAbility) override;
 
+	/** This trace channel is used when finding the ground, and may be different from the direct line trace channel */
+	//UPROPERTY(BlueprintReadWrite, EditAnywhere, meta = (ExposeOnSpawn = true), Category = Trace)
+	//TEnumAsByte<ECollisionChannel> TraceChannelToGround;
+
 	/** Radius for a sphere or capsule. */
 	UPROPERTY(BlueprintReadWrite, EditAnywhere, meta = (ExposeOnSpawn = true), Category = Targeting)
 	float CollisionRadius;
@@ -24,10 +28,13 @@ class GAMEPLAYABILITIES_API AGameplayAbilityTargetActor_GroundTrace : public AGa
 	//TODO: Consider adding the ability to rotate the capsule with a supplied quaternion.
 
 protected:
-	virtual FHitResult PerformTrace(AActor* InSourceActor) const override;
+	virtual FHitResult PerformTrace(AActor* InSourceActor) override;
+
+	virtual bool IsConfirmTargetingAllowed() override;
 
 	bool AdjustCollisionResultForShape(const FVector OriginalStartPoint, const FVector OriginalEndPoint, const FCollisionQueryParams Params, FHitResult& OutHitResult) const;
 
 	FCollisionShape CollisionShape;
 	float CollisionHeightOffset;		//When tracing, give this much extra height to avoid start-in-ground problems. Dealing with thick placement actors while standing near walls may be trickier.
+	bool bLastTraceWasGood;
 };
\ No newline at end of file
diff --git a/Engine/Source/Runtime/GameplayAbilities/Public/Abilities/GameplayAbilityTargetActor_SingleLineTrace.h b/Engine/Source/Runtime/GameplayAbilities/Public/Abilities/GameplayAbilityTargetActor_SingleLineTrace.h
--- a/Engine/Source/Runtime/GameplayAbilities/Public/Abilities/GameplayAbilityTargetActor_SingleLineTrace.h
+++ b/Engine/Source/Runtime/GameplayAbilities/Public/Abilities/GameplayAbilityTargetActor_SingleLineTrace.h
@@ -12,5 +12,5 @@ class GAMEPLAYABILITIES_API AGameplayAbilityTargetActor_SingleLineTrace : public
 	GENERATED_UCLASS_BODY()
 
 protected:
-	virtual FHitResult PerformTrace(AActor* InSourceActor) const override;
+	virtual FHitResult PerformTrace(AActor* InSourceActor) override;
 };
\ No newline at end of file
diff --git a/Engine/Source/Runtime/GameplayAbilities/Public/Abilities/GameplayAbilityTargetActor_Trace.h b/Engine/Source/Runtime/GameplayAbilities/Public/Abilities/GameplayAbilityTargetActor_Trace.h
--- a/Engine/Source/Runtime/GameplayAbilities/Public/Abilities/GameplayAbilityTargetActor_Trace.h
+++ b/Engine/Source/Runtime/GameplayAbilities/Public/Abilities/GameplayAbilityTargetActor_Trace.h
@@ -34,7 +34,7 @@ public:
 	TEnumAsByte<ECollisionChannel> TraceChannel;
 
 protected:
-	virtual FHitResult PerformTrace(AActor* InSourceActor) const PURE_VIRTUAL(AGameplayAbilityTargetActor_Trace, return FHitResult(););
+	virtual FHitResult PerformTrace(AActor* InSourceActor) PURE_VIRTUAL(AGameplayAbilityTargetActor_Trace, return FHitResult(););
 
 	FGameplayAbilityTargetDataHandle MakeTargetData(FHitResult HitResult) const;
 
