Cleaned up names to remove "turret" implication on placed actor targeting.

diff --git a/Engine/Source/Runtime/GameplayAbilities/Public/Abilities/GameplayAbilityTargetActor_ActorPlacement.cpp b/Engine/Source/Runtime/GameplayAbilities/Public/Abilities/GameplayAbilityTargetActor_ActorPlacement.cpp
--- a/Engine/Source/Runtime/GameplayAbilities/Public/Abilities/GameplayAbilityTargetActor_ActorPlacement.cpp
+++ b/Engine/Source/Runtime/GameplayAbilities/Public/Abilities/GameplayAbilityTargetActor_ActorPlacement.cpp
@@ -18,9 +18,9 @@ AGameplayAbilityTargetActor_ActorPlacement::AGameplayAbilityTargetActor_ActorPla
 
 void AGameplayAbilityTargetActor_ActorPlacement::EndPlay(const EEndPlayReason::Type EndPlayReason)
 {
-	if (TurretReticleActor.IsValid())
+	if (ActorVisualizationReticle.IsValid())
 	{
-		TurretReticleActor.Get()->Destroy();
+		ActorVisualizationReticle.Get()->Destroy();
 	}
 
 	Super::EndPlay(EndPlayReason);
@@ -29,20 +29,20 @@ void AGameplayAbilityTargetActor_ActorPlacement::EndPlay(const EEndPlayReason::T
 void AGameplayAbilityTargetActor_ActorPlacement::StartTargeting(UGameplayAbility* InAbility)
 {
 	Super::StartTargeting(InAbility);
-	if (AActor *TurretActor = GetWorld()->SpawnActor(TurretClass))
+	if (AActor *TurretActor = GetWorld()->SpawnActor(PlacedActorClass))
 	{
-		TurretReticleActor = GetWorld()->SpawnActor<AGameplayAbilityWorldReticle_ActorVisualization>();
-		TurretReticleActor->InitializeReticleTurretInformation(TurretActor, TurretMaterial);
+		ActorVisualizationReticle = GetWorld()->SpawnActor<AGameplayAbilityWorldReticle_ActorVisualization>();
+		ActorVisualizationReticle->InitializeReticleTurretInformation(TurretActor, PlacedActorMaterial);
 		GetWorld()->DestroyActor(TurretActor);
 	}
 	if (AGameplayAbilityWorldReticle* CachedReticleActor = ReticleActor.Get())
 	{
-		TurretReticleActor->AttachRootComponentToActor(CachedReticleActor);
+		ActorVisualizationReticle->AttachRootComponentToActor(CachedReticleActor);
 	}
 	else
 	{
-		ReticleActor = TurretReticleActor;
-		TurretReticleActor = nullptr;
+		ReticleActor = ActorVisualizationReticle;
+		ActorVisualizationReticle = nullptr;
 	}
 }
 
diff --git a/Engine/Source/Runtime/GameplayAbilities/Public/Abilities/GameplayAbilityTargetActor_ActorPlacement.h b/Engine/Source/Runtime/GameplayAbilities/Public/Abilities/GameplayAbilityTargetActor_ActorPlacement.h
--- a/Engine/Source/Runtime/GameplayAbilities/Public/Abilities/GameplayAbilityTargetActor_ActorPlacement.h
+++ b/Engine/Source/Runtime/GameplayAbilities/Public/Abilities/GameplayAbilityTargetActor_ActorPlacement.h
@@ -17,16 +17,12 @@ class AGameplayAbilityTargetActor_ActorPlacement : public AGameplayAbilityTarget
 
 	/** Actor we intend to place. */
 	UPROPERTY(BlueprintReadWrite, EditAnywhere, meta = (ExposeOnSpawn = true), Category = Targeting)
-	UClass* TurretClass;		//Using a special class for replication purposes. (Not implemented yet)
+	UClass* PlacedActorClass;		//Using a special class for replication purposes. (Not implemented yet)
 	
-	/** Override Material 0 on our turret's meshes with this material for visualization. */
+	/** Override Material 0 on our placed actor's meshes with this material for visualization. */
 	UPROPERTY(BlueprintReadWrite, EditAnywhere, meta = (ExposeOnSpawn = true), Category = Projectile)
-	UMaterialInterface *TurretMaterial;
+	UMaterialInterface *PlacedActorMaterial;
 
-	/** Visualization for the turret. */
-	TWeakObjectPtr<AGameplayAbilityWorldReticle_ActorVisualization> TurretReticleActor;
-
-protected:
-	//Consider whether we want size-checking here, or in GroundTrace
-	//virtual FHitResult PerformTrace(AActor* InSourceActor) const override;
+	/** Visualization for the intended location of the placed actor. */
+	TWeakObjectPtr<AGameplayAbilityWorldReticle_ActorVisualization> ActorVisualizationReticle;
 };
\ No newline at end of file
