Pass on HideFromModifiers

diff --git a/Engine/Source/Runtime/GameplayAbilities/Public/GameplayEffect.h b/Engine/Source/Runtime/GameplayAbilities/Public/GameplayEffect.h
--- a/Engine/Source/Runtime/GameplayAbilities/Public/GameplayEffect.h
+++ b/Engine/Source/Runtime/GameplayAbilities/Public/GameplayEffect.h
@@ -88,7 +88,7 @@ struct GAMEPLAYABILITIES_API FGameplayModifierInfo
 	TEnumAsByte<EGameplayMod::Type> ModifierType;
 
 	/** The Attribute we modify or the GE we modify modifies. */
-	UPROPERTY(EditDefaultsOnly, Category=GameplayModifier)
+	UPROPERTY(EditDefaultsOnly, Category=GameplayModifier, meta = (FilterMetaTag="HideFromModifiers"))
 	FGameplayAttribute Attribute;
 
 	/** The numeric operation of this modifier: Override, Add, Multiply, etc  */
