check avataractor is valid before activating ability. IsValid() is better than != nullptr since it also checks RF_PendingKill.

diff --git a/Engine/Source/Runtime/GameplayAbilities/Private/AbilitySystemComponent_Abilities.cpp b/Engine/Source/Runtime/GameplayAbilities/Private/AbilitySystemComponent_Abilities.cpp
--- a/Engine/Source/Runtime/GameplayAbilities/Private/AbilitySystemComponent_Abilities.cpp
+++ b/Engine/Source/Runtime/GameplayAbilities/Private/AbilitySystemComponent_Abilities.cpp
@@ -389,7 +389,7 @@ bool UAbilitySystemComponent::TryActivateAbility(FGameplayAbilitySpecHandle Hand
 	const FGameplayAbilityActorInfo* ActorInfo = AbilityActorInfo.Get();
 
 	// make sure the ActorInfo and then Actor on that FGameplayAbilityActorInfo are valid, if not bail out.
-	if (ActorInfo == NULL || ActorInfo->OwnerActor == NULL)
+	if (ActorInfo == NULL || !ActorInfo->OwnerActor.IsValid() || !ActorInfo->AvatarActor.IsValid())
 	{
 		return false;
 	}
