OrionMoveSet.AvatarScale attribute

diff --git a/Engine/Source/Runtime/GameplayAbilities/Private/AttributeSet.cpp b/Engine/Source/Runtime/GameplayAbilities/Private/AttributeSet.cpp
--- a/Engine/Source/Runtime/GameplayAbilities/Private/AttributeSet.cpp
+++ b/Engine/Source/Runtime/GameplayAbilities/Private/AttributeSet.cpp
@@ -77,6 +77,17 @@ UAbilitySystemComponent* UAttributeSet::GetOwningAbilitySystemComponent() const
 	return UAbilitySystemGlobals::GetAbilitySystemComponentFromActor(CastChecked<AActor>(GetOuter()));
 }
 
+FGameplayAbilityActorInfo* UAttributeSet::GetActorInfo() const
+{
+	UAbilitySystemComponent* ASC = GetOwningAbilitySystemComponent();
+	if (ASC)
+	{
+		return ASC->AbilityActorInfo.Get();
+	}
+
+	return nullptr;
+}
+
 void UAttributeSet::PrintDebug()
 {
 	
diff --git a/Engine/Source/Runtime/GameplayAbilities/Public/Abilities/GameplayAbility.h b/Engine/Source/Runtime/GameplayAbilities/Public/Abilities/GameplayAbility.h
--- a/Engine/Source/Runtime/GameplayAbilities/Public/Abilities/GameplayAbility.h
+++ b/Engine/Source/Runtime/GameplayAbilities/Public/Abilities/GameplayAbility.h
@@ -491,6 +491,12 @@ protected:
 	UFUNCTION(BlueprintCallable, Category="Ability|Animation")
 	void MontageStop();
 
+	// ----------------------------------------------------------------------------------------------------------------
+	//
+	//	Target Data
+	//
+	// ----------------------------------------------------------------------------------------------------------------
+
 	UFUNCTION(BlueprintPure, Category = Ability, meta = (HidePin = "WorldContextObject", DefaultToSelf = "WorldContextObject"))
 	FGameplayAbilityTargetingLocationInfo MakeTargetLocationInfoFromOwnerActor();
 
@@ -499,8 +505,7 @@ protected:
 
 	// ----------------------------------------------------------------------------------------------------------------
 	//
-	//	
-	//
+	//	Ability Levels
 	// 
 	// ----------------------------------------------------------------------------------------------------------------
 	
diff --git a/Engine/Source/Runtime/GameplayAbilities/Public/AttributeSet.h b/Engine/Source/Runtime/GameplayAbilities/Public/AttributeSet.h
--- a/Engine/Source/Runtime/GameplayAbilities/Public/AttributeSet.h
+++ b/Engine/Source/Runtime/GameplayAbilities/Public/AttributeSet.h
@@ -5,6 +5,7 @@
 #include "AttributeSet.generated.h"
 
 class UAbilitySystemComponent;
+struct FGameplayAbilityActorInfo;
 
 USTRUCT(BlueprintType)
 struct GAMEPLAYABILITIES_API FGameplayAttribute
@@ -123,6 +124,7 @@ public:
 	void InitFromMetaDataTable(const UDataTable* DataTable);
 
 	UAbilitySystemComponent* GetOwningAbilitySystemComponent() const;
+	FGameplayAbilityActorInfo* GetActorInfo() const;
 
 	virtual void PrintDebug();
 
