TargetDataHasActor blueprint function

diff --git a/Engine/Source/Runtime/GameplayAbilities/Private/AbilitySystemBlueprintLibrary.cpp b/Engine/Source/Runtime/GameplayAbilities/Private/AbilitySystemBlueprintLibrary.cpp
--- a/Engine/Source/Runtime/GameplayAbilities/Private/AbilitySystemBlueprintLibrary.cpp
+++ b/Engine/Source/Runtime/GameplayAbilities/Private/AbilitySystemBlueprintLibrary.cpp
@@ -152,6 +152,19 @@ TArray<AActor*> UAbilitySystemBlueprintLibrary::GetActorsFromTargetData(FGamepla
 	return TArray<AActor*>();
 }
 
+bool UAbilitySystemBlueprintLibrary::TargetDataHasActor(FGameplayAbilityTargetDataHandle TargetData, int32 Index)
+{
+	if (TargetData.Data.IsValidIndex(Index))
+	{
+		FGameplayAbilityTargetData* Data = TargetData.Data[Index].Get();
+		if (Data)
+		{
+			return (Data->GetActors().Num() > 0);
+		}
+	}
+	return false;
+}
+
 bool UAbilitySystemBlueprintLibrary::TargetDataHasHitResult(FGameplayAbilityTargetDataHandle TargetData, int32 Index)
 {
 	if (TargetData.Data.IsValidIndex(Index))
diff --git a/Engine/Source/Runtime/GameplayAbilities/Public/AbilitySystemBlueprintLibrary.h b/Engine/Source/Runtime/GameplayAbilities/Public/AbilitySystemBlueprintLibrary.h
--- a/Engine/Source/Runtime/GameplayAbilities/Public/AbilitySystemBlueprintLibrary.h
+++ b/Engine/Source/Runtime/GameplayAbilities/Public/AbilitySystemBlueprintLibrary.h
@@ -56,6 +56,10 @@ class UAbilitySystemBlueprintLibrary : public UBlueprintFunctionLibrary
 	UFUNCTION(BlueprintPure, Category = "Ability|TargetData")
 	static TArray<AActor*> GetActorsFromTargetData(FGameplayAbilityTargetDataHandle TargetData, int32 Index);
 
+	/** Returns true if the given TargetData has at least 1 actor targeted */
+	UFUNCTION(BlueprintPure, Category = "Ability|TargetData")
+	static bool TargetDataHasActor(FGameplayAbilityTargetDataHandle TargetData, int32 Index);
+
 	UFUNCTION(BlueprintPure, Category = "Ability|TargetData")
 	static bool TargetDataHasHitResult(FGameplayAbilityTargetDataHandle HitResult, int32 Index);
 
