Fixed a bug with single line trace not finding targets correctly due to stopping the line short in pre-trace.

diff --git a/Engine/Source/Runtime/GameplayAbilities/Private/Abilities/GameplayAbilityTargetActor_SingleLineTrace.cpp b/Engine/Source/Runtime/GameplayAbilities/Private/Abilities/GameplayAbilityTargetActor_SingleLineTrace.cpp
--- a/Engine/Source/Runtime/GameplayAbilities/Private/Abilities/GameplayAbilityTargetActor_SingleLineTrace.cpp
+++ b/Engine/Source/Runtime/GameplayAbilities/Private/Abilities/GameplayAbilityTargetActor_SingleLineTrace.cpp
@@ -31,9 +31,19 @@ FHitResult AGameplayAbilityTargetActor_SingleLineTrace::PerformTrace(AActor* InS
 	Params.AddIgnoredActors(ActorsToIgnore);
 
 	FVector TraceStart = StartLocation.GetTargetingTransform().GetLocation();// InSourceActor->GetActorLocation();
-	FVector TraceEnd = TraceStart + (InSourceActor->GetActorForwardVector() * MaxRange);		//Default
+	FVector TraceEnd;
 	AimWithPlayerController(InSourceActor, Params, TraceStart, TraceEnd);		//Effective on server and launching client only
-	FVector AimDirection = (TraceStart - TraceEnd).SafeNormal();
+
+	//Readjust so we have a full-length line going through the predicted target point.
+	FVector AimDirection = (TraceEnd - TraceStart).SafeNormal();
+	if (AimDirection.SizeSquared() > 0.0f)
+	{
+		TraceEnd = TraceStart + (AimDirection * MaxRange);
+	}
+	else
+	{
+		FVector TraceEnd = TraceStart + (InSourceActor->GetActorForwardVector() * MaxRange);		//Default
+	}
 
 	// ------------------------------------------------------
 
diff --git a/Engine/Source/Runtime/GameplayAbilities/Private/Abilities/GameplayAbilityTargetActor_Trace.cpp b/Engine/Source/Runtime/GameplayAbilities/Private/Abilities/GameplayAbilityTargetActor_Trace.cpp
--- a/Engine/Source/Runtime/GameplayAbilities/Private/Abilities/GameplayAbilityTargetActor_Trace.cpp
+++ b/Engine/Source/Runtime/GameplayAbilities/Private/Abilities/GameplayAbilityTargetActor_Trace.cpp
@@ -57,7 +57,6 @@ void AGameplayAbilityTargetActor_Trace::AimWithPlayerController(AActor* InSource
 
 		FHitResult TempHitResult;
 		InSourceActor->GetWorld()->LineTraceSingle(TempHitResult, CamLoc, CamTarget, TraceChannel, Params);
-		//TraceStart = ;
 		if (TempHitResult.bBlockingHit && (FVector::DistSquared(TraceStart, TempHitResult.Location) <= (MaxRange * MaxRange)))
 		{
 			//We actually made a hit? Pull back.
