Might not be needed, but applied the aim direction/distance fix to ground targeting as well for consistency.

diff --git a/Engine/Source/Runtime/GameplayAbilities/Private/Abilities/GameplayAbilityTargetActor_GroundTrace.cpp b/Engine/Source/Runtime/GameplayAbilities/Private/Abilities/GameplayAbilityTargetActor_GroundTrace.cpp
--- a/Engine/Source/Runtime/GameplayAbilities/Private/Abilities/GameplayAbilityTargetActor_GroundTrace.cpp
+++ b/Engine/Source/Runtime/GameplayAbilities/Private/Abilities/GameplayAbilityTargetActor_GroundTrace.cpp
@@ -134,9 +134,19 @@ FHitResult AGameplayAbilityTargetActor_GroundTrace::PerformTrace(AActor* InSourc
 	Params.AddIgnoredActor(InSourceActor);
 
 	FVector TraceStart = StartLocation.GetTargetingTransform().GetLocation();// InSourceActor->GetActorLocation();
-	FVector TraceEnd = TraceStart + (InSourceActor->GetActorForwardVector() * MaxRange);		//Default
+	FVector TraceEnd;
 	AimWithPlayerController(InSourceActor, Params, TraceStart, TraceEnd);		//Effective on server and launching client only
-	FVector AimDirection = (TraceStart - TraceEnd).SafeNormal();
+
+	//Readjust so we have a full-length line going through the predicted target point.
+	FVector AimDirection = (TraceEnd - TraceStart).SafeNormal();
+	if (AimDirection.SizeSquared() > 0.0f)
+	{
+		TraceEnd = TraceStart + (AimDirection * MaxRange);
+	}
+	else
+	{
+		FVector TraceEnd = TraceStart + (InSourceActor->GetActorForwardVector() * MaxRange);		//Default
+	}
 
 	// ------------------------------------------------------
 
