Filtering refactoring added to Gameplay Volume as a spawn parameter. Can also be edited from within the class.

diff --git a/Engine/Source/Runtime/GameplayAbilities/Private/Abilities/GameplayAbilityTargetDataFilter.cpp b/Engine/Source/Runtime/GameplayAbilities/Private/Abilities/GameplayAbilityTargetDataFilter.cpp
--- a/Engine/Source/Runtime/GameplayAbilities/Private/Abilities/GameplayAbilityTargetDataFilter.cpp
+++ b/Engine/Source/Runtime/GameplayAbilities/Private/Abilities/GameplayAbilityTargetDataFilter.cpp
@@ -8,3 +8,8 @@
 //	FGameplayTargetDataFilter
 //
 // --------------------------------------------------------------------------------------------------------------------------------------------------------
+
+void FGameplayTargetDataFilter::InitializeFilterContext(AActor* FilterActor)
+{
+	SelfActor = FilterActor;
+}
diff --git a/Engine/Source/Runtime/GameplayAbilities/Private/AbilitySystemBlueprintLibrary.cpp b/Engine/Source/Runtime/GameplayAbilities/Private/AbilitySystemBlueprintLibrary.cpp
--- a/Engine/Source/Runtime/GameplayAbilities/Private/AbilitySystemBlueprintLibrary.cpp
+++ b/Engine/Source/Runtime/GameplayAbilities/Private/AbilitySystemBlueprintLibrary.cpp
@@ -107,11 +107,11 @@ FGameplayAbilityTargetDataHandle UAbilitySystemBlueprintLibrary::FilterTargetDat
 	return ReturnDataHandle;
 }
 
-FGameplayTargetDataFilterHandle UAbilitySystemBlueprintLibrary::MakeFilterHandle(UGameplayAbility* WorldContextObject, FGameplayTargetDataFilter Filter)
+FGameplayTargetDataFilterHandle UAbilitySystemBlueprintLibrary::MakeFilterHandle(FGameplayTargetDataFilter Filter, AActor* FilterActor)
 {
 	FGameplayTargetDataFilterHandle FilterHandle;
 	FGameplayTargetDataFilter* NewFilter = new FGameplayTargetDataFilter(Filter);
-	NewFilter->SourceAbility = WorldContextObject;
+	NewFilter->InitializeFilterContext(FilterActor);
 	FilterHandle.Filter = TSharedPtr<FGameplayTargetDataFilter>(NewFilter);
 	return FilterHandle;
 }
diff --git a/Engine/Source/Runtime/GameplayAbilities/Public/Abilities/GameplayAbility.h b/Engine/Source/Runtime/GameplayAbilities/Public/Abilities/GameplayAbility.h
--- a/Engine/Source/Runtime/GameplayAbilities/Public/Abilities/GameplayAbility.h
+++ b/Engine/Source/Runtime/GameplayAbilities/Public/Abilities/GameplayAbility.h
@@ -8,6 +8,7 @@
 #include "GameplayAbilityTargetTypes.h"
 #include "GameplayAbilityTypes.h"
 #include "GameplayEffect.h"
+#include "Abilities/GameplayAbilityTargetDataFilter.h"
 #include "GameplayAbility.generated.h"
 
 /**
@@ -404,12 +405,12 @@ protected:
 	UFUNCTION(BlueprintCallable, Category = Ability)
 	AActor* GetOwningActorFromActorInfo() const;
 
-	/** Returns the physical actor that is executing this abilit. May be null */
+	/** Returns the physical actor that is executing this ability. May be null */
 	UFUNCTION(BlueprintCallable, Category = Ability)
 	AActor* GetAvatarActorFromActorInfo() const;
 
 	/** Convenience method for abilities to get skeletal mesh component - useful for aiming abilities */
-	UFUNCTION(BlueprintCallable, Category = Ability)
+	UFUNCTION(BlueprintCallable, FriendlyName = "GetSkeletalMeshComponentFromActorInfo", Category = Ability)
 	USkeletalMeshComponent* GetOwningComponentFromActorInfo() const;
 
 	/** Convenience method for abilities to get outgoing gameplay effect specs (for example, to pass on to projectiles to apply to whoever they hit) */
diff --git a/Engine/Source/Runtime/GameplayAbilities/Public/Abilities/GameplayAbilityTargetDataFilter.h b/Engine/Source/Runtime/GameplayAbilities/Public/Abilities/GameplayAbilityTargetDataFilter.h
--- a/Engine/Source/Runtime/GameplayAbilities/Public/Abilities/GameplayAbilityTargetDataFilter.h
+++ b/Engine/Source/Runtime/GameplayAbilities/Public/Abilities/GameplayAbilityTargetDataFilter.h
@@ -2,7 +2,6 @@
 
 #pragma once
 
-#include "GameplayAbility.h"
 #include "GameplayAbilityTargetDataFilter.generated.h"
 
 UENUM(BlueprintType)
@@ -27,18 +26,16 @@ struct GAMEPLAYABILITIES_API FGameplayTargetDataFilter
 
 	virtual bool FilterPassesForActor(const AActor* ActorToBeFiltered) const
 	{
-		check(SourceAbility.IsValid());
-		check(SourceAbility.Get()->GetAvatarActorFromActorInfo());
 		switch (SelfFilter.GetValue())
 		{
 		case ETargetDataFilterSelf::Type::TDFS_NoOthers:
-			if (ActorToBeFiltered != SourceAbility.Get()->GetAvatarActorFromActorInfo())
+			if (ActorToBeFiltered != SelfActor)
 			{
 				return false;
 			}
 			break;
 		case ETargetDataFilterSelf::Type::TDFS_NoSelf:
-			if (ActorToBeFiltered == SourceAbility.Get()->GetAvatarActorFromActorInfo())
+			if (ActorToBeFiltered == SelfActor)
 			{
 				return false;
 			}
@@ -50,9 +47,13 @@ struct GAMEPLAYABILITIES_API FGameplayTargetDataFilter
 		return true;
 	}
 
+	void InitializeFilterContext(AActor* FilterActor);
+
 	/** Filled out while running */
-	TWeakObjectPtr<UGameplayAbility> SourceAbility;
+	UPROPERTY()
+	AActor* SelfActor;
 
+	/** Our actual filter. */
 	UPROPERTY(BlueprintReadWrite, EditAnywhere, meta = (ExposeOnSpawn = true), Category = Filter)
 	TEnumAsByte<ETargetDataFilterSelf::Type> SelfFilter;
 };
diff --git a/Engine/Source/Runtime/GameplayAbilities/Public/AbilitySystemBlueprintLibrary.h b/Engine/Source/Runtime/GameplayAbilities/Public/AbilitySystemBlueprintLibrary.h
--- a/Engine/Source/Runtime/GameplayAbilities/Public/AbilitySystemBlueprintLibrary.h
+++ b/Engine/Source/Runtime/GameplayAbilities/Public/AbilitySystemBlueprintLibrary.h
@@ -49,9 +49,9 @@ class UAbilitySystemBlueprintLibrary : public UBlueprintFunctionLibrary
 	UFUNCTION(BlueprintPure, Category = "Ability|TargetData")
 	static FGameplayAbilityTargetDataHandle	FilterTargetData(FGameplayAbilityTargetDataHandle TargetDataHandle, FGameplayTargetDataFilterHandle ActorFilterClass);
 
-	/** Create a handle for filtering target data */
-	UFUNCTION(BlueprintPure, meta = (HidePin = "WorldContextObject", DefaultToSelf = "WorldContextObject"), Category = "Filter")
-	static FGameplayTargetDataFilterHandle MakeFilterHandle(UGameplayAbility* WorldContextObject, FGameplayTargetDataFilter Filter);
+	/** Create a handle for filtering target data, filling out all fields */
+	UFUNCTION(BlueprintPure, Category = "Filter")
+	static FGameplayTargetDataFilterHandle MakeFilterHandle(FGameplayTargetDataFilter Filter, AActor* FilterActor);
 
 	UFUNCTION(BlueprintPure, Category = "Ability|TargetData")
 	static TArray<AActor*> GetActorsFromTargetData(FGameplayAbilityTargetDataHandle TargetData, int32 Index);
