#UE4 - Added support for blocking ability activation based on input binding

diff --git a/Engine/Source/Runtime/GameplayAbilities/Private/AbilitySystemComponent.cpp b/Engine/Source/Runtime/GameplayAbilities/Private/AbilitySystemComponent.cpp
--- a/Engine/Source/Runtime/GameplayAbilities/Private/AbilitySystemComponent.cpp
+++ b/Engine/Source/Runtime/GameplayAbilities/Private/AbilitySystemComponent.cpp
@@ -712,6 +712,7 @@ void UAbilitySystemComponent::GetLifetimeReplicatedProps(TArray< FLifetimeProper
 	DOREPLIFETIME(UAbilitySystemComponent, ActiveGameplayCues);
 	
 	DOREPLIFETIME_CONDITION(UAbilitySystemComponent, ActivatableAbilities, COND_OwnerOnly);
+	DOREPLIFETIME_CONDITION(UAbilitySystemComponent, BlockedAbilityBindings, COND_OwnerOnly)
 
 	DOREPLIFETIME(UAbilitySystemComponent, OwnerActor);
 	DOREPLIFETIME(UAbilitySystemComponent, AvatarActor);
diff --git a/Engine/Source/Runtime/GameplayAbilities/Private/AbilitySystemComponent_Abilities.cpp b/Engine/Source/Runtime/GameplayAbilities/Private/AbilitySystemComponent_Abilities.cpp
--- a/Engine/Source/Runtime/GameplayAbilities/Private/AbilitySystemComponent_Abilities.cpp
+++ b/Engine/Source/Runtime/GameplayAbilities/Private/AbilitySystemComponent_Abilities.cpp
@@ -352,6 +352,22 @@ void UAbilitySystemComponent::UnBlockAbilitiesWithTags(const FGameplayTagContain
 	BlockedAbilityTags.UpdateTagMap(Tags, -1);
 }
 
+void UAbilitySystemComponent::BlockAbilityByInputID(int32 InputID)
+{
+	if (InputID >= 0 && InputID < BlockedAbilityBindings.Num())
+	{
+		++BlockedAbilityBindings[InputID];
+	}
+}
+
+void UAbilitySystemComponent::UnBlockAbilityByInputID(int32 InputID)
+{
+	if (InputID >= 0 && InputID < BlockedAbilityBindings.Num() && BlockedAbilityBindings[InputID] > 0)
+	{
+		--BlockedAbilityBindings[InputID];
+	}
+}
+
 #if !(UE_BUILD_SHIPPING || UE_BUILD_TEST)
 int32 DenyClientActivation = 0;
 static FAutoConsoleVariableRef CVarDenyClientActivation(
@@ -406,6 +422,12 @@ bool UAbilitySystemComponent::TryActivateAbility(FGameplayAbilitySpecHandle Hand
 		return false;
 	}
 
+	// Check if this ability's input binding is currently blocked
+	if (Spec->InputID >= 0 && Spec->InputID < BlockedAbilityBindings.Num() && BlockedAbilityBindings[Spec->InputID] > 0)
+	{
+		return false;
+	}
+
 	// Check if any of this ability's tags are currently blocked
 	if (BlockedAbilityTags.HasAnyMatchingGameplayTags(Ability->AbilityTags, EGameplayTagMatchType::IncludeParentTags, false))
 	{
@@ -819,6 +841,7 @@ void UAbilitySystemComponent::BindToInputComponent(UInputComponent* InputCompone
 void UAbilitySystemComponent::BindAbilityActivationToInputComponent(UInputComponent* InputComponent, FGameplayAbiliyInputBinds BindInfo)
 {
 	UEnum* EnumBinds = BindInfo.GetBindEnum();
+	BlockedAbilityBindings.SetNumZeroed(EnumBinds->NumEnums());
 
 	for(int32 idx=0; idx < EnumBinds->NumEnums(); ++idx)
 	{
diff --git a/Engine/Source/Runtime/GameplayAbilities/Public/AbilitySystemComponent.h b/Engine/Source/Runtime/GameplayAbilities/Public/AbilitySystemComponent.h
--- a/Engine/Source/Runtime/GameplayAbilities/Public/AbilitySystemComponent.h
+++ b/Engine/Source/Runtime/GameplayAbilities/Public/AbilitySystemComponent.h
@@ -347,6 +347,9 @@ class GAMEPLAYABILITIES_API UAbilitySystemComponent : public UActorComponent, pu
 
 	void BlockAbilitiesWithTags(const FGameplayTagContainer Tags);
 	void UnBlockAbilitiesWithTags(const FGameplayTagContainer Tags);
+
+	void BlockAbilityByInputID(int32 InputID);
+	void UnBlockAbilityByInputID(int32 InputID);
 	
 	/** FUll list of all instance-per-execution gameplay abilities associated with this component */
 	UPROPERTY()
@@ -612,6 +615,10 @@ private:
 	/** Abilities with these tags are not able to be activated */
 	FGameplayTagCountContainer BlockedAbilityTags;
 
+	/** Tracks abilities that are blocked based on input binding. An ability is blocked if BlockedAbilityBindings[InputID] > 0 */
+	UPROPERTY(Transient, Replicated)
+	TArray<uint8> BlockedAbilityBindings;
+
 	UFUNCTION()
 	void OnRep_GameplayEffects();
 
