Expose GameplyCue parameters to blueprints

diff --git a/Engine/Source/Runtime/GameplayAbilities/Public/GameplayEffectTypes.h b/Engine/Source/Runtime/GameplayAbilities/Public/GameplayEffectTypes.h
--- a/Engine/Source/Runtime/GameplayAbilities/Public/GameplayEffectTypes.h
+++ b/Engine/Source/Runtime/GameplayAbilities/Public/GameplayEffectTypes.h
@@ -698,21 +698,21 @@ struct FGameplayCueParameters
 {
 	GENERATED_USTRUCT_BODY()
 
-	/** Magnitude of source gameplay effect, normalzed from 0-1 */
-	UPROPERTY()
+	/** Magnitude of source gameplay effect, normalzed from 0-1. Use this for "how strong is the gameplay effet" (0=min, 1=,max) */
+	UPROPERTY(BlueprintReadWrite, Category=GameplayCue)
 	float NormalizedMagnitude;
 
-	/** Raw final magnitude of source gameplay effect */
-	UPROPERTY()
+	/** Raw final magnitude of source gameplay effect. Use this is you need to display numbers or for other informational purposes. */
+	UPROPERTY(BlueprintReadWrite, Category=GameplayCue)
 	float RawMagnitude;
 
 	UPROPERTY()
 	FGameplayEffectContextHandle EffectContext;
 
-	UPROPERTY()
+	UPROPERTY(BlueprintReadWrite, Category=GameplayCue)
 	FName MatchedTagName;
 
-	UPROPERTY()
+	UPROPERTY(BlueprintReadWrite, Category=GameplayCue)
 	FGameplayTag OriginalTag;
 };
 
