Don't allow ground target (or actor placement) to accept positions that are on top of other actors.

diff --git a/Engine/Source/Runtime/GameplayAbilities/Private/Abilities/GameplayAbilityTargetActor_GroundTrace.cpp b/Engine/Source/Runtime/GameplayAbilities/Private/Abilities/GameplayAbilityTargetActor_GroundTrace.cpp
--- a/Engine/Source/Runtime/GameplayAbilities/Private/Abilities/GameplayAbilityTargetActor_GroundTrace.cpp
+++ b/Engine/Source/Runtime/GameplayAbilities/Private/Abilities/GameplayAbilityTargetActor_GroundTrace.cpp
@@ -49,7 +49,7 @@ bool AGameplayAbilityTargetActor_GroundTrace::AdjustCollisionResultForShape(cons
 
 	if (bDebug)
 	{
-		if (CollisionHeight > 0.0f)
+		if (CollisionShape.ShapeType == ECollisionShape::Capsule)
 		{
 			DrawDebugCapsule(ThisWorld, OriginalEndPoint, CollisionHeight * 0.5f, CollisionRadius, FQuat::Identity, FColor::Black);
 		}
@@ -77,12 +77,12 @@ bool AGameplayAbilityTargetActor_GroundTrace::AdjustCollisionResultForShape(cons
 		ThisWorld->SweepSingle(LocalResult, TraceStart, TraceEnd, FQuat::Identity, TraceChannel, CollisionShape, Params);
 		if (!LocalResult.bStartPenetrating)
 		{
-			if (!LocalResult.bBlockingHit)
+			if (!LocalResult.bBlockingHit || LocalResult.Actor.IsValid())
 			{
-				//This is probably off the map and should not be considered valid. This should not happen in a non-debug map.
+				//Off the map, or hit an actor
 				if (bDebug)
 				{
-					if (CollisionHeight > 0.0f)
+					if (CollisionShape.ShapeType == ECollisionShape::Capsule)
 					{
 						DrawDebugCapsule(ThisWorld, LocalResult.Location, CollisionHeight * 0.5f, CollisionRadius, FQuat::Identity, FColor::Yellow);
 					}
@@ -92,11 +92,10 @@ bool AGameplayAbilityTargetActor_GroundTrace::AdjustCollisionResultForShape(cons
 					}
 				}
 				continue;
-				//LocalResult.Location = TraceStart;
 			}
 			if (bDebug)
 			{
-				if (CollisionHeight > 0.0f)
+				if (CollisionShape.ShapeType == ECollisionShape::Capsule)
 				{
 					DrawDebugCapsule(ThisWorld, LocalResult.Location, CollisionHeight * 0.5f, CollisionRadius, FQuat::Identity, FColor::Green);
 				}
@@ -105,12 +104,15 @@ bool AGameplayAbilityTargetActor_GroundTrace::AdjustCollisionResultForShape(cons
 					DrawDebugSphere(ThisWorld, LocalResult.Location, CollisionRadius, 8, FColor::Green);
 				}
 			}
+
+			//TODO: Test for flat ground. Concept: Test four corners and the center, make triangles out of the center and adjacent corner points. Check normal.Z of triangles against a minimum Z value.
+
 			OutHitResult = LocalResult;
 			return true;
 		}
 		if (bDebug)
 		{
-			if (CollisionHeight > 0.0f)
+			if (CollisionShape.ShapeType == ECollisionShape::Capsule)
 			{
 				DrawDebugCapsule(ThisWorld, TraceStart, CollisionHeight * 0.5f, CollisionRadius, FQuat::Identity, FColor::Red);
 			}
