Rejected confirmation attempts now include InputID in their notifications.

diff --git a/Engine/Source/Runtime/Engine/Classes/GameFramework/PlayerController.h b/Engine/Source/Runtime/Engine/Classes/GameFramework/PlayerController.h
--- a/Engine/Source/Runtime/Engine/Classes/GameFramework/PlayerController.h
+++ b/Engine/Source/Runtime/Engine/Classes/GameFramework/PlayerController.h
@@ -650,7 +650,7 @@ public:
 
 	/** Informs the client of a rejected attempt to confirm an ability. */
 	UFUNCTION(Unreliable, Client)
-	void ClientNotifyRejectedAbilityConfirmation();
+	void ClientNotifyRejectedAbilityConfirmation(int32 InputID);
 
 	/**
 	 * Outputs a message to HUD
diff --git a/Engine/Source/Runtime/Engine/Private/PlayerController.cpp b/Engine/Source/Runtime/Engine/Private/PlayerController.cpp
--- a/Engine/Source/Runtime/Engine/Private/PlayerController.cpp
+++ b/Engine/Source/Runtime/Engine/Private/PlayerController.cpp
@@ -2741,7 +2741,7 @@ void APlayerController::ClientIgnoreLookInput_Implementation(bool bIgnore)
 	SetIgnoreLookInput(bIgnore);
 }
 
-void APlayerController::ClientNotifyRejectedAbilityConfirmation_Implementation()
+void APlayerController::ClientNotifyRejectedAbilityConfirmation_Implementation(int32 InputID)
 {
 	//TODO: Hook UI notification here. Might be good to include a cooldown on this.
 }
diff --git a/Engine/Source/Runtime/GameplayAbilities/Private/Abilities/GameplayAbilityTargetActor.cpp b/Engine/Source/Runtime/GameplayAbilities/Private/Abilities/GameplayAbilityTargetActor.cpp
--- a/Engine/Source/Runtime/GameplayAbilities/Private/Abilities/GameplayAbilityTargetActor.cpp
+++ b/Engine/Source/Runtime/GameplayAbilities/Private/Abilities/GameplayAbilityTargetActor.cpp
@@ -4,6 +4,7 @@
 #include "GameplayAbilityTargetActor.h"
 #include "AbilitySystemComponent.h"
 #include "Net/UnrealNetwork.h"
+#include "GameplayAbility.h"
 
 // --------------------------------------------------------------------------------------------------------------------------------------------------------
 //
@@ -48,13 +49,7 @@ void AGameplayAbilityTargetActor::ConfirmTargetingAndContinue()
 	}
 	else
 	{
-		if (OwningAbility)
-		{
-			if (APlayerController* PC = OwningAbility->GetActorInfo().PlayerController.Get())
-			{
-				PC->ClientNotifyRejectedAbilityConfirmation();
-			}
-		}
+		NotifyPlayerControllerOfRejectedConfirmation();
 	}
 }
 
@@ -70,11 +65,19 @@ void AGameplayAbilityTargetActor::ConfirmTargeting()
 	}
 	else
 	{
-		if (OwningAbility)
+		NotifyPlayerControllerOfRejectedConfirmation();
+	}
+}
+
+void AGameplayAbilityTargetActor::NotifyPlayerControllerOfRejectedConfirmation()
+{
+	if (OwningAbility && OwningAbility->IsInstantiated())
+	{
+		if (APlayerController* PC = OwningAbility->GetActorInfo().PlayerController.Get())
 		{
-			if (APlayerController* PC = OwningAbility->GetActorInfo().PlayerController.Get())
+			if (FGameplayAbilitySpec* AbilitySpec = OwningAbility->GetCurrentAbilitySpec())
 			{
-				PC->ClientNotifyRejectedAbilityConfirmation();
+				PC->ClientNotifyRejectedAbilityConfirmation(AbilitySpec->InputID);
 			}
 		}
 	}
diff --git a/Engine/Source/Runtime/GameplayAbilities/Private/AbilitySystemComponent_Abilities.cpp b/Engine/Source/Runtime/GameplayAbilities/Private/AbilitySystemComponent_Abilities.cpp
--- a/Engine/Source/Runtime/GameplayAbilities/Private/AbilitySystemComponent_Abilities.cpp
+++ b/Engine/Source/Runtime/GameplayAbilities/Private/AbilitySystemComponent_Abilities.cpp
@@ -999,7 +999,6 @@ void UAbilitySystemComponent::ServerInputRelease_Implementation(FGameplayAbility
 void UAbilitySystemComponent::TargetConfirm()
 {
 	TArray<AGameplayAbilityTargetActor*> LeftoverTargetActors;
-	bool bAnyAbilityDeclined = false;
 	for (AGameplayAbilityTargetActor* TargetActor : SpawnedTargetActors)
 	{
 		if (TargetActor)
@@ -1015,23 +1014,12 @@ void UAbilitySystemComponent::TargetConfirm()
 			}
 			else
 			{
-				bAnyAbilityDeclined = true;
+				TargetActor->NotifyPlayerControllerOfRejectedConfirmation();
 				LeftoverTargetActors.Add(TargetActor);
 			}
 		}
 	}
 	SpawnedTargetActors = LeftoverTargetActors;		//These actors declined to confirm targeting, or are allowed to fire multiple times, so keep contact with them.
-	if (bAnyAbilityDeclined)
-	{
-		//Tell the client's player controller that we rejected at least one attempt to confirm. Mixed success/confirmation is a weird situation that may require further thought.
-		if (FGameplayAbilityActorInfo* ActorInfo = AbilityActorInfo.Get())
-		{
-			if (APlayerController* PC = ActorInfo->PlayerController.Get())
-			{
-				PC->ClientNotifyRejectedAbilityConfirmation();
-			}
-		}
-	}
 }
 
 void UAbilitySystemComponent::TargetCancel()
diff --git a/Engine/Source/Runtime/GameplayAbilities/Public/Abilities/GameplayAbility.h b/Engine/Source/Runtime/GameplayAbilities/Public/Abilities/GameplayAbility.h
--- a/Engine/Source/Runtime/GameplayAbilities/Public/Abilities/GameplayAbility.h
+++ b/Engine/Source/Runtime/GameplayAbilities/Public/Abilities/GameplayAbility.h
@@ -363,11 +363,13 @@ protected:
 	// -------------------------------------
 
 
+	public:
 	bool IsInstantiated() const
 	{
 		return !HasAllFlags(RF_ClassDefaultObject);
 	}
 
+	protected:
 	void SetCurrentActorInfo(const FGameplayAbilitySpecHandle Handle, const FGameplayAbilityActorInfo* ActorInfo) const
 	{
 		if (IsInstantiated())
diff --git a/Engine/Source/Runtime/GameplayAbilities/Public/Abilities/GameplayAbilityTargetActor.h b/Engine/Source/Runtime/GameplayAbilities/Public/Abilities/GameplayAbilityTargetActor.h
--- a/Engine/Source/Runtime/GameplayAbilities/Public/Abilities/GameplayAbilityTargetActor.h
+++ b/Engine/Source/Runtime/GameplayAbilities/Public/Abilities/GameplayAbilityTargetActor.h
@@ -51,6 +51,8 @@ public:
 	UFUNCTION()
 	virtual void ConfirmTargeting();
 
+	void NotifyPlayerControllerOfRejectedConfirmation();
+
 	/** Outside code is saying 'stop everything and just forget about it' */
 	UFUNCTION()
 	virtual void CancelTargeting();
