Cleaned up some code duplication.

diff --git a/Engine/Source/Runtime/GameplayAbilities/Private/Abilities/GameplayAbilityTargetActor_GroundTrace.cpp b/Engine/Source/Runtime/GameplayAbilities/Private/Abilities/GameplayAbilityTargetActor_GroundTrace.cpp
--- a/Engine/Source/Runtime/GameplayAbilities/Private/Abilities/GameplayAbilityTargetActor_GroundTrace.cpp
+++ b/Engine/Source/Runtime/GameplayAbilities/Private/Abilities/GameplayAbilityTargetActor_GroundTrace.cpp
@@ -139,17 +139,6 @@ FHitResult AGameplayAbilityTargetActor_GroundTrace::PerformTrace(AActor* InSourc
 	FVector TraceEnd;
 	AimWithPlayerController(InSourceActor, Params, TraceStart, TraceEnd);		//Effective on server and launching client only
 
-	//Readjust so we have a full-length line going through the predicted target point.
-	FVector AimDirection = (TraceEnd - TraceStart).SafeNormal();
-	if (AimDirection.SizeSquared() > 0.0f)
-	{
-		TraceEnd = TraceStart + (AimDirection * MaxRange);
-	}
-	else
-	{
-		FVector TraceEnd = TraceStart + (InSourceActor->GetActorForwardVector() * MaxRange);		//Default
-	}
-
 	// ------------------------------------------------------
 
 	FHitResult ReturnHitResult;
@@ -164,7 +153,7 @@ FHitResult AGameplayAbilityTargetActor_GroundTrace::PerformTrace(AActor* InSourc
 	}
 
 	//Second trace, straight down. Consider using InSourceActor->GetWorld()->NavigationSystem->ProjectPointToNavigation() instead of tracing in the case of movement abilities (flag/bool).
-	TraceStart = ReturnHitResult.Location - AimDirection;		//Pull back very slightly to avoid scraping down walls
+	TraceStart = ReturnHitResult.Location - (TraceEnd - TraceStart).SafeNormal();		//Pull back very slightly to avoid scraping down walls
 	TraceEnd = TraceStart;
 	TraceStart.Z += CollisionHeightOffset;
 	TraceEnd.Z -= 99999.0f;
diff --git a/Engine/Source/Runtime/GameplayAbilities/Private/Abilities/GameplayAbilityTargetActor_SingleLineTrace.cpp b/Engine/Source/Runtime/GameplayAbilities/Private/Abilities/GameplayAbilityTargetActor_SingleLineTrace.cpp
--- a/Engine/Source/Runtime/GameplayAbilities/Private/Abilities/GameplayAbilityTargetActor_SingleLineTrace.cpp
+++ b/Engine/Source/Runtime/GameplayAbilities/Private/Abilities/GameplayAbilityTargetActor_SingleLineTrace.cpp
@@ -34,17 +34,6 @@ FHitResult AGameplayAbilityTargetActor_SingleLineTrace::PerformTrace(AActor* InS
 	FVector TraceEnd;
 	AimWithPlayerController(InSourceActor, Params, TraceStart, TraceEnd);		//Effective on server and launching client only
 
-	//Readjust so we have a full-length line going through the predicted target point.
-	FVector AimDirection = (TraceEnd - TraceStart).SafeNormal();
-	if (AimDirection.SizeSquared() > 0.0f)
-	{
-		TraceEnd = TraceStart + (AimDirection * MaxRange);
-	}
-	else
-	{
-		FVector TraceEnd = TraceStart + (InSourceActor->GetActorForwardVector() * MaxRange);		//Default
-	}
-
 	// ------------------------------------------------------
 
 	FHitResult ReturnHitResult;
diff --git a/Engine/Source/Runtime/GameplayAbilities/Private/Abilities/GameplayAbilityTargetActor_Trace.cpp b/Engine/Source/Runtime/GameplayAbilities/Private/Abilities/GameplayAbilityTargetActor_Trace.cpp
--- a/Engine/Source/Runtime/GameplayAbilities/Private/Abilities/GameplayAbilityTargetActor_Trace.cpp
+++ b/Engine/Source/Runtime/GameplayAbilities/Private/Abilities/GameplayAbilityTargetActor_Trace.cpp
@@ -67,6 +67,17 @@ void AGameplayAbilityTargetActor_Trace::AimWithPlayerController(AActor* InSource
 			//If we didn't make a hit, use the clipped location.
 			TraceEnd = CamTarget;
 		}
+
+		//Readjust so we have a full-length line going through the predicted target point.
+		FVector AimDirection = (TraceEnd - TraceStart).SafeNormal();
+		if (AimDirection.SizeSquared() > 0.0f)
+		{
+			TraceEnd = TraceStart + (AimDirection * MaxRange);
+		}
+		else
+		{
+			FVector TraceEnd = TraceStart + (InSourceActor->GetActorForwardVector() * MaxRange);		//Default
+		}
 	}
 }
 
