Upgraded target data filter to include actor subclass checking.

diff --git a/Engine/Source/Runtime/GameplayAbilities/Public/Abilities/GameplayAbilityTargetDataFilter.h b/Engine/Source/Runtime/GameplayAbilities/Public/Abilities/GameplayAbilityTargetDataFilter.h
--- a/Engine/Source/Runtime/GameplayAbilities/Public/Abilities/GameplayAbilityTargetDataFilter.h
+++ b/Engine/Source/Runtime/GameplayAbilities/Public/Abilities/GameplayAbilityTargetDataFilter.h
@@ -44,18 +44,29 @@ struct GAMEPLAYABILITIES_API FGameplayTargetDataFilter
 		default:
 			break;
 		}
+
+		if (RequiredActorClass && !ActorToBeFiltered->IsA(RequiredActorClass))
+		{
+			return false;
+		}
+
 		return true;
 	}
 
 	void InitializeFilterContext(AActor* FilterActor);
 
-	/** Filled out while running */
+	/** Actor we're comparing against. */
 	UPROPERTY()
 	AActor* SelfActor;
 
 	/** Our actual filter. */
 	UPROPERTY(BlueprintReadWrite, EditAnywhere, meta = (ExposeOnSpawn = true), Category = Filter)
 	TEnumAsByte<ETargetDataFilterSelf::Type> SelfFilter;
+
+	/** Subclass actors must be to pass the filter. */
+	UPROPERTY(BlueprintReadWrite, EditAnywhere, meta = (ExposeOnSpawn = true), Category = Filter)
+	TSubclassOf<AActor> RequiredActorClass;
+	//UClass* RequiredActorClass;
 };
 
 
@@ -83,6 +94,11 @@ struct GAMEPLAYABILITIES_API FGameplayTargetDataFilterHandle
 		return true;
 	}
 
+	bool FilterPassesForActor(const TWeakObjectPtr<AActor> ActorToBeFiltered) const
+	{
+		return FilterPassesForActor(ActorToBeFiltered.Get());
+	}
+
 	bool operator()(const TWeakObjectPtr<AActor> ActorToBeFiltered) const
 	{
 		return FilterPassesForActor(ActorToBeFiltered.Get());
