Added filtering to "trace" targeting actors so that aimed abilities can selectively ignore actors (e.g. aim through allies).

diff --git a/Engine/Source/Runtime/GameplayAbilities/Private/Abilities/GameplayAbilityTargetActor_GroundTrace.cpp b/Engine/Source/Runtime/GameplayAbilities/Private/Abilities/GameplayAbilityTargetActor_GroundTrace.cpp
--- a/Engine/Source/Runtime/GameplayAbilities/Private/Abilities/GameplayAbilityTargetActor_GroundTrace.cpp
+++ b/Engine/Source/Runtime/GameplayAbilities/Private/Abilities/GameplayAbilityTargetActor_GroundTrace.cpp
@@ -74,7 +74,7 @@ bool AGameplayAbilityTargetActor_GroundTrace::AdjustCollisionResultForShape(cons
 	{
 		TraceEnd = TraceStart = OriginalEndPoint - (LerpValue * Movement);
 		TraceEnd.Z -= 99999.0f;
-		ThisWorld->SweepSingle(LocalResult, TraceStart, TraceEnd, FQuat::Identity, TraceChannel, CollisionShape, Params);
+		SweepWithFilter(LocalResult, ThisWorld, Filter, TraceStart, TraceEnd, FQuat::Identity, TraceChannel, CollisionShape, Params);
 		if (!LocalResult.bStartPenetrating)
 		{
 			if (!LocalResult.bBlockingHit || LocalResult.Actor.IsValid())
@@ -145,19 +145,19 @@ FHitResult AGameplayAbilityTargetActor_GroundTrace::PerformTrace(AActor* InSourc
 	FVector LineTraceStart = TraceStart;
 	FVector LineTraceEnd = TraceEnd;
 	//Use a line trace initially to see where the player is actually pointing
-	InSourceActor->GetWorld()->LineTraceSingle(ReturnHitResult, TraceStart, TraceEnd, TraceChannel, Params);
+	LineTraceWithFilter(ReturnHitResult, InSourceActor->GetWorld(), Filter, TraceStart, TraceEnd, TraceChannel, Params);
 	//Default to end of trace line if we don't hit anything.
 	if (!ReturnHitResult.bBlockingHit)
 	{
 		ReturnHitResult.Location = TraceEnd;
 	}
 
-	//Second trace, straight down. Consider using InSourceActor->GetWorld()->NavigationSystem->ProjectPointToNavigation() instead of tracing in the case of movement abilities (flag/bool).
+	//Second trace, straight down. Consider using InSourceActor->GetWorld()->NavigationSystem->ProjectPointToNavigation() instead of just going straight down in the case of movement abilities (flag/bool).
 	TraceStart = ReturnHitResult.Location - (TraceEnd - TraceStart).SafeNormal();		//Pull back very slightly to avoid scraping down walls
 	TraceEnd = TraceStart;
 	TraceStart.Z += CollisionHeightOffset;
 	TraceEnd.Z -= 99999.0f;
-	InSourceActor->GetWorld()->LineTraceSingle(ReturnHitResult, TraceStart, TraceEnd, TraceChannel, Params);
+	LineTraceWithFilter(ReturnHitResult, InSourceActor->GetWorld(), Filter, TraceStart, TraceEnd, TraceChannel, Params);
 	//if (!ReturnHitResult.bBlockingHit) then our endpoint may be off the map. Hopefully this is only possible in debug maps.
 
 	bLastTraceWasGood = true;		//So far, we're good. If we need a ground spot and can't find one, we'll come back.
diff --git a/Engine/Source/Runtime/GameplayAbilities/Private/Abilities/GameplayAbilityTargetActor_SingleLineTrace.cpp b/Engine/Source/Runtime/GameplayAbilities/Private/Abilities/GameplayAbilityTargetActor_SingleLineTrace.cpp
--- a/Engine/Source/Runtime/GameplayAbilities/Private/Abilities/GameplayAbilityTargetActor_SingleLineTrace.cpp
+++ b/Engine/Source/Runtime/GameplayAbilities/Private/Abilities/GameplayAbilityTargetActor_SingleLineTrace.cpp
@@ -37,7 +37,7 @@ FHitResult AGameplayAbilityTargetActor_SingleLineTrace::PerformTrace(AActor* InS
 	// ------------------------------------------------------
 
 	FHitResult ReturnHitResult;
-	InSourceActor->GetWorld()->LineTraceSingle(ReturnHitResult, TraceStart, TraceEnd, TraceChannel, Params);
+	LineTraceWithFilter(ReturnHitResult, InSourceActor->GetWorld(), Filter, TraceStart, TraceEnd, TraceChannel, Params);
 	//Default to end of trace line if we don't hit anything.
 	if (!ReturnHitResult.bBlockingHit)
 	{
diff --git a/Engine/Source/Runtime/GameplayAbilities/Private/Abilities/GameplayAbilityTargetActor_Trace.cpp b/Engine/Source/Runtime/GameplayAbilities/Private/Abilities/GameplayAbilityTargetActor_Trace.cpp
--- a/Engine/Source/Runtime/GameplayAbilities/Private/Abilities/GameplayAbilityTargetActor_Trace.cpp
+++ b/Engine/Source/Runtime/GameplayAbilities/Private/Abilities/GameplayAbilityTargetActor_Trace.cpp
@@ -33,6 +33,40 @@ void AGameplayAbilityTargetActor_Trace::EndPlay(const EEndPlayReason::Type EndPl
 	Super::EndPlay(EndPlayReason);
 }
 
+void AGameplayAbilityTargetActor_Trace::LineTraceWithFilter(FHitResult& ReturnHitResult, const UWorld* InWorld, const FGameplayTargetDataFilterHandle InFilterHandle, const FVector InTraceStart, const FVector InTraceEnd, ECollisionChannel Channel, const FCollisionQueryParams Params) const
+{
+	check(InWorld);
+	FCollisionQueryParams LocalParams = Params;
+	while (true)
+	{
+		InWorld->LineTraceSingle(ReturnHitResult, InTraceStart, InTraceEnd, Channel, LocalParams);
+		if (ReturnHitResult.Actor.IsValid() && !InFilterHandle.FilterPassesForActor(ReturnHitResult.Actor))
+		{
+			LocalParams.AddIgnoredActor(ReturnHitResult.Actor.Get());
+			continue;
+		}
+		//Either hit something we're not ignoring, or didn't hit anything.
+		return;
+	};
+}
+
+void AGameplayAbilityTargetActor_Trace::SweepWithFilter(FHitResult& ReturnHitResult, const UWorld* InWorld, const FGameplayTargetDataFilterHandle InFilterHandle, const FVector InTraceStart, const FVector InTraceEnd, const FQuat InRotation, ECollisionChannel Channel, const FCollisionShape CollisionShape, const FCollisionQueryParams Params) const
+{
+	check(InWorld);
+	FCollisionQueryParams LocalParams = Params;
+	while (true)
+	{
+		InWorld->SweepSingle(ReturnHitResult, InTraceStart, InTraceEnd, InRotation, Channel, CollisionShape, LocalParams);
+		if (ReturnHitResult.Actor.IsValid() && !InFilterHandle.FilterPassesForActor(ReturnHitResult.Actor))
+		{
+			LocalParams.AddIgnoredActor(ReturnHitResult.Actor.Get());
+			continue;
+		}
+		//Either hit something we're not ignoring, or didn't hit anything.
+		return;
+	};
+}
+
 FGameplayAbilityTargetDataHandle AGameplayAbilityTargetActor_Trace::StaticGetTargetData(UWorld * World, const FGameplayAbilityActorInfo* ActorInfo, FGameplayAbilityActivationInfo ActivationInfo) const
 {
 	check(false);		//This should never actually be called, and if it is, it will require a const version of PerformTrace()
@@ -56,7 +90,7 @@ void AGameplayAbilityTargetActor_Trace::AimWithPlayerController(AActor* InSource
 		ClipCameraRayToAbilityRange(CamLoc, CamDir, TraceStart, MaxRange, CamTarget);
 
 		FHitResult TempHitResult;
-		InSourceActor->GetWorld()->LineTraceSingle(TempHitResult, CamLoc, CamTarget, TraceChannel, Params);
+		LineTraceWithFilter(TempHitResult, InSourceActor->GetWorld(), Filter, CamLoc, CamTarget, TraceChannel, Params);
 		if (TempHitResult.bBlockingHit && (FVector::DistSquared(TraceStart, TempHitResult.Location) <= (MaxRange * MaxRange)))
 		{
 			//We actually made a hit? Pull back.
diff --git a/Engine/Source/Runtime/GameplayAbilities/Public/Abilities/GameplayAbilityTargetActor.h b/Engine/Source/Runtime/GameplayAbilities/Public/Abilities/GameplayAbilityTargetActor.h
--- a/Engine/Source/Runtime/GameplayAbilities/Public/Abilities/GameplayAbilityTargetActor.h
+++ b/Engine/Source/Runtime/GameplayAbilities/Public/Abilities/GameplayAbilityTargetActor.h
@@ -94,6 +94,9 @@ public:
 	UPROPERTY(BlueprintReadWrite, EditAnywhere, meta = (ExposeOnSpawn = true), Category = Targeting)
 	TSubclassOf<AGameplayAbilityWorldReticle> ReticleClass;		//Using a special class for replication purposes.
 
+	UPROPERTY(BlueprintReadWrite, EditAnywhere, Replicated, meta = (ExposeOnSpawn = true), Category = Targeting)
+	FGameplayTargetDataFilterHandle Filter;
+
 	/** Draw the debug information (if applicable) for this targeting actor. */
 	UPROPERTY(BlueprintReadWrite, EditAnywhere, Replicated, meta = (ExposeOnSpawn = true), Category = Targeting)
 	bool bDebug;
diff --git a/Engine/Source/Runtime/GameplayAbilities/Public/Abilities/GameplayAbilityTargetActor_ActorPlacement.cpp b/Engine/Source/Runtime/GameplayAbilities/Public/Abilities/GameplayAbilityTargetActor_ActorPlacement.cpp
--- a/Engine/Source/Runtime/GameplayAbilities/Public/Abilities/GameplayAbilityTargetActor_ActorPlacement.cpp
+++ b/Engine/Source/Runtime/GameplayAbilities/Public/Abilities/GameplayAbilityTargetActor_ActorPlacement.cpp
@@ -29,11 +29,11 @@ void AGameplayAbilityTargetActor_ActorPlacement::EndPlay(const EEndPlayReason::T
 void AGameplayAbilityTargetActor_ActorPlacement::StartTargeting(UGameplayAbility* InAbility)
 {
 	Super::StartTargeting(InAbility);
-	if (AActor *TurretActor = GetWorld()->SpawnActor(PlacedActorClass))
+	if (AActor *VisualizationActor = GetWorld()->SpawnActor(PlacedActorClass))
 	{
 		ActorVisualizationReticle = GetWorld()->SpawnActor<AGameplayAbilityWorldReticle_ActorVisualization>();
-		ActorVisualizationReticle->InitializeReticleTurretInformation(TurretActor, PlacedActorMaterial);
-		GetWorld()->DestroyActor(TurretActor);
+		ActorVisualizationReticle->InitializeReticleVisualizationInformation(VisualizationActor, PlacedActorMaterial);
+		GetWorld()->DestroyActor(VisualizationActor);
 	}
 	if (AGameplayAbilityWorldReticle* CachedReticleActor = ReticleActor.Get())
 	{
diff --git a/Engine/Source/Runtime/GameplayAbilities/Public/Abilities/GameplayAbilityTargetActor_Trace.h b/Engine/Source/Runtime/GameplayAbilities/Public/Abilities/GameplayAbilityTargetActor_Trace.h
--- a/Engine/Source/Runtime/GameplayAbilities/Public/Abilities/GameplayAbilityTargetActor_Trace.h
+++ b/Engine/Source/Runtime/GameplayAbilities/Public/Abilities/GameplayAbilityTargetActor_Trace.h
@@ -15,6 +15,12 @@ class GAMEPLAYABILITIES_API AGameplayAbilityTargetActor_Trace : public AGameplay
 public:
 	virtual void EndPlay(const EEndPlayReason::Type EndPlayReason) override;
 
+	/** Traces as normal, but will manually filter all hit actors */
+	void LineTraceWithFilter(FHitResult& ReturnHitResult, const UWorld* InWorld, const FGameplayTargetDataFilterHandle InFilterHandle, const FVector InTraceStart, const FVector InTraceEnd, ECollisionChannel Channel, const FCollisionQueryParams Params) const;
+
+	/** Sweeps as normal, but will manually filter all hit actors */
+	void SweepWithFilter(FHitResult& ReturnHitResult, const UWorld* InWorld, const FGameplayTargetDataFilterHandle InFilterHandle, const FVector InTraceStart, const FVector InTraceEnd, const FQuat InRotation, ECollisionChannel Channel, const FCollisionShape CollisionShape, const FCollisionQueryParams Params) const;
+
 	virtual FGameplayAbilityTargetDataHandle StaticGetTargetData(UWorld * World, const FGameplayAbilityActorInfo* ActorInfo, FGameplayAbilityActivationInfo ActivationInfo) const override;
 
 	void AimWithPlayerController(AActor* InSourceActor, FCollisionQueryParams Params, FVector TraceStart, FVector& TraceEnd) const;
diff --git a/Engine/Source/Runtime/GameplayAbilities/Public/Abilities/GameplayAbilityWorldReticle_ActorVisualization.cpp b/Engine/Source/Runtime/GameplayAbilities/Public/Abilities/GameplayAbilityWorldReticle_ActorVisualization.cpp
--- a/Engine/Source/Runtime/GameplayAbilities/Public/Abilities/GameplayAbilityWorldReticle_ActorVisualization.cpp
+++ b/Engine/Source/Runtime/GameplayAbilities/Public/Abilities/GameplayAbilityWorldReticle_ActorVisualization.cpp
@@ -29,34 +29,34 @@ AGameplayAbilityWorldReticle_ActorVisualization::AGameplayAbilityWorldReticle_Ac
 	RootComponent = CollisionComponent;
 }
 
-void AGameplayAbilityWorldReticle_ActorVisualization::InitializeReticleTurretInformation(AActor* TurretActor, UMaterialInterface *TurretMaterial)
+void AGameplayAbilityWorldReticle_ActorVisualization::InitializeReticleVisualizationInformation(AActor* VisualizationActor, UMaterialInterface *VisualizationMaterial)
 {
-	if (TurretActor)
+	if (VisualizationActor)
 	{
 		//Get components
 		TArray<UMeshComponent*> MeshComps;
 		USceneComponent* MyRoot = GetRootComponent();
-		TurretActor->GetComponents(MeshComps);
+		VisualizationActor->GetComponents(MeshComps);
 		check(MyRoot);
 
 		for (UMeshComponent* MeshComp : MeshComps)
 		{
-			//Special case: If we don't clear the root component explicitly, the component will be destroyed along with the original turret actor.
-			if (MeshComp == TurretActor->GetRootComponent())
+			//Special case: If we don't clear the root component explicitly, the component will be destroyed along with the original visualization actor.
+			if (MeshComp == VisualizationActor->GetRootComponent())
 			{
-				TurretActor->SetRootComponent(NULL);
+				VisualizationActor->SetRootComponent(NULL);
 			}
 
-			//Disable collision on turret mesh parts so it doesn't interfere with aiming or any other client-side collision/prediction/physics stuff
+			//Disable collision on visualization mesh parts so it doesn't interfere with aiming or any other client-side collision/prediction/physics stuff
 			MeshComp->SetCollisionEnabled(ECollisionEnabled::NoCollision);		//All mesh components are primitive components, so no cast is needed
 
 			//Move components from one actor to the other, attaching as needed. Hierarchy should not be important, but we can do fixups if it becomes important later.
 			MeshComp->DetachFromParent();
 			MeshComp->AttachTo(MyRoot);
 			MeshComp->Rename(nullptr, this);
-			if (TurretMaterial)
+			if (VisualizationMaterial)
 			{
-				MeshComp->SetMaterial(0, TurretMaterial);
+				MeshComp->SetMaterial(0, VisualizationMaterial);
 			}
 		}
 	}
diff --git a/Engine/Source/Runtime/GameplayAbilities/Public/Abilities/GameplayAbilityWorldReticle_ActorVisualization.h b/Engine/Source/Runtime/GameplayAbilities/Public/Abilities/GameplayAbilityWorldReticle_ActorVisualization.h
--- a/Engine/Source/Runtime/GameplayAbilities/Public/Abilities/GameplayAbilityWorldReticle_ActorVisualization.h
+++ b/Engine/Source/Runtime/GameplayAbilities/Public/Abilities/GameplayAbilityWorldReticle_ActorVisualization.h
@@ -7,7 +7,7 @@
 
 class AGameplayAbilityTargetActor;
 
-/** This is a dummy reticle for internal use by turret placement tasks. It builds a custom visual model of the turret being placed. */
+/** This is a dummy reticle for internal use by visualization placement tasks. It builds a custom visual model of the visualization being placed. */
 UCLASS(notplaceable)
 class AGameplayAbilityWorldReticle_ActorVisualization : public AGameplayAbilityWorldReticle
 {
@@ -15,14 +15,14 @@ class AGameplayAbilityWorldReticle_ActorVisualization : public AGameplayAbilityW
 
 public:
 
-	void InitializeReticleTurretInformation(AActor* TurretActor, UMaterialInterface *TurretMaterial);
+	void InitializeReticleVisualizationInformation(AActor* VisualizationActor, UMaterialInterface *VisualizationMaterial);
 
-	/** Hardcoded collision component, so other objects don't think they can collide with the visualization turret */
+	/** Hardcoded collision component, so other objects don't think they can collide with the visualization actor */
 	UPROPERTY()
 	TSubobjectPtr<class UCapsuleComponent> CollisionComponent;
 
 	UPROPERTY()
-	TArray<UActorComponent*> TurretComponents;
+	TArray<UActorComponent*> VisualizationComponents;
 
 	/** Overridable function called whenever this actor is being removed from a level */
 	virtual void EndPlay(const EEndPlayReason::Type EndPlayReason);
