GUBP fix.

diff --git a/Engine/Source/Runtime/GameplayAbilities/Private/Abilities/GameplayAbilityTargetActor_Trace.cpp b/Engine/Source/Runtime/GameplayAbilities/Private/Abilities/GameplayAbilityTargetActor_Trace.cpp
--- a/Engine/Source/Runtime/GameplayAbilities/Private/Abilities/GameplayAbilityTargetActor_Trace.cpp
+++ b/Engine/Source/Runtime/GameplayAbilities/Private/Abilities/GameplayAbilityTargetActor_Trace.cpp
@@ -33,7 +33,7 @@ void AGameplayAbilityTargetActor_Trace::EndPlay(const EEndPlayReason::Type EndPl
 	Super::EndPlay(EndPlayReason);
 }
 
-void AGameplayAbilityTargetActor_Trace::LineTraceWithFilter(FHitResult& ReturnHitResult, const UWorld* InWorld, const FGameplayTargetDataFilterHandle InFilterHandle, const FVector InTraceStart, const FVector InTraceEnd, ECollisionChannel Channel, const FCollisionQueryParams Params) const
+void AGameplayAbilityTargetActor_Trace::LineTraceWithFilter(FHitResult& ReturnHitResult, const UWorld* InWorld, const FGameplayTargetDataFilterHandle InFilterHandle, const FVector& InTraceStart, const FVector& InTraceEnd, ECollisionChannel Channel, const FCollisionQueryParams Params) const
 {
 	check(InWorld);
 	FCollisionQueryParams LocalParams = Params;
@@ -50,7 +50,7 @@ void AGameplayAbilityTargetActor_Trace::LineTraceWithFilter(FHitResult& ReturnHi
 	};
 }
 
-void AGameplayAbilityTargetActor_Trace::SweepWithFilter(FHitResult& ReturnHitResult, const UWorld* InWorld, const FGameplayTargetDataFilterHandle InFilterHandle, const FVector InTraceStart, const FVector InTraceEnd, const FQuat InRotation, ECollisionChannel Channel, const FCollisionShape CollisionShape, const FCollisionQueryParams Params) const
+void AGameplayAbilityTargetActor_Trace::SweepWithFilter(FHitResult& ReturnHitResult, const UWorld* InWorld, const FGameplayTargetDataFilterHandle InFilterHandle, const FVector& InTraceStart, const FVector& InTraceEnd, const FQuat& InRotation, ECollisionChannel Channel, const FCollisionShape CollisionShape, const FCollisionQueryParams Params) const
 {
 	check(InWorld);
 	FCollisionQueryParams LocalParams = Params;
diff --git a/Engine/Source/Runtime/GameplayAbilities/Public/Abilities/GameplayAbilityTargetActor_Trace.h b/Engine/Source/Runtime/GameplayAbilities/Public/Abilities/GameplayAbilityTargetActor_Trace.h
--- a/Engine/Source/Runtime/GameplayAbilities/Public/Abilities/GameplayAbilityTargetActor_Trace.h
+++ b/Engine/Source/Runtime/GameplayAbilities/Public/Abilities/GameplayAbilityTargetActor_Trace.h
@@ -16,10 +16,10 @@ public:
 	virtual void EndPlay(const EEndPlayReason::Type EndPlayReason) override;
 
 	/** Traces as normal, but will manually filter all hit actors */
-	void LineTraceWithFilter(FHitResult& ReturnHitResult, const UWorld* InWorld, const FGameplayTargetDataFilterHandle InFilterHandle, const FVector InTraceStart, const FVector InTraceEnd, ECollisionChannel Channel, const FCollisionQueryParams Params) const;
+	void LineTraceWithFilter(FHitResult& ReturnHitResult, const UWorld* InWorld, const FGameplayTargetDataFilterHandle InFilterHandle, const FVector& InTraceStart, const FVector& InTraceEnd, ECollisionChannel Channel, const FCollisionQueryParams Params) const;
 
 	/** Sweeps as normal, but will manually filter all hit actors */
-	void SweepWithFilter(FHitResult& ReturnHitResult, const UWorld* InWorld, const FGameplayTargetDataFilterHandle InFilterHandle, const FVector InTraceStart, const FVector InTraceEnd, const FQuat InRotation, ECollisionChannel Channel, const FCollisionShape CollisionShape, const FCollisionQueryParams Params) const;
+	void SweepWithFilter(FHitResult& ReturnHitResult, const UWorld* InWorld, const FGameplayTargetDataFilterHandle InFilterHandle, const FVector& InTraceStart, const FVector& InTraceEnd, const FQuat& InRotation, ECollisionChannel Channel, const FCollisionShape CollisionShape, const FCollisionQueryParams Params) const;
 
 	virtual FGameplayAbilityTargetDataHandle StaticGetTargetData(UWorld * World, const FGameplayAbilityActorInfo* ActorInfo, FGameplayAbilityActivationInfo ActivationInfo) const override;
 
