Target data filters are reversible.

diff --git a/Engine/Source/Runtime/GameplayAbilities/Public/Abilities/GameplayAbilityTargetDataFilter.h b/Engine/Source/Runtime/GameplayAbilities/Public/Abilities/GameplayAbilityTargetDataFilter.h
--- a/Engine/Source/Runtime/GameplayAbilities/Public/Abilities/GameplayAbilityTargetDataFilter.h
+++ b/Engine/Source/Runtime/GameplayAbilities/Public/Abilities/GameplayAbilityTargetDataFilter.h
@@ -31,13 +31,13 @@ struct GAMEPLAYABILITIES_API FGameplayTargetDataFilter
 		case ETargetDataFilterSelf::Type::TDFS_NoOthers:
 			if (ActorToBeFiltered != SelfActor)
 			{
-				return false;
+				return (bReverseFilter ^ false);
 			}
 			break;
 		case ETargetDataFilterSelf::Type::TDFS_NoSelf:
 			if (ActorToBeFiltered == SelfActor)
 			{
-				return false;
+				return (bReverseFilter ^ false);
 			}
 			break;
 		case ETargetDataFilterSelf::Type::TDFS_Any:
@@ -47,10 +47,10 @@ struct GAMEPLAYABILITIES_API FGameplayTargetDataFilter
 
 		if (RequiredActorClass && !ActorToBeFiltered->IsA(RequiredActorClass))
 		{
-			return false;
+			return (bReverseFilter ^ false);
 		}
 
-		return true;
+		return (bReverseFilter ^ true);
 	}
 
 	void InitializeFilterContext(AActor* FilterActor);
@@ -59,14 +59,17 @@ struct GAMEPLAYABILITIES_API FGameplayTargetDataFilter
 	UPROPERTY()
 	AActor* SelfActor;
 
-	/** Our actual filter. */
+	/** Filter based on whether or not this actor is "self." */
 	UPROPERTY(BlueprintReadWrite, EditAnywhere, meta = (ExposeOnSpawn = true), Category = Filter)
 	TEnumAsByte<ETargetDataFilterSelf::Type> SelfFilter;
 
 	/** Subclass actors must be to pass the filter. */
 	UPROPERTY(BlueprintReadWrite, EditAnywhere, meta = (ExposeOnSpawn = true), Category = Filter)
 	TSubclassOf<AActor> RequiredActorClass;
-	//UClass* RequiredActorClass;
+
+	/** Reverses the meaning of the filter, so it will exclude all actors that pass. */
+	UPROPERTY(BlueprintReadWrite, EditAnywhere, meta = (ExposeOnSpawn = true), Category = Filter)
+	bool bReverseFilter;
 };
 
 
