Moving some CPP files to the proper location.

diff --git a/Engine/Source/Runtime/GameplayAbilities/Private/Abilities/GameplayAbilityTargetActor_ActorPlacement.cpp b/Engine/Source/Runtime/GameplayAbilities/Private/Abilities/GameplayAbilityTargetActor_ActorPlacement.cpp
new file mode 100644
--- /dev/null
+++ b/Engine/Source/Runtime/GameplayAbilities/Private/Abilities/GameplayAbilityTargetActor_ActorPlacement.cpp
@@ -0,0 +1,50 @@
+// Copyright 1998-2014 Epic Games, Inc. All Rights Reserved.
+
+#include "AbilitySystemPrivatePCH.h"
+#include "AbilitySystemComponent.h"
+#include "GameplayAbilityTargetActor_ActorPlacement.h"
+#include "GameplayAbilityWorldReticle_ActorVisualization.h"
+
+// --------------------------------------------------------------------------------------------------------------------------------------------------------
+//
+//	AGameplayAbilityTargetActor_ActorPlacement
+//
+// --------------------------------------------------------------------------------------------------------------------------------------------------------
+
+AGameplayAbilityTargetActor_ActorPlacement::AGameplayAbilityTargetActor_ActorPlacement(const FObjectInitializer& ObjectInitializer)
+	: Super(ObjectInitializer)
+{
+}
+
+void AGameplayAbilityTargetActor_ActorPlacement::EndPlay(const EEndPlayReason::Type EndPlayReason)
+{
+	if (ActorVisualizationReticle.IsValid())
+	{
+		ActorVisualizationReticle.Get()->Destroy();
+	}
+
+	Super::EndPlay(EndPlayReason);
+}
+
+void AGameplayAbilityTargetActor_ActorPlacement::StartTargeting(UGameplayAbility* InAbility)
+{
+	Super::StartTargeting(InAbility);
+	if (AActor *VisualizationActor = GetWorld()->SpawnActor(PlacedActorClass))
+	{
+		ActorVisualizationReticle = GetWorld()->SpawnActor<AGameplayAbilityWorldReticle_ActorVisualization>();
+		ActorVisualizationReticle->InitializeReticleVisualizationInformation(VisualizationActor, PlacedActorMaterial);
+		GetWorld()->DestroyActor(VisualizationActor);
+	}
+	if (AGameplayAbilityWorldReticle* CachedReticleActor = ReticleActor.Get())
+	{
+		ActorVisualizationReticle->AttachRootComponentToActor(CachedReticleActor);
+	}
+	else
+	{
+		ReticleActor = ActorVisualizationReticle;
+		ActorVisualizationReticle = nullptr;
+	}
+}
+
+//Might want to override this function to allow for a radius check against the ground, possibly including a height check. Or might want to do it in ground trace.
+//FHitResult AGameplayAbilityTargetActor_ActorPlacement::PerformTrace(AActor* InSourceActor) const
diff --git a/Engine/Source/Runtime/GameplayAbilities/Private/Abilities/GameplayAbilityWorldReticle_ActorVisualization.cpp b/Engine/Source/Runtime/GameplayAbilities/Private/Abilities/GameplayAbilityWorldReticle_ActorVisualization.cpp
new file mode 100644
--- /dev/null
+++ b/Engine/Source/Runtime/GameplayAbilities/Private/Abilities/GameplayAbilityWorldReticle_ActorVisualization.cpp
@@ -0,0 +1,72 @@
+// Copyright 1998-2013 Epic Games, Inc. All Rights Reserved.
+
+#include "AbilitySystemPrivatePCH.h"
+#include "GameplayAbilityTargetActor.h"
+#include "GameplayAbilityWorldReticle_ActorVisualization.h"
+
+// --------------------------------------------------------------------------------------------------------------------------------------------------------
+//
+//	AGameplayAbilityWorldReticle_ActorVisualization
+//
+// --------------------------------------------------------------------------------------------------------------------------------------------------------
+
+AGameplayAbilityWorldReticle_ActorVisualization::AGameplayAbilityWorldReticle_ActorVisualization(const FObjectInitializer& ObjectInitializer)
+	: Super(ObjectInitializer)
+{
+	//PrimaryActorTick.bCanEverTick = true;
+	//PrimaryActorTick.TickGroup = TG_PrePhysics;
+
+	CollisionComponent = ObjectInitializer.CreateDefaultSubobject<UCapsuleComponent>(this, TEXT("CollisionCapsule0"));
+	CollisionComponent->InitCapsuleSize(0.f, 0.f);
+	CollisionComponent->AlwaysLoadOnClient = true;
+	CollisionComponent->bAbsoluteScale = true;
+	//CollisionComponent->AlwaysLoadOnServer = true;
+	CollisionComponent->bCanEverAffectNavigation = false;
+	CollisionComponent->BodyInstance.bEnableCollision_DEPRECATED = false;
+	CollisionComponent->SetCollisionEnabled(ECollisionEnabled::NoCollision);
+
+	//TSubobjectPtr<USceneComponent> SceneComponent = ObjectInitializer.CreateDefaultSubobject<USceneComponent>(this, TEXT("RootComponent0"));
+	RootComponent = CollisionComponent;
+}
+
+void AGameplayAbilityWorldReticle_ActorVisualization::InitializeReticleVisualizationInformation(AActor* VisualizationActor, UMaterialInterface *VisualizationMaterial)
+{
+	if (VisualizationActor)
+	{
+		//Get components
+		TArray<UMeshComponent*> MeshComps;
+		USceneComponent* MyRoot = GetRootComponent();
+		VisualizationActor->GetComponents(MeshComps);
+		check(MyRoot);
+
+		for (UMeshComponent* MeshComp : MeshComps)
+		{
+			//Special case: If we don't clear the root component explicitly, the component will be destroyed along with the original visualization actor.
+			if (MeshComp == VisualizationActor->GetRootComponent())
+			{
+				VisualizationActor->SetRootComponent(NULL);
+			}
+
+			//Disable collision on visualization mesh parts so it doesn't interfere with aiming or any other client-side collision/prediction/physics stuff
+			MeshComp->SetCollisionEnabled(ECollisionEnabled::NoCollision);		//All mesh components are primitive components, so no cast is needed
+
+			//Move components from one actor to the other, attaching as needed. Hierarchy should not be important, but we can do fixups if it becomes important later.
+			MeshComp->DetachFromParent();
+			MeshComp->AttachTo(MyRoot);
+			MeshComp->Rename(nullptr, this);
+			if (VisualizationMaterial)
+			{
+				MeshComp->SetMaterial(0, VisualizationMaterial);
+			}
+		}
+	}
+}
+
+void AGameplayAbilityWorldReticle_ActorVisualization::EndPlay(const EEndPlayReason::Type EndPlayReason)
+{
+	TArray<UActorComponent*> CompArray;
+	GetComponents(CompArray);
+	CompArray.Num();
+
+	Super::EndPlay(EndPlayReason);
+}
diff --git a/Engine/Source/Runtime/GameplayAbilities/Public/Abilities/GameplayAbilityTargetActor_ActorPlacement.cpp b/Engine/Source/Runtime/GameplayAbilities/Public/Abilities/GameplayAbilityTargetActor_ActorPlacement.cpp
deleted file mode 100644
diff --git a/Engine/Source/Runtime/GameplayAbilities/Public/Abilities/GameplayAbilityWorldReticle_ActorVisualization.cpp b/Engine/Source/Runtime/GameplayAbilities/Public/Abilities/GameplayAbilityWorldReticle_ActorVisualization.cpp
deleted file mode 100644
