Bugfix to weird ground placement issues. Still some rough corners to sand down, but it works.

diff --git a/Engine/Source/Runtime/GameplayAbilities/Private/Abilities/GameplayAbilityTargetActor_GroundTrace.cpp b/Engine/Source/Runtime/GameplayAbilities/Private/Abilities/GameplayAbilityTargetActor_GroundTrace.cpp
--- a/Engine/Source/Runtime/GameplayAbilities/Private/Abilities/GameplayAbilityTargetActor_GroundTrace.cpp
+++ b/Engine/Source/Runtime/GameplayAbilities/Private/Abilities/GameplayAbilityTargetActor_GroundTrace.cpp
@@ -36,6 +36,11 @@ void AGameplayAbilityTargetActor_GroundTrace::StartTargeting(UGameplayAbility* I
 			CollisionHeightOffset = CollisionRadius;
 		}
 	}
+	else
+	{
+		//Make sure these are clean.
+		CollisionRadius = CollisionHeight = 0.0f;
+	}
 	Super::StartTargeting(InAbility);
 }
 
@@ -59,17 +64,19 @@ bool AGameplayAbilityTargetActor_GroundTrace::AdjustCollisionResultForShape(cons
 		}
 	}
 
-	if ((MovementMagnitude2D < (CollisionRadius * 2.0f)) || (CollisionRadius <= 1.0f))
+	if (MovementMagnitude2D <= (CollisionRadius * 2.0f))
 	{
 		return false;		//Bad case!
 	}
 
-	float IncrementSize = FMath::Clamp<float>(CollisionRadius * 0.5f, 25.0f, 250.0f);
+	//TODO This increment value needs to ramp up - the first few increments should be small, then we should start moving in larger steps. A few ideas for this:
+	//1. Use a curve! Even one defined by a hardcoded formula would be fine, this isn't something that should require user tuning, or that the user should really know/care about.
+	//2. Use larger increments as the object is further from the player/camera, since the user can't really perceive precision at long range.
+	float IncrementSize = FMath::Clamp<float>(CollisionRadius * 0.5f, 20.0f, 50.0f);
 	float LerpIncrement = IncrementSize / MovementMagnitude2D;
 	FHitResult LocalResult;
 	FVector TraceStart;
 	FVector TraceEnd;
-	//This needs to ramp up - the first few increments should be small, then we should start moving in larger steps.
 	for (float LerpValue = CollisionRadius / MovementMagnitude2D; LerpValue < 1.0f; LerpValue += LerpIncrement)
 	{
 		TraceEnd = TraceStart = OriginalEndPoint - (LerpValue * Movement);
