Restore non modifier checking for duration mods. We cant check modifiers when there are no modifiers (e.g., cooldown)

diff --git a/Engine/Source/Runtime/GameplayAbilities/Private/GameplayEffect.cpp b/Engine/Source/Runtime/GameplayAbilities/Private/GameplayEffect.cpp
--- a/Engine/Source/Runtime/GameplayAbilities/Private/GameplayEffect.cpp
+++ b/Engine/Source/Runtime/GameplayAbilities/Private/GameplayEffect.cpp
@@ -260,48 +260,62 @@ int32 FGameplayEffectSpec::ApplyModifier(const FModifierSpec &InMod, const FModi
 		return 0;
 	}
 
-	for (FModifierSpec &MyMod : Modifiers)
+	switch (InMod.Info.EffectType)
 	{
-		if (InMod.CanModifyModifier(MyMod, QualifierContext))
+		case EGameplayModEffect::Magnitude:
 		{
-			bool FinishedLoop = true;
-			NumApplied = 1;
-
-			switch (InMod.Info.EffectType)
+			for (FModifierSpec &MyMod : Modifiers)
 			{
-				case EGameplayModEffect::Magnitude:
+				if (InMod.CanModifyModifier(MyMod, QualifierContext))
+				{
 					InMod.ApplyModTo(MyMod, bApplyAsSnapshot);
-					FinishedLoop = false;	// Keeep going: we need to go through the entire mod list
-					break;
+					NumApplied = 1;
+				}
+			}
+			break;
+		}
 
-				case EGameplayModEffect::Duration:
-					// don't modify infinite duration or instant unless we're overriding it
-					if (GetDuration() > 0.f || InMod.Info.ModifierOp == EGameplayModOp::Override)
-					{
-						Duration.Get()->ApplyMod(InMod.Info.ModifierOp, InMod.Aggregator, bApplyAsSnapshot);
-					}
-					else
-					{
-						NumApplied = 0;
-					}
-					break;
-				case EGameplayModEffect::ChanceApplyTarget:
-					ChanceToApplyToTarget.Get()->ApplyMod(InMod.Info.ModifierOp, InMod.Aggregator, bApplyAsSnapshot);
-					break;
-				case EGameplayModEffect::ChanceExecuteEffect:
-					ChanceToExecuteOnGameplayEffect.Get()->ApplyMod(InMod.Info.ModifierOp, InMod.Aggregator, bApplyAsSnapshot);
-					break;
-
-				case EGameplayModEffect::LinkedGameplayEffect:
-					TargetEffectSpecs.Add(InMod.TargetEffectSpec.ToSharedRef());
-					break;
+		/**
+		 *	We do not check for Modifier qualification in the below cases, since these modifiers are independant of other modifiers!
+		 *	We can modify the duration of a GameplayEffect which only gives a tag but does not modify any attributes. (E.g., a cooldown)
+		 *	The downside is that you must tag your GameplayEffects rigorously in order to use Incoming/Outgoing/Active modififcation
+		 */
+
+		case EGameplayModEffect::Duration:
+		{
+			// don't modify infinite duration or instant unless we're overriding it
+			if (GetDuration() > 0.f || InMod.Info.ModifierOp == EGameplayModOp::Override)
+			{
+				Duration.Get()->ApplyMod(InMod.Info.ModifierOp, InMod.Aggregator, bApplyAsSnapshot);
+				NumApplied = 1;
 			}
 
-			if (FinishedLoop)
-				break;
+			break;
+		}
+
+		case EGameplayModEffect::ChanceApplyTarget:
+		{
+			ChanceToApplyToTarget.Get()->ApplyMod(InMod.Info.ModifierOp, InMod.Aggregator, bApplyAsSnapshot);
+			NumApplied = 1;
+			break;
+		}
+
+		case EGameplayModEffect::ChanceExecuteEffect:
+		{
+			ChanceToExecuteOnGameplayEffect.Get()->ApplyMod(InMod.Info.ModifierOp, InMod.Aggregator, bApplyAsSnapshot);
+			NumApplied = 1;
+			break;
+		}
+
+		case EGameplayModEffect::LinkedGameplayEffect:
+		{
+			TargetEffectSpecs.Add(InMod.TargetEffectSpec.ToSharedRef());
+			NumApplied = 1;
+			break;
 		}
 	}
 
+
 	return NumApplied;
 }
 
@@ -2233,3 +2247,4 @@ bool FActiveGameplayEffectQuery::Matches(const FActiveGameplayEffect& Effect) co
 
 	return true;
 }
+
