Fix for Remove GameplayCue event not being handled when done outside of GameplayEffects

diff --git a/Engine/Source/Runtime/GameplayAbilities/Private/AbilitySystemComponent.cpp b/Engine/Source/Runtime/GameplayAbilities/Private/AbilitySystemComponent.cpp
--- a/Engine/Source/Runtime/GameplayAbilities/Private/AbilitySystemComponent.cpp
+++ b/Engine/Source/Runtime/GameplayAbilities/Private/AbilitySystemComponent.cpp
@@ -570,6 +570,7 @@ void UAbilitySystemComponent::RemoveGameplayCue(const FGameplayTag GameplayCueTa
 	if (IsOwnerActorAuthoritative())
 	{
 		ActiveGameplayCues.RemoveCue(GameplayCueTag);
+		NetMulticast_InvokeGameplayCueRemoved(GameplayCueTag, PredictionKey);
 	}
 	else if (PredictionKey.IsValidKey())
 	{
@@ -619,6 +620,14 @@ void UAbilitySystemComponent::NetMulticast_InvokeGameplayCueAdded_Implementation
 	}
 }
 
+void UAbilitySystemComponent::NetMulticast_InvokeGameplayCueRemoved_Implementation(const FGameplayTag GameplayCueTag, FPredictionKey PredictionKey)
+{
+	if (IsOwnerActorAuthoritative() || PredictionKey.IsValidKey() == false)
+	{
+		InvokeGameplayCueEvent(GameplayCueTag, EGameplayCueEvent::Removed);
+	}
+}
+
 bool UAbilitySystemComponent::IsGameplayCueActive(const FGameplayTag GameplayCueTag) const
 {
 	return (ActiveGameplayEffects.HasMatchingGameplayTag(GameplayCueTag) || ActiveGameplayCues.HasMatchingGameplayTag(GameplayCueTag));
diff --git a/Engine/Source/Runtime/GameplayAbilities/Public/AbilitySystemComponent.h b/Engine/Source/Runtime/GameplayAbilities/Public/AbilitySystemComponent.h
--- a/Engine/Source/Runtime/GameplayAbilities/Public/AbilitySystemComponent.h
+++ b/Engine/Source/Runtime/GameplayAbilities/Public/AbilitySystemComponent.h
@@ -300,6 +300,9 @@ class GAMEPLAYABILITIES_API UAbilitySystemComponent : public UActorComponent, pu
 	UFUNCTION(NetMulticast, unreliable)
 	void NetMulticast_InvokeGameplayCueAdded(const FGameplayTag GameplayCueTag, FPredictionKey PredictionKey);
 
+	UFUNCTION(NetMulticast, unreliable)
+	void NetMulticast_InvokeGameplayCueRemoved(const FGameplayTag GameplayCueTag, FPredictionKey PredictionKey);
+
 	void InvokeGameplayCueEvent(const FGameplayTag GameplayCueTag, EGameplayCueEvent::Type EventType);
 
 	/** Allows polling to see if a GameplayCue is active. We expect most GameplayCue handling to be event based, but some cases we may need to check if a GamepalyCue is active (Animation Blueprint for example) */
