First phase of deprecating TSubobjectPtr<> replacing TSubobjectPtr properties in the Engine with normal pointers and adding subobject accessors for derived classes/client code.

diff --git a/Engine/Source/Runtime/Engine/Private/TriggerActors.cpp b/Engine/Source/Runtime/Engine/Private/TriggerActors.cpp
--- a/Engine/Source/Runtime/Engine/Private/TriggerActors.cpp
+++ b/Engine/Source/Runtime/Engine/Private/TriggerActors.cpp
@@ -16,7 +16,7 @@ namespace
 ATriggerCapsule::ATriggerCapsule(const FObjectInitializer& ObjectInitializer)
 	: Super(ObjectInitializer.SetDefaultSubobjectClass<UCapsuleComponent>(TEXT("CollisionComp")))
 {
-	UCapsuleComponent* CapsuleCollisionComponent = CastChecked<UCapsuleComponent>(CollisionComponent);
+	UCapsuleComponent* CapsuleCollisionComponent = CastChecked<UCapsuleComponent>(GetCollisionComponent());
 	CapsuleCollisionComponent->ShapeColor = TriggerBaseColor;
 	CapsuleCollisionComponent->InitCapsuleSize(+40.0f, +80.0f);
 	CapsuleCollisionComponent->BodyInstance.bEnableCollision_DEPRECATED = true;
@@ -25,9 +25,9 @@ ATriggerCapsule::ATriggerCapsule(const FObjectInitializer& ObjectInitializer)
 
 	bCollideWhenPlacing = true;
 
-	if (SpriteComponent)
+	if (GetSpriteComponent())
 	{
-		SpriteComponent->AttachParent = CapsuleCollisionComponent;
+		GetSpriteComponent()->AttachParent = CapsuleCollisionComponent;
 	}
 }
 
@@ -84,7 +84,7 @@ void ATriggerCapsule::EditorApplyScale(const FVector& DeltaScale, const FVector*
 ATriggerBox::ATriggerBox(const FObjectInitializer& ObjectInitializer)
 	: Super(ObjectInitializer.SetDefaultSubobjectClass<UBoxComponent>(TEXT("CollisionComp")))
 {
-	UBoxComponent* BoxCollisionComponent = CastChecked<UBoxComponent>(CollisionComponent);
+	UBoxComponent* BoxCollisionComponent = CastChecked<UBoxComponent>(GetCollisionComponent());
 
 	BoxCollisionComponent->ShapeColor = TriggerBaseColor;
 	BoxCollisionComponent->InitBoxExtent(FVector(40.0f, 40.0f, 40.0f));
@@ -93,9 +93,9 @@ ATriggerBox::ATriggerBox(const FObjectInitializer& ObjectInitializer)
 	static FName CollisionProfileName(TEXT("Trigger"));
 	BoxCollisionComponent->SetCollisionProfileName(CollisionProfileName);
 
-	if (SpriteComponent)
+	if (GetSpriteComponent())
 	{
-		SpriteComponent->AttachParent = BoxCollisionComponent;
+		GetSpriteComponent()->AttachParent = BoxCollisionComponent;
 	}
 }
 
@@ -128,7 +128,7 @@ void ATriggerBox::EditorApplyScale(const FVector& DeltaScale, const FVector* Piv
 ATriggerSphere::ATriggerSphere(const FObjectInitializer& ObjectInitializer)
 	: Super(ObjectInitializer.SetDefaultSubobjectClass<USphereComponent>(TEXT("CollisionComp")))
 {
-	USphereComponent* SphereCollisionComponent = CastChecked<USphereComponent>(CollisionComponent);
+	USphereComponent* SphereCollisionComponent = CastChecked<USphereComponent>(GetCollisionComponent());
 
 	SphereCollisionComponent->ShapeColor = TriggerBaseColor;
 	SphereCollisionComponent->InitSphereRadius(+40.0f);
@@ -136,9 +136,9 @@ ATriggerSphere::ATriggerSphere(const FObjectInitializer& ObjectInitializer)
 	static FName CollisionProfileName(TEXT("Trigger"));
 	SphereCollisionComponent->SetCollisionProfileName(CollisionProfileName);
 
-	if (SpriteComponent)
+	if (GetSpriteComponent())
 	{
-		SpriteComponent->AttachParent = SphereCollisionComponent;
+		GetSpriteComponent()->AttachParent = SphereCollisionComponent;
 	}
 }
 
diff --git a/Engine/Source/Runtime/Engine/Private/TriggerVolume.cpp b/Engine/Source/Runtime/Engine/Private/TriggerVolume.cpp
--- a/Engine/Source/Runtime/Engine/Private/TriggerVolume.cpp
+++ b/Engine/Source/Runtime/Engine/Private/TriggerVolume.cpp
@@ -7,7 +7,7 @@ ATriggerVolume::ATriggerVolume(const FObjectInitializer& ObjectInitializer)
 	: Super(ObjectInitializer)
 {
 	static FName CollisionProfileName(TEXT("Trigger"));
-	BrushComponent->SetCollisionProfileName(CollisionProfileName);
+	GetBrushComponent()->SetCollisionProfileName(CollisionProfileName);
 
 	bColored = true;
 	BrushColor.R = 100;
@@ -29,9 +29,9 @@ void ATriggerVolume::LoadedFromAnotherClass(const FName& OldClassName)
 
 		if(OldClassName == DynamicTriggerVolume_NAME)
 		{
-			BrushComponent->Mobility = EComponentMobility::Movable;
+			GetBrushComponent()->Mobility = EComponentMobility::Movable;
 		}
 	}
 }
 
-#endif
\ No newline at end of file
+#endif
diff --git a/Engine/Source/Runtime/Engine/Private/UnrealEngine.cpp b/Engine/Source/Runtime/Engine/Private/UnrealEngine.cpp
--- a/Engine/Source/Runtime/Engine/Private/UnrealEngine.cpp
+++ b/Engine/Source/Runtime/Engine/Private/UnrealEngine.cpp
@@ -2581,10 +2581,10 @@ bool UEngine::HandleMergeMeshCommand( const TCHAR* Cmd, FOutputDevice& Ar, UWorl
 	for( FConstPlayerControllerIterator Iterator = InWorld->GetPlayerControllerIterator(); Iterator; ++Iterator )
 	{
 		APlayerController* PlayerController = *Iterator;
-		if (PlayerController->GetCharacter() != NULL && PlayerController->GetCharacter()->Mesh.IsValid())
+		if (PlayerController->GetCharacter() != NULL && PlayerController->GetCharacter()->GetMesh())
 		{
 			PlayerPawn = PlayerController->GetCharacter();
-			PlayerMesh = PlayerController->GetCharacter()->Mesh->SkeletalMesh;
+			PlayerMesh = PlayerController->GetCharacter()->GetMesh()->SkeletalMesh;
 			break;
 		}
 	}
@@ -2635,7 +2635,7 @@ bool UEngine::HandleMergeMeshCommand( const TCHAR* Cmd, FOutputDevice& Ar, UWorl
 		ASkeletalMeshActor* const SMA = PlayerPawn->GetWorld()->SpawnActor<ASkeletalMeshActor>( SpawnLocation, PlayerPawn->GetActorRotation()*-1);
 		if (SMA)
 		{
-			SMA->SkeletalMeshComponent->SetSkeletalMesh(CompositeMesh);
+			SMA->GetSkeletalMeshComponent()->SetSkeletalMesh(CompositeMesh);
 		}
 	}
 
@@ -6232,14 +6232,14 @@ static FVector2D TransformLocationToMap(FVector2D TopLeftPos, FVector2D BottomRi
 /** Utility for drawing a volume geometry (as seen from above) onto the canvas */
 static void DrawVolumeOnCanvas(const AVolume* Volume, FCanvas* Canvas, const FVector2D& TopLeftPos, const FVector2D& BottomRightPos, const FVector2D& MapOrigin, const FVector2D& MapSize, const FColor& VolColor)
 {
-	if(Volume && Volume->BrushComponent && Volume->BrushComponent->BrushBodySetup)
+	if(Volume && Volume->GetBrushComponent() && Volume->GetBrushComponent()->BrushBodySetup)
 	{
-		FTransform BrushTM = Volume->BrushComponent->ComponentToWorld;
+		FTransform BrushTM = Volume->GetBrushComponent()->ComponentToWorld;
 
 		// Iterate over each piece
-		for(int32 ConIdx=0; ConIdx<Volume->BrushComponent->BrushBodySetup->AggGeom.ConvexElems.Num(); ConIdx++)
+		for(int32 ConIdx=0; ConIdx<Volume->GetBrushComponent()->BrushBodySetup->AggGeom.ConvexElems.Num(); ConIdx++)
 		{
-			FKConvexElem& ConvElem = Volume->BrushComponent->BrushBodySetup->AggGeom.ConvexElems[ConIdx];
+			FKConvexElem& ConvElem = Volume->GetBrushComponent()->BrushBodySetup->AggGeom.ConvexElems[ConIdx];
 
 #if 0 // @todo UE4 physx fix this once we have convexelem drawing again
 			// Draw each triangle that makes up the convex hull
diff --git a/Engine/Source/Runtime/Engine/Private/Volume.cpp b/Engine/Source/Runtime/Engine/Private/Volume.cpp
--- a/Engine/Source/Runtime/Engine/Private/Volume.cpp
+++ b/Engine/Source/Runtime/Engine/Private/Volume.cpp
@@ -16,12 +16,12 @@ DEFINE_LOG_CATEGORY_STATIC(LogVolume, Log, All);
 AVolume::AVolume(const FObjectInitializer& ObjectInitializer)
 	: Super(ObjectInitializer)
 {
-	BrushComponent->AlwaysLoadOnClient = true;
-	BrushComponent->AlwaysLoadOnServer = true;
-	BrushComponent->BodyInstance.bEnableCollision_DEPRECATED = true;
+	GetBrushComponent()->AlwaysLoadOnClient = true;
+	GetBrushComponent()->AlwaysLoadOnServer = true;
+	GetBrushComponent()->BodyInstance.bEnableCollision_DEPRECATED = true;
 	static FName CollisionProfileName(TEXT("OverlapAll"));
-	BrushComponent->SetCollisionProfileName(CollisionProfileName);
-	BrushComponent->bGenerateOverlapEvents = true;
+	GetBrushComponent()->SetCollisionProfileName(CollisionProfileName);
+	GetBrushComponent()->bGenerateOverlapEvents = true;
 	bReplicateMovement = false;
 #if WITH_EDITORONLY_DATA
 	bActorLabelEditable = true;
@@ -56,11 +56,11 @@ void AVolume::PostEditChangeProperty(FPropertyChangedEvent& PropertyChangedEvent
 
 bool AVolume::EncompassesPoint(FVector Point, float SphereRadius/*=0.f*/, float* OutDistanceToPoint)
 {
-	if(BrushComponent)
+	if(GetBrushComponent())
 	{
 #if WITH_PHYSX
 		FVector ClosestPoint;
-		float Distance = BrushComponent->GetDistanceToCollision(Point, ClosestPoint);
+		float Distance = GetBrushComponent()->GetDistanceToCollision(Point, ClosestPoint);
 #else
 		FBoxSphereBounds Bounds = BrushComponent->CalcBounds(BrushComponent->ComponentToWorld);
 		float Distance = FMath::Sqrt(Bounds.GetBox().ComputeSquaredDistanceToPoint(Point));
diff --git a/Engine/Source/Runtime/Engine/Private/WindDirectionalSource.cpp b/Engine/Source/Runtime/Engine/Private/WindDirectionalSource.cpp
--- a/Engine/Source/Runtime/Engine/Private/WindDirectionalSource.cpp
+++ b/Engine/Source/Runtime/Engine/Private/WindDirectionalSource.cpp
@@ -40,13 +40,13 @@ AWindDirectionalSource::AWindDirectionalSource(const FObjectInitializer& ObjectI
 			ArrowComponent->bUseInEditorScaling = true;
 		}
 
-		if (SpriteComponent)
+		if (GetSpriteComponent())
 		{
-			SpriteComponent->Sprite = ConstructorStatics.SpriteTexture.Get();
-			SpriteComponent->RelativeScale3D = FVector(0.5f, 0.5f, 0.5f);
-			SpriteComponent->SpriteInfo.Category = ConstructorStatics.ID_Wind;
-			SpriteComponent->SpriteInfo.DisplayName = ConstructorStatics.NAME_Wind;
-			SpriteComponent->AttachParent = Component;
+			GetSpriteComponent()->Sprite = ConstructorStatics.SpriteTexture.Get();
+			GetSpriteComponent()->RelativeScale3D = FVector(0.5f, 0.5f, 0.5f);
+			GetSpriteComponent()->SpriteInfo.Category = ConstructorStatics.ID_Wind;
+			GetSpriteComponent()->SpriteInfo.DisplayName = ConstructorStatics.NAME_Wind;
+			GetSpriteComponent()->AttachParent = Component;
 		}
 	}
 #endif // WITH_EDITORONLY_DATA
diff --git a/Engine/Source/Runtime/Engine/Private/World.cpp b/Engine/Source/Runtime/Engine/Private/World.cpp
--- a/Engine/Source/Runtime/Engine/Private/World.cpp
+++ b/Engine/Source/Runtime/Engine/Private/World.cpp
@@ -854,13 +854,13 @@ void UWorld::InitWorld(const InitializationValues IVS)
 	ABrush* DefaultBrush = PersistentLevel->Actors.Num()<2 ? NULL : Cast<ABrush>(PersistentLevel->Actors[1]);
 	if (GIsEditor)
 	{
-		if (!DefaultBrush || !DefaultBrush->IsStaticBrush() || DefaultBrush->BrushType!=Brush_Default || !DefaultBrush->BrushComponent || !DefaultBrush->Brush)
+		if (!DefaultBrush || !DefaultBrush->IsStaticBrush() || DefaultBrush->BrushType!=Brush_Default || !DefaultBrush->GetBrushComponent() || !DefaultBrush->Brush)
 		{
 			// Spawn the default brush.
 			DefaultBrush = SpawnBrush();
-			check(DefaultBrush->BrushComponent);
+			check(DefaultBrush->GetBrushComponent());
 			DefaultBrush->Brush = new( DefaultBrush->GetOuter(), TEXT("Brush") )UModel( FObjectInitializer(), DefaultBrush, 1 );
-			DefaultBrush->BrushComponent->Brush = DefaultBrush->Brush;
+			DefaultBrush->GetBrushComponent()->Brush = DefaultBrush->Brush;
 			DefaultBrush->SetNotForClientOrServer();
 			DefaultBrush->Brush->SetFlags( RF_Transactional );
 			DefaultBrush->Brush->Polys->SetFlags( RF_Transactional );
@@ -878,7 +878,7 @@ void UWorld::InitWorld(const InitializationValues IVS)
 		else
 		{
 			// Ensure that the Brush and BrushComponent both point to the same model
-			DefaultBrush->BrushComponent->Brush = DefaultBrush->Brush;
+			DefaultBrush->GetBrushComponent()->Brush = DefaultBrush->Brush;
 		}
 
 		// Reset the lightmass settings on the default brush; they can't be edited by the user but could have
@@ -4777,9 +4777,9 @@ void UWorld::UpdateConstraintActors()
 		for( TActorIterator<APhysicsConstraintActor> It(this); It; ++It )
 		{
 			APhysicsConstraintActor* ConstraintActor = *It;
-			if( ConstraintActor->ConstraintComp)
+			if( ConstraintActor->GetConstraintComp())
 			{
-				ConstraintActor->ConstraintComp->UpdateConstraintFrames();
+				ConstraintActor->GetConstraintComp()->UpdateConstraintFrames();
 			}
 		}
 		bAreConstraintsDirty = false;
diff --git a/Engine/Source/Runtime/Engine/Private/WorldSettings.cpp b/Engine/Source/Runtime/Engine/Private/WorldSettings.cpp
--- a/Engine/Source/Runtime/Engine/Private/WorldSettings.cpp
+++ b/Engine/Source/Runtime/Engine/Private/WorldSettings.cpp
@@ -56,7 +56,7 @@ AWorldSettings::AWorldSettings(const FObjectInitializer& ObjectInitializer)
 	VisibilityAggressiveness = VIS_LeastAggressive;
 	LevelLightingQuality = Quality_MAX;
 
-	TSubobjectPtr<UStaticMeshComponent> StaticMeshComponent = ObjectInitializer.CreateDefaultSubobject<UStaticMeshComponent>(this, TEXT("StaticMeshComponent0"));
+	UStaticMeshComponent* StaticMeshComponent = ObjectInitializer.CreateDefaultSubobject<UStaticMeshComponent>(this, TEXT("StaticMeshComponent0"));
 	StaticMeshComponent->bHiddenInGame = true;
 	StaticMeshComponent->BodyInstance.bEnableCollision_DEPRECATED = true;
 	StaticMeshComponent->SetCollisionProfileName(UCollisionProfile::BlockAll_ProfileName);
diff --git a/Engine/Source/Runtime/GameplayAbilities/Private/GameplayEffectTests.cpp b/Engine/Source/Runtime/GameplayAbilities/Private/GameplayEffectTests.cpp
--- a/Engine/Source/Runtime/GameplayAbilities/Private/GameplayEffectTests.cpp
+++ b/Engine/Source/Runtime/GameplayAbilities/Private/GameplayEffectTests.cpp
@@ -112,8 +112,8 @@ bool GameplayEffectsTest_InstantDamage(UWorld *World, FAutomationTestBase * Test
 
 	UProperty *HealthProperty = FindFieldChecked<UProperty>(UAbilitySystemTestAttributeSet::StaticClass(), GET_MEMBER_NAME_CHECKED(UAbilitySystemTestAttributeSet, Health));
 
-	UAbilitySystemComponent * SourceComponent = SourceActor->AbilitySystemComponent;
-	UAbilitySystemComponent * DestComponent = DestActor->AbilitySystemComponent;
+	UAbilitySystemComponent * SourceComponent = SourceActor->GetAbilitySystemComponent();
+	UAbilitySystemComponent * DestComponent = DestActor->GetAbilitySystemComponent();
 	SourceComponent->GetSet<UAbilitySystemTestAttributeSet>()->Health = StartHealth;
 	DestComponent->GetSet<UAbilitySystemTestAttributeSet>()->Health = StartHealth;
 
@@ -154,8 +154,8 @@ bool GameplayEffectsTest_InstantDamageRemap(UWorld *World, FAutomationTestBase *
 	UProperty *HealthProperty = FindFieldChecked<UProperty>(UAbilitySystemTestAttributeSet::StaticClass(), GET_MEMBER_NAME_CHECKED(UAbilitySystemTestAttributeSet, Health));
 	UProperty *DamageProperty = FindFieldChecked<UProperty>(UAbilitySystemTestAttributeSet::StaticClass(), GET_MEMBER_NAME_CHECKED(UAbilitySystemTestAttributeSet, Damage));
 
-	UAbilitySystemComponent * SourceComponent = SourceActor->AbilitySystemComponent;
-	UAbilitySystemComponent * DestComponent = DestActor->AbilitySystemComponent;
+	UAbilitySystemComponent * SourceComponent = SourceActor->GetAbilitySystemComponent();
+	UAbilitySystemComponent * DestComponent = DestActor->GetAbilitySystemComponent();
 	SourceComponent->GetSet<UAbilitySystemTestAttributeSet>()->Health = StartHealth;
 	DestComponent->GetSet<UAbilitySystemTestAttributeSet>()->Health = StartHealth;
 
@@ -205,8 +205,8 @@ bool GameplayEffectsTest_InstantDamage_Buffed(UWorld *World, FAutomationTestBase
 
 	UProperty *HealthProperty = FindFieldChecked<UProperty>(UAbilitySystemTestAttributeSet::StaticClass(), GET_MEMBER_NAME_CHECKED(UAbilitySystemTestAttributeSet, Health));
 
-	UAbilitySystemComponent * SourceComponent = SourceActor->AbilitySystemComponent;
-	UAbilitySystemComponent * DestComponent = DestActor->AbilitySystemComponent;
+	UAbilitySystemComponent * SourceComponent = SourceActor->GetAbilitySystemComponent();
+	UAbilitySystemComponent * DestComponent = DestActor->GetAbilitySystemComponent();
 	SourceComponent->GetSet<UAbilitySystemTestAttributeSet>()->Health = StartHealth;
 	DestComponent->GetSet<UAbilitySystemTestAttributeSet>()->Health = StartHealth;
 
@@ -269,8 +269,8 @@ bool GameplayEffectsTest_TemporaryDamage(UWorld *World, FAutomationTestBase * Te
 
 	UProperty *HealthProperty = FindFieldChecked<UProperty>(UAbilitySystemTestAttributeSet::StaticClass(), GET_MEMBER_NAME_CHECKED(UAbilitySystemTestAttributeSet, Health));
 
-	UAbilitySystemComponent * SourceComponent = SourceActor->AbilitySystemComponent;
-	UAbilitySystemComponent * DestComponent = DestActor->AbilitySystemComponent;
+	UAbilitySystemComponent * SourceComponent = SourceActor->GetAbilitySystemComponent();
+	UAbilitySystemComponent * DestComponent = DestActor->GetAbilitySystemComponent();
 	SourceComponent->GetSet<UAbilitySystemTestAttributeSet>()->Health = StartHealth;
 	DestComponent->GetSet<UAbilitySystemTestAttributeSet>()->Health = StartHealth;
 
@@ -328,8 +328,8 @@ bool GameplayEffectsTest_TemporaryDamageBuffed(UWorld *World, FAutomationTestBas
 
 	UProperty *HealthProperty = FindFieldChecked<UProperty>(UAbilitySystemTestAttributeSet::StaticClass(), GET_MEMBER_NAME_CHECKED(UAbilitySystemTestAttributeSet, Health));
 
-	UAbilitySystemComponent * SourceComponent = SourceActor->AbilitySystemComponent;
-	UAbilitySystemComponent * DestComponent = DestActor->AbilitySystemComponent;
+	UAbilitySystemComponent * SourceComponent = SourceActor->GetAbilitySystemComponent();
+	UAbilitySystemComponent * DestComponent = DestActor->GetAbilitySystemComponent();
 	SourceComponent->GetSet<UAbilitySystemTestAttributeSet>()->Health = StartHealth;
 	DestComponent->GetSet<UAbilitySystemTestAttributeSet>()->Health = StartHealth;
 
@@ -418,8 +418,8 @@ bool GameplayEffectsTest_TemporaryDamageTemporaryBuff(UWorld *World, FAutomation
 
 	UProperty *HealthProperty = FindFieldChecked<UProperty>(UAbilitySystemTestAttributeSet::StaticClass(), GET_MEMBER_NAME_CHECKED(UAbilitySystemTestAttributeSet, Health));
 
-	UAbilitySystemComponent * SourceComponent = SourceActor->AbilitySystemComponent;
-	UAbilitySystemComponent * DestComponent = DestActor->AbilitySystemComponent;
+	UAbilitySystemComponent * SourceComponent = SourceActor->GetAbilitySystemComponent();
+	UAbilitySystemComponent * DestComponent = DestActor->GetAbilitySystemComponent();
 	SourceComponent->GetSet<UAbilitySystemTestAttributeSet>()->Health = StartHealth;
 	DestComponent->GetSet<UAbilitySystemTestAttributeSet>()->Health = StartHealth;
 
@@ -527,8 +527,8 @@ bool GameplayEffectsTest_LinkedBuffDestroy(UWorld *World, FAutomationTestBase *
 
 	UProperty *HealthProperty = FindFieldChecked<UProperty>(UAbilitySystemTestAttributeSet::StaticClass(), GET_MEMBER_NAME_CHECKED(UAbilitySystemTestAttributeSet, Health));
 
-	UAbilitySystemComponent * SourceComponent = SourceActor->AbilitySystemComponent;
-	UAbilitySystemComponent * DestComponent = DestActor->AbilitySystemComponent;
+	UAbilitySystemComponent * SourceComponent = SourceActor->GetAbilitySystemComponent();
+	UAbilitySystemComponent * DestComponent = DestActor->GetAbilitySystemComponent();
 	SourceComponent->GetSet<UAbilitySystemTestAttributeSet>()->Health = StartHealth;
 	DestComponent->GetSet<UAbilitySystemTestAttributeSet>()->Health = StartHealth;
 
@@ -614,8 +614,8 @@ bool GameplayEffectsTest_SnapshotBuffDestroy(UWorld *World, FAutomationTestBase
 
 	UProperty *HealthProperty = FindFieldChecked<UProperty>(UAbilitySystemTestAttributeSet::StaticClass(), GET_MEMBER_NAME_CHECKED(UAbilitySystemTestAttributeSet, Health));
 
-	UAbilitySystemComponent * SourceComponent = SourceActor->AbilitySystemComponent;
-	UAbilitySystemComponent * DestComponent = DestActor->AbilitySystemComponent;
+	UAbilitySystemComponent * SourceComponent = SourceActor->GetAbilitySystemComponent();
+	UAbilitySystemComponent * DestComponent = DestActor->GetAbilitySystemComponent();
 	
 	SourceComponent->GetSet<UAbilitySystemTestAttributeSet>()->Health = StartHealth;
 	DestComponent->GetSet<UAbilitySystemTestAttributeSet>()->Health = StartHealth;
@@ -706,8 +706,8 @@ bool GameplayEffectsTest_DurationBuff(UWorld *World, FAutomationTestBase * Test)
 
 	UProperty *HealthProperty = FindFieldChecked<UProperty>(UAbilitySystemTestAttributeSet::StaticClass(), GET_MEMBER_NAME_CHECKED(UAbilitySystemTestAttributeSet, Health));
 
-	UAbilitySystemComponent * SourceComponent = SourceActor->AbilitySystemComponent;
-	UAbilitySystemComponent * DestComponent = DestActor->AbilitySystemComponent;
+	UAbilitySystemComponent * SourceComponent = SourceActor->GetAbilitySystemComponent();
+	UAbilitySystemComponent * DestComponent = DestActor->GetAbilitySystemComponent();
 
 	SourceComponent->GetSet<UAbilitySystemTestAttributeSet>()->Health = StartHealth;
 	DestComponent->GetSet<UAbilitySystemTestAttributeSet>()->Health = StartHealth;
@@ -961,8 +961,8 @@ bool GameplayEffectsTest_DamageBuffBuff_Basic(UWorld *World, FAutomationTestBase
 
 	UProperty *HealthProperty = FindFieldChecked<UProperty>(UAbilitySystemTestAttributeSet::StaticClass(), GET_MEMBER_NAME_CHECKED(UAbilitySystemTestAttributeSet, Health));
 
-	UAbilitySystemComponent * SourceComponent = SourceActor->AbilitySystemComponent;
-	UAbilitySystemComponent * DestComponent = DestActor->AbilitySystemComponent;
+	UAbilitySystemComponent * SourceComponent = SourceActor->GetAbilitySystemComponent();
+	UAbilitySystemComponent * DestComponent = DestActor->GetAbilitySystemComponent();
 	SourceComponent->GetSet<UAbilitySystemTestAttributeSet>()->Health = StartHealth;
 	DestComponent->GetSet<UAbilitySystemTestAttributeSet>()->Health = StartHealth;
 
@@ -1102,8 +1102,8 @@ bool GameplayEffectsTest_DamageBuffBuff_FullLink(UWorld *World, FAutomationTestB
 
 	UProperty *HealthProperty = FindFieldChecked<UProperty>(UAbilitySystemTestAttributeSet::StaticClass(), GET_MEMBER_NAME_CHECKED(UAbilitySystemTestAttributeSet, Health));
 
-	UAbilitySystemComponent * SourceComponent = SourceActor->AbilitySystemComponent;
-	UAbilitySystemComponent * DestComponent = DestActor->AbilitySystemComponent;
+	UAbilitySystemComponent * SourceComponent = SourceActor->GetAbilitySystemComponent();
+	UAbilitySystemComponent * DestComponent = DestActor->GetAbilitySystemComponent();
 	SourceComponent->GetSet<UAbilitySystemTestAttributeSet>()->Health = StartHealth;
 	DestComponent->GetSet<UAbilitySystemTestAttributeSet>()->Health = StartHealth;
 
@@ -1240,8 +1240,8 @@ bool GameplayEffectsTest_DamageBuffBuff_FullSnapshot(UWorld *World, FAutomationT
 
 	UProperty *HealthProperty = FindFieldChecked<UProperty>(UAbilitySystemTestAttributeSet::StaticClass(), GET_MEMBER_NAME_CHECKED(UAbilitySystemTestAttributeSet, Health));
 
-	UAbilitySystemComponent * SourceComponent = SourceActor->AbilitySystemComponent;
-	UAbilitySystemComponent * DestComponent = DestActor->AbilitySystemComponent;
+	UAbilitySystemComponent * SourceComponent = SourceActor->GetAbilitySystemComponent();
+	UAbilitySystemComponent * DestComponent = DestActor->GetAbilitySystemComponent();
 	SourceComponent->GetSet<UAbilitySystemTestAttributeSet>()->Health = StartHealth;
 	DestComponent->GetSet<UAbilitySystemTestAttributeSet>()->Health = StartHealth;
 
@@ -1378,8 +1378,8 @@ bool GameplayEffectsTest_DamageBuffBuff_SnapshotLink(UWorld *World, FAutomationT
 
 	UProperty *HealthProperty = FindFieldChecked<UProperty>(UAbilitySystemTestAttributeSet::StaticClass(), GET_MEMBER_NAME_CHECKED(UAbilitySystemTestAttributeSet, Health));
 
-	UAbilitySystemComponent * SourceComponent = SourceActor->AbilitySystemComponent;
-	UAbilitySystemComponent * DestComponent = DestActor->AbilitySystemComponent;
+	UAbilitySystemComponent * SourceComponent = SourceActor->GetAbilitySystemComponent();
+	UAbilitySystemComponent * DestComponent = DestActor->GetAbilitySystemComponent();
 	SourceComponent->GetSet<UAbilitySystemTestAttributeSet>()->Health = StartHealth;
 	DestComponent->GetSet<UAbilitySystemTestAttributeSet>()->Health = StartHealth;
 
@@ -1512,8 +1512,8 @@ bool GameplayEffectsTest_DamageAppliesBuff(UWorld *World, FAutomationTestBase *
 
 	UProperty *DamageProperty = FindFieldChecked<UProperty>(UAbilitySystemTestAttributeSet::StaticClass(), GET_MEMBER_NAME_CHECKED(UAbilitySystemTestAttributeSet, Damage));
 
-	UAbilitySystemComponent * SourceComponent = SourceActor->AbilitySystemComponent;
-	UAbilitySystemComponent * DestComponent = DestActor->AbilitySystemComponent;
+	UAbilitySystemComponent * SourceComponent = SourceActor->GetAbilitySystemComponent();
+	UAbilitySystemComponent * DestComponent = DestActor->GetAbilitySystemComponent();
 	SourceComponent->GetSet<UAbilitySystemTestAttributeSet>()->Health = StartHealth;
 	DestComponent->GetSet<UAbilitySystemTestAttributeSet>()->Health = StartHealth;
 
@@ -1589,8 +1589,8 @@ bool GameplayEffectsTest_BuffAppliesBuff(UWorld *World, FAutomationTestBase * Te
 
 	UProperty *DamageProperty = FindFieldChecked<UProperty>(UAbilitySystemTestAttributeSet::StaticClass(), GET_MEMBER_NAME_CHECKED(UAbilitySystemTestAttributeSet, Damage));
 
-	UAbilitySystemComponent * SourceComponent = SourceActor->AbilitySystemComponent;
-	UAbilitySystemComponent * DestComponent = DestActor->AbilitySystemComponent;
+	UAbilitySystemComponent * SourceComponent = SourceActor->GetAbilitySystemComponent();
+	UAbilitySystemComponent * DestComponent = DestActor->GetAbilitySystemComponent();
 	SourceComponent->GetSet<UAbilitySystemTestAttributeSet>()->Health = StartHealth;
 	DestComponent->GetSet<UAbilitySystemTestAttributeSet>()->Health = StartHealth;
 
@@ -1721,8 +1721,8 @@ bool GameplayEffectsTest_BuffIndirection(UWorld *World, FAutomationTestBase * Te
 
 	UProperty *DamageProperty = FindFieldChecked<UProperty>(UAbilitySystemTestAttributeSet::StaticClass(), GET_MEMBER_NAME_CHECKED(UAbilitySystemTestAttributeSet, Damage));
 
-	UAbilitySystemComponent * SourceComponent = SourceActor->AbilitySystemComponent;
-	UAbilitySystemComponent * DestComponent = DestActor->AbilitySystemComponent;
+	UAbilitySystemComponent * SourceComponent = SourceActor->GetAbilitySystemComponent();
+	UAbilitySystemComponent * DestComponent = DestActor->GetAbilitySystemComponent();
 	SourceComponent->GetSet<UAbilitySystemTestAttributeSet>()->Health = StartHealth;
 	DestComponent->GetSet<UAbilitySystemTestAttributeSet>()->Health = StartHealth;
 
@@ -1841,8 +1841,8 @@ bool GameplayEffectsTest_DurationDamage(UWorld *World, FAutomationTestBase * Tes
 
 	UProperty *HealthProperty = FindFieldChecked<UProperty>(UAbilitySystemTestAttributeSet::StaticClass(), GET_MEMBER_NAME_CHECKED(UAbilitySystemTestAttributeSet, Health));
 
-	UAbilitySystemComponent * SourceComponent = SourceActor->AbilitySystemComponent;
-	UAbilitySystemComponent * DestComponent = DestActor->AbilitySystemComponent;
+	UAbilitySystemComponent * SourceComponent = SourceActor->GetAbilitySystemComponent();
+	UAbilitySystemComponent * DestComponent = DestActor->GetAbilitySystemComponent();
 
 	SourceComponent->GetSet<UAbilitySystemTestAttributeSet>()->Health = StartHealth;
 	DestComponent->GetSet<UAbilitySystemTestAttributeSet>()->Health = StartHealth;
@@ -1918,8 +1918,8 @@ bool GameplayEffectsTest_PeriodicDamage(UWorld *World, FAutomationTestBase * Tes
 
 	UProperty *HealthProperty = FindFieldChecked<UProperty>(UAbilitySystemTestAttributeSet::StaticClass(), GET_MEMBER_NAME_CHECKED(UAbilitySystemTestAttributeSet, Health));
 
-	UAbilitySystemComponent * SourceComponent = SourceActor->AbilitySystemComponent;
-	UAbilitySystemComponent * DestComponent = DestActor->AbilitySystemComponent;
+	UAbilitySystemComponent * SourceComponent = SourceActor->GetAbilitySystemComponent();
+	UAbilitySystemComponent * DestComponent = DestActor->GetAbilitySystemComponent();
 
 	SourceComponent->GetSet<UAbilitySystemTestAttributeSet>()->Health = StartHealth;
 	DestComponent->GetSet<UAbilitySystemTestAttributeSet>()->Health = StartHealth;
@@ -1993,8 +1993,8 @@ bool GameplayEffectsTest_LifestealExtension(UWorld *World, FAutomationTestBase *
 	UProperty *HealthProperty = FindFieldChecked<UProperty>(UAbilitySystemTestAttributeSet::StaticClass(), GET_MEMBER_NAME_CHECKED(UAbilitySystemTestAttributeSet, Health));
 	UProperty *DamageProperty = FindFieldChecked<UProperty>(UAbilitySystemTestAttributeSet::StaticClass(), GET_MEMBER_NAME_CHECKED(UAbilitySystemTestAttributeSet, Damage));
 
-	UAbilitySystemComponent * SourceComponent = SourceActor->AbilitySystemComponent;
-	UAbilitySystemComponent * DestComponent = DestActor->AbilitySystemComponent;
+	UAbilitySystemComponent * SourceComponent = SourceActor->GetAbilitySystemComponent();
+	UAbilitySystemComponent * DestComponent = DestActor->GetAbilitySystemComponent();
 
 	SourceComponent->GetSet<UAbilitySystemTestAttributeSet>()->Health = StartHealth;
 	DestComponent->GetSet<UAbilitySystemTestAttributeSet>()->Health = StartHealth;
@@ -2068,8 +2068,8 @@ bool GameplayEffectsTest_ShieldExtension(UWorld *World, FAutomationTestBase * Te
 	UProperty *HealthProperty = FindFieldChecked<UProperty>(UAbilitySystemTestAttributeSet::StaticClass(), GET_MEMBER_NAME_CHECKED(UAbilitySystemTestAttributeSet, Health));
 	UProperty *DamageProperty = FindFieldChecked<UProperty>(UAbilitySystemTestAttributeSet::StaticClass(), GET_MEMBER_NAME_CHECKED(UAbilitySystemTestAttributeSet, Damage));
 
-	UAbilitySystemComponent * SourceComponent = SourceActor->AbilitySystemComponent;
-	UAbilitySystemComponent * DestComponent = DestActor->AbilitySystemComponent;
+	UAbilitySystemComponent * SourceComponent = SourceActor->GetAbilitySystemComponent();
+	UAbilitySystemComponent * DestComponent = DestActor->GetAbilitySystemComponent();
 
 	SourceComponent->GetSet<UAbilitySystemTestAttributeSet>()->Health = StartHealth;
 	DestComponent->GetSet<UAbilitySystemTestAttributeSet>()->Health = StartHealth;
@@ -2183,8 +2183,8 @@ bool GameplayEffectsTest_ShieldExtensionMultiple(UWorld *World, FAutomationTestB
 	UProperty *HealthProperty = FindFieldChecked<UProperty>(UAbilitySystemTestAttributeSet::StaticClass(), GET_MEMBER_NAME_CHECKED(UAbilitySystemTestAttributeSet, Health));
 	UProperty *DamageProperty = FindFieldChecked<UProperty>(UAbilitySystemTestAttributeSet::StaticClass(), GET_MEMBER_NAME_CHECKED(UAbilitySystemTestAttributeSet, Damage));
 
-	UAbilitySystemComponent * SourceComponent = SourceActor->AbilitySystemComponent;
-	UAbilitySystemComponent * DestComponent = DestActor->AbilitySystemComponent;
+	UAbilitySystemComponent * SourceComponent = SourceActor->GetAbilitySystemComponent();
+	UAbilitySystemComponent * DestComponent = DestActor->GetAbilitySystemComponent();
 
 	SourceComponent->GetSet<UAbilitySystemTestAttributeSet>()->Health = StartHealth;
 	DestComponent->GetSet<UAbilitySystemTestAttributeSet>()->Health = StartHealth;
@@ -2442,8 +2442,8 @@ bool GameplayEffectsTest_InstantDamage_ScalingExplicit(UWorld *World, FAutomatio
 
 	UProperty *DamageProperty = FindFieldChecked<UProperty>(UAbilitySystemTestAttributeSet::StaticClass(), GET_MEMBER_NAME_CHECKED(UAbilitySystemTestAttributeSet, Damage));
 
-	UAbilitySystemComponent * SourceComponent = SourceActor->AbilitySystemComponent;
-	UAbilitySystemComponent * DestComponent = DestActor->AbilitySystemComponent;
+	UAbilitySystemComponent * SourceComponent = SourceActor->GetAbilitySystemComponent();
+	UAbilitySystemComponent * DestComponent = DestActor->GetAbilitySystemComponent();
 	SourceComponent->GetSet<UAbilitySystemTestAttributeSet>()->Health = StartHealth;
 	DestComponent->GetSet<UAbilitySystemTestAttributeSet>()->Health = StartHealth;
 
@@ -2491,8 +2491,8 @@ bool GameplayEffectsTest_InstantDamage_ScalingGlobal(UWorld *World, FAutomationT
 	
 	UProperty *DamageProperty = FindFieldChecked<UProperty>(UAbilitySystemTestAttributeSet::StaticClass(), GET_MEMBER_NAME_CHECKED(UAbilitySystemTestAttributeSet, Damage));
 
-	UAbilitySystemComponent * SourceComponent = SourceActor->AbilitySystemComponent;
-	UAbilitySystemComponent * DestComponent = DestActor->AbilitySystemComponent;
+	UAbilitySystemComponent * SourceComponent = SourceActor->GetAbilitySystemComponent();
+	UAbilitySystemComponent * DestComponent = DestActor->GetAbilitySystemComponent();
 	SourceComponent->GetSet<UAbilitySystemTestAttributeSet>()->Health = StartHealth;
 	DestComponent->GetSet<UAbilitySystemTestAttributeSet>()->Health = StartHealth;
 
@@ -2540,8 +2540,8 @@ bool GameplayEffectsTest_InstantDamage_OverrideScaling(UWorld *World, FAutomatio
 
 	UProperty *DamageProperty = FindFieldChecked<UProperty>(UAbilitySystemTestAttributeSet::StaticClass(), GET_MEMBER_NAME_CHECKED(UAbilitySystemTestAttributeSet, Damage));
 
-	UAbilitySystemComponent * SourceComponent = SourceActor->AbilitySystemComponent;
-	UAbilitySystemComponent * DestComponent = DestActor->AbilitySystemComponent;
+	UAbilitySystemComponent * SourceComponent = SourceActor->GetAbilitySystemComponent();
+	UAbilitySystemComponent * DestComponent = DestActor->GetAbilitySystemComponent();
 	SourceComponent->GetSet<UAbilitySystemTestAttributeSet>()->Health = StartHealth;
 	DestComponent->GetSet<UAbilitySystemTestAttributeSet>()->Health = StartHealth;
 
@@ -2586,8 +2586,8 @@ bool GameplayEffectsTest_InstantDamageRequiredTag(UWorld *World, FAutomationTest
 
 	UProperty *HealthProperty = FindFieldChecked<UProperty>(UAbilitySystemTestAttributeSet::StaticClass(), GET_MEMBER_NAME_CHECKED(UAbilitySystemTestAttributeSet, Health));
 
-	UAbilitySystemComponent * SourceComponent = SourceActor->AbilitySystemComponent;
-	UAbilitySystemComponent * DestComponent = DestActor->AbilitySystemComponent;
+	UAbilitySystemComponent * SourceComponent = SourceActor->GetAbilitySystemComponent();
+	UAbilitySystemComponent * DestComponent = DestActor->GetAbilitySystemComponent();
 	SourceComponent->GetSet<UAbilitySystemTestAttributeSet>()->Health = StartHealth;
 	DestComponent->GetSet<UAbilitySystemTestAttributeSet>()->Health = StartHealth;
 
@@ -2676,8 +2676,8 @@ bool GameplayEffectsTest_InstantDamageIgnoreTag(UWorld *World, FAutomationTestBa
 
 	UProperty *HealthProperty = FindFieldChecked<UProperty>(UAbilitySystemTestAttributeSet::StaticClass(), GET_MEMBER_NAME_CHECKED(UAbilitySystemTestAttributeSet, Health));
 
-	UAbilitySystemComponent * SourceComponent = SourceActor->AbilitySystemComponent;
-	UAbilitySystemComponent * DestComponent = DestActor->AbilitySystemComponent;
+	UAbilitySystemComponent * SourceComponent = SourceActor->GetAbilitySystemComponent();
+	UAbilitySystemComponent * DestComponent = DestActor->GetAbilitySystemComponent();
 	SourceComponent->GetSet<UAbilitySystemTestAttributeSet>()->Health = StartHealth;
 	DestComponent->GetSet<UAbilitySystemTestAttributeSet>()->Health = StartHealth;
 
@@ -2767,8 +2767,8 @@ bool GameplayEffectsTest_InstantDamageModifierPassesTag(UWorld *World, FAutomati
 
 	UProperty *HealthProperty = FindFieldChecked<UProperty>(UAbilitySystemTestAttributeSet::StaticClass(), GET_MEMBER_NAME_CHECKED(UAbilitySystemTestAttributeSet, Health));
 
-	UAbilitySystemComponent * SourceComponent = SourceActor->AbilitySystemComponent;
-	UAbilitySystemComponent * DestComponent = DestActor->AbilitySystemComponent;
+	UAbilitySystemComponent * SourceComponent = SourceActor->GetAbilitySystemComponent();
+	UAbilitySystemComponent * DestComponent = DestActor->GetAbilitySystemComponent();
 	SourceComponent->GetSet<UAbilitySystemTestAttributeSet>()->Health = StartHealth;
 	DestComponent->GetSet<UAbilitySystemTestAttributeSet>()->Health = StartHealth;
 
@@ -2849,8 +2849,8 @@ bool GameplayEffectsTest_InstantDamageModifierTag(UWorld *World, FAutomationTest
 
 	UProperty *HealthProperty = FindFieldChecked<UProperty>(UAbilitySystemTestAttributeSet::StaticClass(), GET_MEMBER_NAME_CHECKED(UAbilitySystemTestAttributeSet, Health));
 
-	UAbilitySystemComponent * SourceComponent = SourceActor->AbilitySystemComponent;
-	UAbilitySystemComponent * DestComponent = DestActor->AbilitySystemComponent;
+	UAbilitySystemComponent * SourceComponent = SourceActor->GetAbilitySystemComponent();
+	UAbilitySystemComponent * DestComponent = DestActor->GetAbilitySystemComponent();
 	SourceComponent->GetSet<UAbilitySystemTestAttributeSet>()->Health = StartHealth;
 	DestComponent->GetSet<UAbilitySystemTestAttributeSet>()->Health = StartHealth;
 
@@ -2935,8 +2935,8 @@ bool GameplayEffectsTest_InstantDamage_ScalingProperty(UWorld *World, FAutomatio
 	UProperty *DamageProperty = FindFieldChecked<UProperty>(UAbilitySystemTestAttributeSet::StaticClass(), GET_MEMBER_NAME_CHECKED(UAbilitySystemTestAttributeSet, Damage));
 	UProperty *PhysicalDamageProperty = FindFieldChecked<UProperty>(UAbilitySystemTestAttributeSet::StaticClass(), GET_MEMBER_NAME_CHECKED(UAbilitySystemTestAttributeSet, PhysicalDamage));
 
-	UAbilitySystemComponent * SourceComponent = SourceActor->AbilitySystemComponent;
-	UAbilitySystemComponent * DestComponent = DestActor->AbilitySystemComponent;
+	UAbilitySystemComponent * SourceComponent = SourceActor->GetAbilitySystemComponent();
+	UAbilitySystemComponent * DestComponent = DestActor->GetAbilitySystemComponent();
 	SourceComponent->GetSet<UAbilitySystemTestAttributeSet>()->Health = StartHealth;
 	DestComponent->GetSet<UAbilitySystemTestAttributeSet>()->Health = StartHealth;
 
@@ -2991,8 +2991,8 @@ bool GameplayEffectsTest_InstantDamage_ScalingPropertyNested(UWorld *World, FAut
 	UProperty *DamageProperty = FindFieldChecked<UProperty>(UAbilitySystemTestAttributeSet::StaticClass(), GET_MEMBER_NAME_CHECKED(UAbilitySystemTestAttributeSet, Damage));
 	UProperty *PhysicalDamageProperty = FindFieldChecked<UProperty>(UAbilitySystemTestAttributeSet::StaticClass(), GET_MEMBER_NAME_CHECKED(UAbilitySystemTestAttributeSet, PhysicalDamage));
 
-	UAbilitySystemComponent * SourceComponent = SourceActor->AbilitySystemComponent;
-	UAbilitySystemComponent * DestComponent = DestActor->AbilitySystemComponent;
+	UAbilitySystemComponent * SourceComponent = SourceActor->GetAbilitySystemComponent();
+	UAbilitySystemComponent * DestComponent = DestActor->GetAbilitySystemComponent();
 	SourceComponent->GetSet<UAbilitySystemTestAttributeSet>()->Health = StartHealth;
 	DestComponent->GetSet<UAbilitySystemTestAttributeSet>()->Health = StartHealth;
 
@@ -3048,8 +3048,8 @@ bool GameplayEffectsTest_DotDamage_ScalingProperty_Snapshot(UWorld *World, FAuto
 	UProperty *DamageProperty = FindFieldChecked<UProperty>(UAbilitySystemTestAttributeSet::StaticClass(), GET_MEMBER_NAME_CHECKED(UAbilitySystemTestAttributeSet, Damage));
 	UProperty *SpellDamageProperty = FindFieldChecked<UProperty>(UAbilitySystemTestAttributeSet::StaticClass(), GET_MEMBER_NAME_CHECKED(UAbilitySystemTestAttributeSet, SpellDamage));
 
-	UAbilitySystemComponent * SourceComponent = SourceActor->AbilitySystemComponent;
-	UAbilitySystemComponent * DestComponent = DestActor->AbilitySystemComponent;
+	UAbilitySystemComponent * SourceComponent = SourceActor->GetAbilitySystemComponent();
+	UAbilitySystemComponent * DestComponent = DestActor->GetAbilitySystemComponent();
 	SourceComponent->GetSet<UAbilitySystemTestAttributeSet>()->Health = StartHealth;
 	DestComponent->GetSet<UAbilitySystemTestAttributeSet>()->Health = StartHealth;
 
@@ -3134,8 +3134,8 @@ bool GameplayEffectsTest_DotDamage_ScalingProperty_Dynamic(UWorld *World, FAutom
 	UProperty *DamageProperty = FindFieldChecked<UProperty>(UAbilitySystemTestAttributeSet::StaticClass(), GET_MEMBER_NAME_CHECKED(UAbilitySystemTestAttributeSet, Damage));
 	UProperty *SpellDamageProperty = FindFieldChecked<UProperty>(UAbilitySystemTestAttributeSet::StaticClass(), GET_MEMBER_NAME_CHECKED(UAbilitySystemTestAttributeSet, SpellDamage));
 
-	UAbilitySystemComponent * SourceComponent = SourceActor->AbilitySystemComponent;
-	UAbilitySystemComponent * DestComponent = DestActor->AbilitySystemComponent;
+	UAbilitySystemComponent * SourceComponent = SourceActor->GetAbilitySystemComponent();
+	UAbilitySystemComponent * DestComponent = DestActor->GetAbilitySystemComponent();
 	SourceComponent->GetSet<UAbilitySystemTestAttributeSet>()->Health = StartHealth;
 	DestComponent->GetSet<UAbilitySystemTestAttributeSet>()->Health = StartHealth;
 
@@ -3233,7 +3233,7 @@ bool GameplayEffectsTest_MetaAttributes(UWorld *World, FAutomationTestBase * Tes
 	UProperty *MaxHealthProperty = FindFieldChecked<UProperty>(UAbilitySystemTestAttributeSet::StaticClass(), GET_MEMBER_NAME_CHECKED(UAbilitySystemTestAttributeSet, MaxHealth));
 	UProperty *StrengthProperty = FindFieldChecked<UProperty>(UAbilitySystemTestAttributeSet::StaticClass(), GET_MEMBER_NAME_CHECKED(UAbilitySystemTestAttributeSet, Strength));
 
-	UAbilitySystemComponent * SourceComponent = SourceActor->AbilitySystemComponent;
+	UAbilitySystemComponent * SourceComponent = SourceActor->GetAbilitySystemComponent();
 	SourceComponent->GetSet<UAbilitySystemTestAttributeSet>()->MaxHealth = StartHealth;
 	SourceComponent->GetSet<UAbilitySystemTestAttributeSet>()->Strength = 0.f;
 
@@ -3310,8 +3310,8 @@ bool GameplayEffectsTest_TagOrdering(UWorld *World, FAutomationTestBase * Test)
 	UProperty *HealthProperty = FindFieldChecked<UProperty>(UAbilitySystemTestAttributeSet::StaticClass(), GET_MEMBER_NAME_CHECKED(UAbilitySystemTestAttributeSet, Health));
 	UProperty *DamageProperty = FindFieldChecked<UProperty>(UAbilitySystemTestAttributeSet::StaticClass(), GET_MEMBER_NAME_CHECKED(UAbilitySystemTestAttributeSet, Damage));
 
-	UAbilitySystemComponent * SourceComponent = SourceActor->AbilitySystemComponent;
-	UAbilitySystemComponent * DestComponent = DestActor->AbilitySystemComponent;
+	UAbilitySystemComponent * SourceComponent = SourceActor->GetAbilitySystemComponent();
+	UAbilitySystemComponent * DestComponent = DestActor->GetAbilitySystemComponent();
 	SourceComponent->GetSet<UAbilitySystemTestAttributeSet>()->Health = StartHealth;
 	DestComponent->GetSet<UAbilitySystemTestAttributeSet>()->Health = StartHealth;
 
@@ -3399,8 +3399,8 @@ bool GameplayEffectsTest_StackingHighest(UWorld *World, FAutomationTestBase * Te
 
 	UProperty *StackingProperty = FindFieldChecked<UProperty>(UAbilitySystemTestAttributeSet::StaticClass(), GET_MEMBER_NAME_CHECKED(UAbilitySystemTestAttributeSet, StackingAttribute1));
 
-	UAbilitySystemComponent * SourceComponent = SourceActor->AbilitySystemComponent;
-	UAbilitySystemComponent * DestComponent = DestActor->AbilitySystemComponent;
+	UAbilitySystemComponent * SourceComponent = SourceActor->GetAbilitySystemComponent();
+	UAbilitySystemComponent * DestComponent = DestActor->GetAbilitySystemComponent();
 
 	{
 		ABILITY_LOG_SCOPE(TEXT("Apply StackingEffect"));
@@ -3491,8 +3491,8 @@ bool GameplayEffectsTest_StackingLowest(UWorld *World, FAutomationTestBase * Tes
 
 	UProperty *StackingProperty = FindFieldChecked<UProperty>(UAbilitySystemTestAttributeSet::StaticClass(), GET_MEMBER_NAME_CHECKED(UAbilitySystemTestAttributeSet, StackingAttribute1));
 
-	UAbilitySystemComponent * SourceComponent = SourceActor->AbilitySystemComponent;
-	UAbilitySystemComponent * DestComponent = DestActor->AbilitySystemComponent;
+	UAbilitySystemComponent * SourceComponent = SourceActor->GetAbilitySystemComponent();
+	UAbilitySystemComponent * DestComponent = DestActor->GetAbilitySystemComponent();
 
 	{
 		ABILITY_LOG_SCOPE(TEXT("Apply StackingEffect"));
@@ -3578,8 +3578,8 @@ bool GameplayEffectsTest_StackingUnlimited(UWorld *World, FAutomationTestBase *
 
 	UProperty *StackingProperty = FindFieldChecked<UProperty>(UAbilitySystemTestAttributeSet::StaticClass(), GET_MEMBER_NAME_CHECKED(UAbilitySystemTestAttributeSet, StackingAttribute1));
 
-	UAbilitySystemComponent * SourceComponent = SourceActor->AbilitySystemComponent;
-	UAbilitySystemComponent * DestComponent = DestActor->AbilitySystemComponent;
+	UAbilitySystemComponent * SourceComponent = SourceActor->GetAbilitySystemComponent();
+	UAbilitySystemComponent * DestComponent = DestActor->GetAbilitySystemComponent();
 
 	{
 		ABILITY_LOG_SCOPE(TEXT("Apply StackingEffect"));
@@ -3653,8 +3653,8 @@ bool GameplayEffectsTest_StackingRemoval(UWorld *World, FAutomationTestBase * Te
 
 	UProperty *StackingProperty = FindFieldChecked<UProperty>(UAbilitySystemTestAttributeSet::StaticClass(), GET_MEMBER_NAME_CHECKED(UAbilitySystemTestAttributeSet, StackingAttribute1));
 
-	UAbilitySystemComponent * SourceComponent = SourceActor->AbilitySystemComponent;
-	UAbilitySystemComponent * DestComponent = DestActor->AbilitySystemComponent;
+	UAbilitySystemComponent * SourceComponent = SourceActor->GetAbilitySystemComponent();
+	UAbilitySystemComponent * DestComponent = DestActor->GetAbilitySystemComponent();
 
 	{
 		ABILITY_LOG_SCOPE(TEXT("Apply StackingEffect"));
@@ -3735,8 +3735,8 @@ bool GameplayEffectsTest_StackingNoStack(UWorld *World, FAutomationTestBase * Te
 
 	UProperty *NoStackProperty = FindFieldChecked<UProperty>(UAbilitySystemTestAttributeSet::StaticClass(), GET_MEMBER_NAME_CHECKED(UAbilitySystemTestAttributeSet, NoStackAttribute));
 
-	UAbilitySystemComponent * SourceComponent = SourceActor->AbilitySystemComponent;
-	UAbilitySystemComponent * DestComponent = DestActor->AbilitySystemComponent;
+	UAbilitySystemComponent * SourceComponent = SourceActor->GetAbilitySystemComponent();
+	UAbilitySystemComponent * DestComponent = DestActor->GetAbilitySystemComponent();
 
 	{
 		ABILITY_LOG_SCOPE(TEXT("Apply StackingEffect"));
@@ -3810,8 +3810,8 @@ bool GameplayEffectsTest_StackingCustomCapped(UWorld *World, FAutomationTestBase
 
 	UProperty *StackingProperty = FindFieldChecked<UProperty>(UAbilitySystemTestAttributeSet::StaticClass(), GET_MEMBER_NAME_CHECKED(UAbilitySystemTestAttributeSet, StackingAttribute1));
 
-	UAbilitySystemComponent * SourceComponent = SourceActor->AbilitySystemComponent;
-	UAbilitySystemComponent * DestComponent = DestActor->AbilitySystemComponent;
+	UAbilitySystemComponent * SourceComponent = SourceActor->GetAbilitySystemComponent();
+	UAbilitySystemComponent * DestComponent = DestActor->GetAbilitySystemComponent();
 
 	{
 		ABILITY_LOG_SCOPE(TEXT("Apply StackingEffect"));
@@ -3961,8 +3961,8 @@ bool GameplayEffectsTest_StackingCustomDiminishingReturns(UWorld *World, FAutoma
 
 	UProperty *StackingProperty = FindFieldChecked<UProperty>(UAbilitySystemTestAttributeSet::StaticClass(), GET_MEMBER_NAME_CHECKED(UAbilitySystemTestAttributeSet, StackingAttribute1));
 
-	UAbilitySystemComponent * SourceComponent = SourceActor->AbilitySystemComponent;
-	UAbilitySystemComponent * DestComponent = DestActor->AbilitySystemComponent;
+	UAbilitySystemComponent * SourceComponent = SourceActor->GetAbilitySystemComponent();
+	UAbilitySystemComponent * DestComponent = DestActor->GetAbilitySystemComponent();
 
 	{
 		ABILITY_LOG_SCOPE(TEXT("Apply StackingEffect"));
@@ -4118,8 +4118,8 @@ bool GameplayEffectsTest_StackingDifferentRules(UWorld *World, FAutomationTestBa
 
 	UProperty *StackingProperty = FindFieldChecked<UProperty>(UAbilitySystemTestAttributeSet::StaticClass(), GET_MEMBER_NAME_CHECKED(UAbilitySystemTestAttributeSet, StackingAttribute1));
 
-	UAbilitySystemComponent * SourceComponent = SourceActor->AbilitySystemComponent;
-	UAbilitySystemComponent * DestComponent = DestActor->AbilitySystemComponent;
+	UAbilitySystemComponent * SourceComponent = SourceActor->GetAbilitySystemComponent();
+	UAbilitySystemComponent * DestComponent = DestActor->GetAbilitySystemComponent();
 
 	{
 		ABILITY_LOG_SCOPE(TEXT("Apply StackingEffect"));
@@ -4194,8 +4194,8 @@ bool GameplayEffectsTest_StackingDifferentAttributes(UWorld *World, FAutomationT
 	UProperty *StackingProperty1 = FindFieldChecked<UProperty>(UAbilitySystemTestAttributeSet::StaticClass(), GET_MEMBER_NAME_CHECKED(UAbilitySystemTestAttributeSet, StackingAttribute1));
 	UProperty *StackingProperty2 = FindFieldChecked<UProperty>(UAbilitySystemTestAttributeSet::StaticClass(), GET_MEMBER_NAME_CHECKED(UAbilitySystemTestAttributeSet, StackingAttribute2));
 
-	UAbilitySystemComponent * SourceComponent = SourceActor->AbilitySystemComponent;
-	UAbilitySystemComponent * DestComponent = DestActor->AbilitySystemComponent;
+	UAbilitySystemComponent * SourceComponent = SourceActor->GetAbilitySystemComponent();
+	UAbilitySystemComponent * DestComponent = DestActor->GetAbilitySystemComponent();
 
 	{
 		ABILITY_LOG_SCOPE(TEXT("Apply StackingEffect"));
@@ -4272,8 +4272,8 @@ bool GameplayEffectsTest_StackingCustomTwoRules(UWorld *World, FAutomationTestBa
 
 	UProperty *StackingProperty = FindFieldChecked<UProperty>(UAbilitySystemTestAttributeSet::StaticClass(), GET_MEMBER_NAME_CHECKED(UAbilitySystemTestAttributeSet, StackingAttribute1));
 
-	UAbilitySystemComponent * SourceComponent = SourceActor->AbilitySystemComponent;
-	UAbilitySystemComponent * DestComponent = DestActor->AbilitySystemComponent;
+	UAbilitySystemComponent * SourceComponent = SourceActor->GetAbilitySystemComponent();
+	UAbilitySystemComponent * DestComponent = DestActor->GetAbilitySystemComponent();
 
 	{
 		ABILITY_LOG_SCOPE(TEXT("Apply StackingEffect"));
@@ -4373,8 +4373,8 @@ bool GameplayEffectsTest_StackingCustomTwoAttributes(UWorld *World, FAutomationT
 	UProperty *StackingProperty1 = FindFieldChecked<UProperty>(UAbilitySystemTestAttributeSet::StaticClass(), GET_MEMBER_NAME_CHECKED(UAbilitySystemTestAttributeSet, StackingAttribute1));
 	UProperty *StackingProperty2 = FindFieldChecked<UProperty>(UAbilitySystemTestAttributeSet::StaticClass(), GET_MEMBER_NAME_CHECKED(UAbilitySystemTestAttributeSet, StackingAttribute2));
 
-	UAbilitySystemComponent * SourceComponent = SourceActor->AbilitySystemComponent;
-	UAbilitySystemComponent * DestComponent = DestActor->AbilitySystemComponent;
+	UAbilitySystemComponent * SourceComponent = SourceActor->GetAbilitySystemComponent();
+	UAbilitySystemComponent * DestComponent = DestActor->GetAbilitySystemComponent();
 
 	{
 		ABILITY_LOG_SCOPE(TEXT("Apply StackingEffect"));
@@ -4477,8 +4477,8 @@ bool GameplayEffectsTest_StackingRemovingModifiers(UWorld *World, FAutomationTes
 
 	UProperty *StackingProperty = FindFieldChecked<UProperty>(UAbilitySystemTestAttributeSet::StaticClass(), GET_MEMBER_NAME_CHECKED(UAbilitySystemTestAttributeSet, StackingAttribute1));
 
-	UAbilitySystemComponent * SourceComponent = SourceActor->AbilitySystemComponent;
-	UAbilitySystemComponent * DestComponent = DestActor->AbilitySystemComponent;
+	UAbilitySystemComponent * SourceComponent = SourceActor->GetAbilitySystemComponent();
+	UAbilitySystemComponent * DestComponent = DestActor->GetAbilitySystemComponent();
 
 	// Setup a GE to modify IncomingGEs
 	{
@@ -4579,8 +4579,8 @@ bool GameplayEffectsTest_StackingAddingModifiers(UWorld *World, FAutomationTestB
 
 	UProperty *StackingProperty = FindFieldChecked<UProperty>(UAbilitySystemTestAttributeSet::StaticClass(), GET_MEMBER_NAME_CHECKED(UAbilitySystemTestAttributeSet, StackingAttribute1));
 
-	UAbilitySystemComponent * SourceComponent = SourceActor->AbilitySystemComponent;
-	UAbilitySystemComponent * DestComponent = DestActor->AbilitySystemComponent;
+	UAbilitySystemComponent * SourceComponent = SourceActor->GetAbilitySystemComponent();
+	UAbilitySystemComponent * DestComponent = DestActor->GetAbilitySystemComponent();
 
 	{
 		ABILITY_LOG_SCOPE(TEXT("Apply StackingEffect"));
@@ -4684,8 +4684,8 @@ bool GameplayEffectsTest_ImmunityIncoming(UWorld *World, FAutomationTestBase * T
 
 	UProperty *DamageProperty = FindFieldChecked<UProperty>(UAbilitySystemTestAttributeSet::StaticClass(), GET_MEMBER_NAME_CHECKED(UAbilitySystemTestAttributeSet, Damage));
 
-	UAbilitySystemComponent * SourceComponent = SourceActor->AbilitySystemComponent;
-	UAbilitySystemComponent * DestComponent = DestActor->AbilitySystemComponent;
+	UAbilitySystemComponent * SourceComponent = SourceActor->GetAbilitySystemComponent();
+	UAbilitySystemComponent * DestComponent = DestActor->GetAbilitySystemComponent();
 	SourceComponent->GetSet<UAbilitySystemTestAttributeSet>()->Health = StartHealth;
 	DestComponent->GetSet<UAbilitySystemTestAttributeSet>()->Health = StartHealth;
 
@@ -4742,8 +4742,8 @@ bool GameplayEffectsTest_ImmunityOutgoing(UWorld *World, FAutomationTestBase * T
 
 	UProperty *DamageProperty = FindFieldChecked<UProperty>(UAbilitySystemTestAttributeSet::StaticClass(), GET_MEMBER_NAME_CHECKED(UAbilitySystemTestAttributeSet, Damage));
 
-	UAbilitySystemComponent * SourceComponent = SourceActor->AbilitySystemComponent;
-	UAbilitySystemComponent * DestComponent = DestActor->AbilitySystemComponent;
+	UAbilitySystemComponent * SourceComponent = SourceActor->GetAbilitySystemComponent();
+	UAbilitySystemComponent * DestComponent = DestActor->GetAbilitySystemComponent();
 	SourceComponent->GetSet<UAbilitySystemTestAttributeSet>()->Health = StartHealth;
 	DestComponent->GetSet<UAbilitySystemTestAttributeSet>()->Health = StartHealth;
 
@@ -4803,8 +4803,8 @@ bool GameplayEffectsTest_ImmunityMod(UWorld *World, FAutomationTestBase * Test)
 
 	UProperty *DamageProperty = FindFieldChecked<UProperty>(UAbilitySystemTestAttributeSet::StaticClass(), GET_MEMBER_NAME_CHECKED(UAbilitySystemTestAttributeSet, Damage));
 
-	UAbilitySystemComponent * SourceComponent = SourceActor->AbilitySystemComponent;
-	UAbilitySystemComponent * DestComponent = DestActor->AbilitySystemComponent;
+	UAbilitySystemComponent * SourceComponent = SourceActor->GetAbilitySystemComponent();
+	UAbilitySystemComponent * DestComponent = DestActor->GetAbilitySystemComponent();
 	SourceComponent->GetSet<UAbilitySystemTestAttributeSet>()->Health = StartHealth;
 	DestComponent->GetSet<UAbilitySystemTestAttributeSet>()->Health = StartHealth;
 
@@ -4879,8 +4879,8 @@ bool GameplayEffectsTest_ImmunityBlockedBuff(UWorld *World, FAutomationTestBase
 
 	UProperty *DamageProperty = FindFieldChecked<UProperty>(UAbilitySystemTestAttributeSet::StaticClass(), GET_MEMBER_NAME_CHECKED(UAbilitySystemTestAttributeSet, Damage));
 
-	UAbilitySystemComponent * SourceComponent = SourceActor->AbilitySystemComponent;
-	UAbilitySystemComponent * DestComponent = DestActor->AbilitySystemComponent;
+	UAbilitySystemComponent * SourceComponent = SourceActor->GetAbilitySystemComponent();
+	UAbilitySystemComponent * DestComponent = DestActor->GetAbilitySystemComponent();
 	SourceComponent->GetSet<UAbilitySystemTestAttributeSet>()->Health = StartHealth;
 	DestComponent->GetSet<UAbilitySystemTestAttributeSet>()->Health = StartHealth;
 
@@ -4948,8 +4948,8 @@ bool GameplayEffectsTest_ImmunityBlockedBaseAndBuff(UWorld *World, FAutomationTe
 
 	UProperty *DamageProperty = FindFieldChecked<UProperty>(UAbilitySystemTestAttributeSet::StaticClass(), GET_MEMBER_NAME_CHECKED(UAbilitySystemTestAttributeSet, Damage));
 
-	UAbilitySystemComponent * SourceComponent = SourceActor->AbilitySystemComponent;
-	UAbilitySystemComponent * DestComponent = DestActor->AbilitySystemComponent;
+	UAbilitySystemComponent * SourceComponent = SourceActor->GetAbilitySystemComponent();
+	UAbilitySystemComponent * DestComponent = DestActor->GetAbilitySystemComponent();
 	SourceComponent->GetSet<UAbilitySystemTestAttributeSet>()->Health = StartHealth;
 	DestComponent->GetSet<UAbilitySystemTestAttributeSet>()->Health = StartHealth;
 
@@ -5017,8 +5017,8 @@ bool GameplayEffectsTest_ImmunityActiveGE(UWorld *World, FAutomationTestBase * T
 
 	UProperty *DamageProperty = FindFieldChecked<UProperty>(UAbilitySystemTestAttributeSet::StaticClass(), GET_MEMBER_NAME_CHECKED(UAbilitySystemTestAttributeSet, Damage));
 
-	UAbilitySystemComponent * SourceComponent = SourceActor->AbilitySystemComponent;
-	UAbilitySystemComponent * DestComponent = DestActor->AbilitySystemComponent;
+	UAbilitySystemComponent * SourceComponent = SourceActor->GetAbilitySystemComponent();
+	UAbilitySystemComponent * DestComponent = DestActor->GetAbilitySystemComponent();
 	SourceComponent->GetSet<UAbilitySystemTestAttributeSet>()->Health = StartHealth;
 	DestComponent->GetSet<UAbilitySystemTestAttributeSet>()->Health = StartHealth;
 
@@ -5114,8 +5114,8 @@ bool GameplayEffectsTest_ImmunityActiveMod(UWorld *World, FAutomationTestBase *
 
 	UProperty *DamageProperty = FindFieldChecked<UProperty>(UAbilitySystemTestAttributeSet::StaticClass(), GET_MEMBER_NAME_CHECKED(UAbilitySystemTestAttributeSet, Damage));
 
-	UAbilitySystemComponent * SourceComponent = SourceActor->AbilitySystemComponent;
-	UAbilitySystemComponent * DestComponent = DestActor->AbilitySystemComponent;
+	UAbilitySystemComponent* SourceComponent = SourceActor->GetAbilitySystemComponent();
+	UAbilitySystemComponent* DestComponent = DestActor->GetAbilitySystemComponent();
 	SourceComponent->GetSet<UAbilitySystemTestAttributeSet>()->Health = StartHealth;
 	DestComponent->GetSet<UAbilitySystemTestAttributeSet>()->Health = StartHealth;
 
@@ -5210,8 +5210,8 @@ bool GameplayEffectsTest_ChanceToApplyToTarget(UWorld *World, FAutomationTestBas
 
 	UProperty *DamageProperty = FindFieldChecked<UProperty>(UAbilitySystemTestAttributeSet::StaticClass(), GET_MEMBER_NAME_CHECKED(UAbilitySystemTestAttributeSet, Damage));
 
-	UAbilitySystemComponent * SourceComponent = SourceActor->AbilitySystemComponent;
-	UAbilitySystemComponent * DestComponent = DestActor->AbilitySystemComponent;
+	UAbilitySystemComponent * SourceComponent = SourceActor->GetAbilitySystemComponent();
+	UAbilitySystemComponent * DestComponent = DestActor->GetAbilitySystemComponent();
 
 	SourceComponent->GetSet<UAbilitySystemTestAttributeSet>()->Health = StartHealth;
 	DestComponent->GetSet<UAbilitySystemTestAttributeSet>()->Health = StartHealth;
@@ -5279,8 +5279,8 @@ bool GameplayEffectsTest_ChanceToExecuteOnActiveGEMod(UWorld *World, FAutomation
 
 	UProperty *DamageProperty = FindFieldChecked<UProperty>(UAbilitySystemTestAttributeSet::StaticClass(), GET_MEMBER_NAME_CHECKED(UAbilitySystemTestAttributeSet, Damage));
 
-	UAbilitySystemComponent * SourceComponent = SourceActor->AbilitySystemComponent;
-	UAbilitySystemComponent * DestComponent = DestActor->AbilitySystemComponent;
+	UAbilitySystemComponent * SourceComponent = SourceActor->GetAbilitySystemComponent();
+	UAbilitySystemComponent * DestComponent = DestActor->GetAbilitySystemComponent();
 
 	SourceComponent->GetSet<UAbilitySystemTestAttributeSet>()->Health = StartHealth;
 	DestComponent->GetSet<UAbilitySystemTestAttributeSet>()->Health = StartHealth;
@@ -5396,8 +5396,8 @@ bool GameplayEffectsTest_ChanceToExecuteOnActiveGEImmunity(UWorld *World, FAutom
 
 	UProperty *DamageProperty = FindFieldChecked<UProperty>(UAbilitySystemTestAttributeSet::StaticClass(), GET_MEMBER_NAME_CHECKED(UAbilitySystemTestAttributeSet, Damage));
 
-	UAbilitySystemComponent * SourceComponent = SourceActor->AbilitySystemComponent;
-	UAbilitySystemComponent * DestComponent = DestActor->AbilitySystemComponent;
+	UAbilitySystemComponent * SourceComponent = SourceActor->GetAbilitySystemComponent();
+	UAbilitySystemComponent * DestComponent = DestActor->GetAbilitySystemComponent();
 
 	SourceComponent->GetSet<UAbilitySystemTestAttributeSet>()->Health = StartHealth;
 	DestComponent->GetSet<UAbilitySystemTestAttributeSet>()->Health = StartHealth;
@@ -5507,8 +5507,8 @@ bool GameplayEffectsTest_ChanceToExecuteOnOutgoingGEMod(UWorld *World, FAutomati
 
 	UProperty *DamageProperty = FindFieldChecked<UProperty>(UAbilitySystemTestAttributeSet::StaticClass(), GET_MEMBER_NAME_CHECKED(UAbilitySystemTestAttributeSet, Damage));
 	
-	UAbilitySystemComponent * SourceComponent = SourceActor->AbilitySystemComponent;
-	UAbilitySystemComponent * DestComponent = DestActor->AbilitySystemComponent;
+	UAbilitySystemComponent * SourceComponent = SourceActor->GetAbilitySystemComponent();
+	UAbilitySystemComponent * DestComponent = DestActor->GetAbilitySystemComponent();
 	SourceComponent->GetSet<UAbilitySystemTestAttributeSet>()->Health = StartHealth;
 	DestComponent->GetSet<UAbilitySystemTestAttributeSet>()->Health = StartHealth;
 
@@ -5610,8 +5610,8 @@ bool GameplayEffectsTest_ChanceToExecuteOnOutgoingGEImmunity(UWorld *World, FAut
 
 	UProperty *DamageProperty = FindFieldChecked<UProperty>(UAbilitySystemTestAttributeSet::StaticClass(), GET_MEMBER_NAME_CHECKED(UAbilitySystemTestAttributeSet, Damage));
 
-	UAbilitySystemComponent * SourceComponent = SourceActor->AbilitySystemComponent;
-	UAbilitySystemComponent * DestComponent = DestActor->AbilitySystemComponent;
+	UAbilitySystemComponent * SourceComponent = SourceActor->GetAbilitySystemComponent();
+	UAbilitySystemComponent * DestComponent = DestActor->GetAbilitySystemComponent();
 	SourceComponent->GetSet<UAbilitySystemTestAttributeSet>()->Health = StartHealth;
 	DestComponent->GetSet<UAbilitySystemTestAttributeSet>()->Health = StartHealth;
 
@@ -5703,8 +5703,8 @@ bool GameplayEffectsTest_ChanceToExecuteOnIncomingGEMod(UWorld *World, FAutomati
 
 	UProperty *DamageProperty = FindFieldChecked<UProperty>(UAbilitySystemTestAttributeSet::StaticClass(), GET_MEMBER_NAME_CHECKED(UAbilitySystemTestAttributeSet, Damage));
 
-	UAbilitySystemComponent * SourceComponent = SourceActor->AbilitySystemComponent;
-	UAbilitySystemComponent * DestComponent = DestActor->AbilitySystemComponent;
+	UAbilitySystemComponent * SourceComponent = SourceActor->GetAbilitySystemComponent();
+	UAbilitySystemComponent * DestComponent = DestActor->GetAbilitySystemComponent();
 	SourceComponent->GetSet<UAbilitySystemTestAttributeSet>()->Health = StartHealth;
 	DestComponent->GetSet<UAbilitySystemTestAttributeSet>()->Health = StartHealth;
 
@@ -5806,8 +5806,8 @@ bool GameplayEffectsTest_ChanceToExecuteOnIncomingGEImmunity(UWorld *World, FAut
 
 	UProperty *DamageProperty = FindFieldChecked<UProperty>(UAbilitySystemTestAttributeSet::StaticClass(), GET_MEMBER_NAME_CHECKED(UAbilitySystemTestAttributeSet, Damage));
 
-	UAbilitySystemComponent * SourceComponent = SourceActor->AbilitySystemComponent;
-	UAbilitySystemComponent * DestComponent = DestActor->AbilitySystemComponent;
+	UAbilitySystemComponent * SourceComponent = SourceActor->GetAbilitySystemComponent();
+	UAbilitySystemComponent * DestComponent = DestActor->GetAbilitySystemComponent();
 	SourceComponent->GetSet<UAbilitySystemTestAttributeSet>()->Health = StartHealth;
 	DestComponent->GetSet<UAbilitySystemTestAttributeSet>()->Health = StartHealth;
 
@@ -5898,8 +5898,8 @@ bool GameplayEffectsTest_ModifyChanceToApplyToTarget(UWorld *World, FAutomationT
 
 	UProperty *DamageProperty = FindFieldChecked<UProperty>(UAbilitySystemTestAttributeSet::StaticClass(), GET_MEMBER_NAME_CHECKED(UAbilitySystemTestAttributeSet, Damage));
 
-	UAbilitySystemComponent * SourceComponent = SourceActor->AbilitySystemComponent;
-	UAbilitySystemComponent * DestComponent = DestActor->AbilitySystemComponent;
+	UAbilitySystemComponent * SourceComponent = SourceActor->GetAbilitySystemComponent();
+	UAbilitySystemComponent * DestComponent = DestActor->GetAbilitySystemComponent();
 	SourceComponent->GetSet<UAbilitySystemTestAttributeSet>()->Health = StartHealth;
 	DestComponent->GetSet<UAbilitySystemTestAttributeSet>()->Health = StartHealth;
 
@@ -5957,8 +5957,8 @@ bool GameplayEffectsTest_ModifyChanceToExecuteOnGE(UWorld *World, FAutomationTes
 
 	UProperty *DamageProperty = FindFieldChecked<UProperty>(UAbilitySystemTestAttributeSet::StaticClass(), GET_MEMBER_NAME_CHECKED(UAbilitySystemTestAttributeSet, Damage));
 
-	UAbilitySystemComponent * SourceComponent = SourceActor->AbilitySystemComponent;
-	UAbilitySystemComponent * DestComponent = DestActor->AbilitySystemComponent;
+	UAbilitySystemComponent * SourceComponent = SourceActor->GetAbilitySystemComponent();
+	UAbilitySystemComponent * DestComponent = DestActor->GetAbilitySystemComponent();
 	SourceComponent->GetSet<UAbilitySystemTestAttributeSet>()->Health = StartHealth;
 	DestComponent->GetSet<UAbilitySystemTestAttributeSet>()->Health = StartHealth;
 
diff --git a/Engine/Source/Runtime/GameplayAbilities/Public/Abilities/GameplayAbilityWorldReticle_ActorVisualization.h b/Engine/Source/Runtime/GameplayAbilities/Public/Abilities/GameplayAbilityWorldReticle_ActorVisualization.h
--- a/Engine/Source/Runtime/GameplayAbilities/Public/Abilities/GameplayAbilityWorldReticle_ActorVisualization.h
+++ b/Engine/Source/Runtime/GameplayAbilities/Public/Abilities/GameplayAbilityWorldReticle_ActorVisualization.h
@@ -17,9 +17,11 @@ public:
 
 	void InitializeReticleVisualizationInformation(AActor* VisualizationActor, UMaterialInterface *VisualizationMaterial);
 
+private:
 	/** Hardcoded collision component, so other objects don't think they can collide with the visualization actor */
 	UPROPERTY()
-	TSubobjectPtr<class UCapsuleComponent> CollisionComponent;
+	class UCapsuleComponent* CollisionComponent;
+public:
 
 	UPROPERTY()
 	TArray<UActorComponent*> VisualizationComponents;
@@ -28,4 +30,8 @@ public:
 	virtual void EndPlay(const EEndPlayReason::Type EndPlayReason);
 
 protected:
+
+public:
+	/** Returns CollisionComponent subobject **/
+	FORCEINLINE class UCapsuleComponent* GetCollisionComponent() { return CollisionComponent; }
 };
diff --git a/Engine/Source/Runtime/GameplayAbilities/Public/AbilitySystemComponent.h b/Engine/Source/Runtime/GameplayAbilities/Public/AbilitySystemComponent.h
--- a/Engine/Source/Runtime/GameplayAbilities/Public/AbilitySystemComponent.h
+++ b/Engine/Source/Runtime/GameplayAbilities/Public/AbilitySystemComponent.h
@@ -83,11 +83,10 @@ class GAMEPLAYABILITIES_API UAbilitySystemComponent : public UActorComponent, pu
 
 	/** Adds a new AttributeSet that is a DSO (created by called in their CStor) */
 	template <class T>
-	const T*	AddDefaultSubobjectSet(TSubobjectPtr<T> Subobject)
+	const T*	AddDefaultSubobjectSet(T* Subobject)
 	{
-		const T* Set = Subobject.Get();
-		SpawnedAttributes.Add(Set);
-		return Set;
+		SpawnedAttributes.Add(Subobject);
+		return Subobject;
 	}
 
 	const UAttributeSet* InitStats(TSubclassOf<class UAttributeSet> Attributes, const UDataTable* DataTable);
diff --git a/Engine/Source/Runtime/GameplayAbilities/Public/AbilitySystemTestPawn.h b/Engine/Source/Runtime/GameplayAbilities/Public/AbilitySystemTestPawn.h
--- a/Engine/Source/Runtime/GameplayAbilities/Public/AbilitySystemTestPawn.h
+++ b/Engine/Source/Runtime/GameplayAbilities/Public/AbilitySystemTestPawn.h
@@ -15,9 +15,11 @@ class GAMEPLAYABILITIES_API AAbilitySystemTestPawn : public ADefaultPawn, public
 
 	virtual UAbilitySystemComponent* GetAbilitySystemComponent() const override;
 	
+private:
 	/** DefaultPawn collision component */
-	UPROPERTY(Category = AbilitySystem, VisibleAnywhere, BlueprintReadOnly)
-	TSubobjectPtr<class UAbilitySystemComponent> AbilitySystemComponent;
+	UPROPERTY(Category = AbilitySystem, VisibleAnywhere, BlueprintReadOnly, meta = (AllowPrivateAccess = "true"))
+	class UAbilitySystemComponent* AbilitySystemComponent;
+public:
 
 	UPROPERTY(EditDefaultsOnly, Category=GameplayEffects)
 	FGameplayCueHandler	GameplayCueHandler;
@@ -26,4 +28,8 @@ class GAMEPLAYABILITIES_API AAbilitySystemTestPawn : public ADefaultPawn, public
 	//UGameplayAbilitySet * DefaultAbilitySet;
 
 	static FName AbilitySystemComponentName;
-};
\ No newline at end of file
+
+public:
+	/** Returns AbilitySystemComponent subobject **/
+	FORCEINLINE class UAbilitySystemComponent* GetAbilitySystemComponent() { return AbilitySystemComponent; }
+};
diff --git a/Engine/Source/Runtime/Landscape/Classes/LandscapeGizmoActor.h b/Engine/Source/Runtime/Landscape/Classes/LandscapeGizmoActor.h
--- a/Engine/Source/Runtime/Landscape/Classes/LandscapeGizmoActor.h
+++ b/Engine/Source/Runtime/Landscape/Classes/LandscapeGizmoActor.h
@@ -33,9 +33,11 @@ class ALandscapeGizmoActor : public AActor
 	UPROPERTY(EditAnywhere, transient, Category=Gizmo)
 	class ULandscapeInfo* TargetLandscapeInfo;
 
+private:
 	UPROPERTY()
-	TSubobjectPtr<UBillboardComponent> SpriteComponent;
+	UBillboardComponent* SpriteComponent;
 #endif // WITH_EDITORONLY_DATA
+public:
 
 
 #if WITH_EDITOR
@@ -52,6 +54,12 @@ class ALandscapeGizmoActor : public AActor
 	 * Indicates whether this actor should participate in level bounds calculations
 	 */
 	virtual bool IsLevelBoundsRelevant() const override { return false; }
+
+public:
+#if WITH_EDITORONLY_DATA
+	/** Returns SpriteComponent subobject **/
+	FORCEINLINE UBillboardComponent* GetSpriteComponent() const { return SpriteComponent; }
+#endif
 };
 
 
diff --git a/Engine/Source/Runtime/Landscape/Private/Landscape.cpp b/Engine/Source/Runtime/Landscape/Private/Landscape.cpp
--- a/Engine/Source/Runtime/Landscape/Private/Landscape.cpp
+++ b/Engine/Source/Runtime/Landscape/Private/Landscape.cpp
@@ -507,7 +507,7 @@ ALandscapeProxy::ALandscapeProxy(const FObjectInitializer& ObjectInitializer)
 	bReplicateMovement = false;
 	bCanBeDamaged = false;
 
-	TSubobjectPtr<USceneComponent> SceneComponent = ObjectInitializer.CreateDefaultSubobject<USceneComponent>(this, TEXT("RootComponent0"));
+	USceneComponent* SceneComponent = ObjectInitializer.CreateDefaultSubobject<USceneComponent>(this, TEXT("RootComponent0"));
 	RootComponent = SceneComponent;
 	RootComponent->RelativeScale3D = FVector(128.0f, 128.0f, 256.0f); // Old default scale, preserved for compatibility. See ULandscapeEditorObject::NewLandscape_Scale
 	RootComponent->Mobility = EComponentMobility::Static;
diff --git a/Engine/Source/Runtime/Landscape/Private/LandscapeGizmoActor.cpp b/Engine/Source/Runtime/Landscape/Private/LandscapeGizmoActor.cpp
--- a/Engine/Source/Runtime/Landscape/Private/LandscapeGizmoActor.cpp
+++ b/Engine/Source/Runtime/Landscape/Private/LandscapeGizmoActor.cpp
@@ -467,7 +467,7 @@ ALandscapeGizmoActor::ALandscapeGizmoActor(const FObjectInitializer& ObjectIniti
 	}
 #endif
 
-	TSubobjectPtr<USceneComponent> SceneComponent = ObjectInitializer.CreateDefaultSubobject<USceneComponent>(this, TEXT("RootComponent0"));
+	USceneComponent* SceneComponent = ObjectInitializer.CreateDefaultSubobject<USceneComponent>(this, TEXT("RootComponent0"));
 	RootComponent = SceneComponent;
 	RootComponent->Mobility = EComponentMobility::Static;
 
@@ -536,7 +536,7 @@ ALandscapeGizmoActiveActor::ALandscapeGizmoActiveActor(const FObjectInitializer&
 	}
 #endif // WITH_EDITORONLY_DATA
 
-	TSubobjectPtr<ULandscapeGizmoRenderComponent> LandscapeGizmoRenderComponent = ObjectInitializer.CreateDefaultSubobject<ULandscapeGizmoRenderComponent>(this, TEXT("GizmoRendererComponent0"));
+	ULandscapeGizmoRenderComponent* LandscapeGizmoRenderComponent = ObjectInitializer.CreateDefaultSubobject<ULandscapeGizmoRenderComponent>(this, TEXT("GizmoRendererComponent0"));
 	LandscapeGizmoRenderComponent->SetCollisionProfileName(UCollisionProfile::BlockAllDynamic_ProfileName);
 
 	RootComponent = LandscapeGizmoRenderComponent;
