Removed redundant filter from Radius targeting actor.

diff --git a/Engine/Source/Runtime/GameplayAbilities/Private/Abilities/GameplayAbilityTargetActor_Radius.cpp b/Engine/Source/Runtime/GameplayAbilities/Private/Abilities/GameplayAbilityTargetActor_Radius.cpp
--- a/Engine/Source/Runtime/GameplayAbilities/Private/Abilities/GameplayAbilityTargetActor_Radius.cpp
+++ b/Engine/Source/Runtime/GameplayAbilities/Private/Abilities/GameplayAbilityTargetActor_Radius.cpp
@@ -6,7 +6,7 @@
 #include "Engine/World.h"
 #include "Runtime/Engine/Public/Net/UnrealNetwork.h"
 #include "GameplayAbilityTargetTypes.h"
-
+#pragma optimize("",off)
 // --------------------------------------------------------------------------------------------------------------------------------------------------------
 //
 //	AGameplayAbilityTargetActor_Radius
@@ -69,7 +69,7 @@ TArray<TWeakObjectPtr<AActor> >	AGameplayAbilityTargetActor_Radius::PerformOverl
 	{
 		//Should this check to see if these pawns are in the AimTarget list?
 		APawn* PawnActor = Cast<APawn>(Overlaps[i].GetActor());
-		if (PawnActor && FilterHandle.FilterPassesForActor(PawnActor))
+		if (PawnActor && Filter.FilterPassesForActor(PawnActor))
 		{
 			HitActors.Add(PawnActor);
 		}
diff --git a/Engine/Source/Runtime/GameplayAbilities/Public/Abilities/GameplayAbilityTargetActor_Radius.h b/Engine/Source/Runtime/GameplayAbilities/Public/Abilities/GameplayAbilityTargetActor_Radius.h
--- a/Engine/Source/Runtime/GameplayAbilities/Public/Abilities/GameplayAbilityTargetActor_Radius.h
+++ b/Engine/Source/Runtime/GameplayAbilities/Public/Abilities/GameplayAbilityTargetActor_Radius.h
@@ -19,10 +19,6 @@ public:
 	
 	virtual void ConfirmTargetingAndContinue() override;
 
-	/** Filter to decide who can be selected. */
-	UPROPERTY(BlueprintReadWrite, EditAnywhere, meta = (ExposeOnSpawn = true), Category = Targeting)
-	FGameplayTargetDataFilterHandle FilterHandle;
-
 	/** Radius of target acquisition around the ability's start location. */
 	UPROPERTY(BlueprintReadWrite, EditAnywhere, meta = (ExposeOnSpawn = true), Category = Radius)
 	float Radius;
