Replaced deprecated FPostConstructInitializeProperties with FObjectInitializer.

diff --git a/Engine/Source/Runtime/GameplayAbilities/Private/TickableAttributeSetInterface.cpp b/Engine/Source/Runtime/GameplayAbilities/Private/TickableAttributeSetInterface.cpp
--- a/Engine/Source/Runtime/GameplayAbilities/Private/TickableAttributeSetInterface.cpp
+++ b/Engine/Source/Runtime/GameplayAbilities/Private/TickableAttributeSetInterface.cpp
@@ -3,7 +3,7 @@
 #include "AbilitySystemPrivatePCH.h"
 #include "TickableAttributeSetInterface.h"
 
-UTickableAttributeSetInterface::UTickableAttributeSetInterface(const FPostConstructInitializeProperties& PCIP)
-	: Super(PCIP)
+UTickableAttributeSetInterface::UTickableAttributeSetInterface(const FObjectInitializer& ObjectInitializer)
+: Super(ObjectInitializer)
 {
 }
