WIP - First pass on GameplayCueManager/GameplayCueNotify

diff --git a/Engine/Source/Runtime/GameplayAbilities/Private/Abilities/GameplayAbility.cpp b/Engine/Source/Runtime/GameplayAbilities/Private/Abilities/GameplayAbility.cpp
--- a/Engine/Source/Runtime/GameplayAbilities/Private/Abilities/GameplayAbility.cpp
+++ b/Engine/Source/Runtime/GameplayAbilities/Private/Abilities/GameplayAbility.cpp
@@ -623,16 +623,16 @@ void UGameplayAbility::TaskEnded(UAbilityTask* Task)
  *	work with the PredictionKey system.
  */
 
-void UGameplayAbility::K2_ExecuteGameplayCue(FGameplayTag GameplayCueTag)
+void UGameplayAbility::K2_ExecuteGameplayCue(FGameplayTag GameplayCueTag, FGameplayEffectContextHandle Context)
 {
 	check(CurrentActorInfo);
-	CurrentActorInfo->AbilitySystemComponent->ExecuteGameplayCue(GameplayCueTag, CurrentActivationInfo.GetPredictionKeyForNewAction());
+	CurrentActorInfo->AbilitySystemComponent->ExecuteGameplayCue(GameplayCueTag, CurrentActivationInfo.GetPredictionKeyForNewAction(), Context);
 }
 
-void UGameplayAbility::K2_AddGameplayCue(FGameplayTag GameplayCueTag)
+void UGameplayAbility::K2_AddGameplayCue(FGameplayTag GameplayCueTag, FGameplayEffectContextHandle Context)
 {
 	check(CurrentActorInfo);
-	CurrentActorInfo->AbilitySystemComponent->AddGameplayCue(GameplayCueTag, CurrentActivationInfo.GetPredictionKeyForNewAction());
+	CurrentActorInfo->AbilitySystemComponent->AddGameplayCue(GameplayCueTag, CurrentActivationInfo.GetPredictionKeyForNewAction(), Context);
 }
 
 void UGameplayAbility::K2_RemoveGameplayCue(FGameplayTag GameplayCueTag)
@@ -641,6 +641,22 @@ void UGameplayAbility::K2_RemoveGameplayCue(FGameplayTag GameplayCueTag)
 	CurrentActorInfo->AbilitySystemComponent->RemoveGameplayCue(GameplayCueTag, CurrentActivationInfo.GetPredictionKeyForNewAction());
 }
 
+FGameplayEffectContextHandle UGameplayAbility::GetContextFromOwner(FGameplayAbilityTargetDataHandle OptionalTargetData) const
+{
+	check(CurrentActorInfo);
+	FGameplayEffectContextHandle Context = CurrentActorInfo->AbilitySystemComponent->GetEffectContext();
+	
+	for (auto Data : OptionalTargetData.Data)
+	{
+		if (Data.IsValid())
+		{
+			Data->AddTargetDataToContext(Context);
+		}
+	}
+
+	return Context;
+}
+
 int32 UGameplayAbility::GetAbilityLevel() const
 {
 	check(IsInstantiated()); // You should not call this on non instanced abilities.
diff --git a/Engine/Source/Runtime/GameplayAbilities/Private/AbilitySystemBlueprintLibrary.cpp b/Engine/Source/Runtime/GameplayAbilities/Private/AbilitySystemBlueprintLibrary.cpp
--- a/Engine/Source/Runtime/GameplayAbilities/Private/AbilitySystemBlueprintLibrary.cpp
+++ b/Engine/Source/Runtime/GameplayAbilities/Private/AbilitySystemBlueprintLibrary.cpp
@@ -312,3 +312,13 @@ FTransform UAbilitySystemBlueprintLibrary::GetInstigatorTransform(FGameplayCuePa
 	ABILITY_LOG(Warning, TEXT("UAbilitySystemBlueprintLibrary::GetInstigatorTransform called on GameplayCue with no valid instigator"));
 	return FTransform::Identity;
 }
+
+FVector UAbilitySystemBlueprintLibrary::GetOrigin(FGameplayCueParameters Parameters)
+{
+	if (Parameters.EffectContext.HasOrigin())
+	{
+		return Parameters.EffectContext.GetOrigin();
+	}
+
+	return FVector::ZeroVector;
+}
\ No newline at end of file
diff --git a/Engine/Source/Runtime/GameplayAbilities/Private/AbilitySystemComponent.cpp b/Engine/Source/Runtime/GameplayAbilities/Private/AbilitySystemComponent.cpp
--- a/Engine/Source/Runtime/GameplayAbilities/Private/AbilitySystemComponent.cpp
+++ b/Engine/Source/Runtime/GameplayAbilities/Private/AbilitySystemComponent.cpp
@@ -6,6 +6,7 @@
 #include "GameplayCueInterface.h"
 #include "Abilities/GameplayAbility.h"
 #include "Abilities/Tasks/AbilityTask.h"
+#include "GameplayCueManager.h"
 
 #include "Net/UnrealNetwork.h"
 #include "Engine/ActorChannel.h"
@@ -197,12 +198,17 @@ FGameplayEffectContextHandle UAbilitySystemComponent::GetEffectContext() const
 }
 
 /** This is a helper function used in automated testing, I'm not sure how useful it will be to gamecode or blueprints */
-FActiveGameplayEffectHandle UAbilitySystemComponent::ApplyGameplayEffectToTarget(UGameplayEffect *GameplayEffect, UAbilitySystemComponent *Target, float Level, FModifierQualifier BaseQualifier)
+FActiveGameplayEffectHandle UAbilitySystemComponent::ApplyGameplayEffectToTarget(UGameplayEffect *GameplayEffect, UAbilitySystemComponent *Target, float Level, FGameplayEffectContextHandle Context, FModifierQualifier BaseQualifier)
 {
 	check(GameplayEffect);
 	if (HasNetworkAuthorityToApplyGameplayEffect(BaseQualifier))
 	{
-		FGameplayEffectSpec	Spec(GameplayEffect, GetEffectContext(), Level, GetCurveDataOverride());
+		if (!Context.IsValid())
+		{
+			Context = GetEffectContext();
+		}
+
+		FGameplayEffectSpec	Spec(GameplayEffect, Context, Level, GetCurveDataOverride());
 		return ApplyGameplayEffectSpecToTarget(Spec, Target, BaseQualifier);
 	}
 
@@ -497,13 +503,6 @@ void UAbilitySystemComponent::InvokeGameplayCueEvent(const FGameplayEffectSpec &
 		ABILITY_LOG(Warning, TEXT("InvokeGameplayCueEvent: %s"), *Spec.ToSimpleString());
 	}
 
-	IGameplayCueInterface* GameplayCueInterface = Cast<IGameplayCueInterface>(ActorAvatar);
-	if (!GameplayCueInterface)
-	{
-		ABILITY_LOG(Warning, TEXT("InvokeGameplayCueEvent %s on Actor %s that is not IGameplayCueInterface"), *Spec.ToSimpleString(), ActorAvatar ? *ActorAvatar->GetName() : TEXT("NULL"));
-		return;
-	}
-
 	// FIXME: Replication of level not finished
 	float ExecuteLevel =  (Spec.ModifierLevel.IsValid() && Spec.ModifierLevel.Get()->IsValid()) ? Spec.ModifierLevel.Get()->GetLevel() : 1.f;
 
@@ -529,7 +528,8 @@ void UAbilitySystemComponent::InvokeGameplayCueEvent(const FGameplayEffectSpec &
 		}
 
 		CueParameters.NormalizedMagnitude = CueInfo.NormalizeLevel(ExecuteLevel);
-		GameplayCueInterface->HandleGameplayCues(ActorAvatar, CueInfo.GameplayCueTags, EventType, CueParameters);
+
+		UAbilitySystemGlobals::Get().GetGameplayCueManager()->HandleGameplayCues(ActorAvatar, CueInfo.GameplayCueTags, EventType, CueParameters);
 
 		if (DebugGameplayCues && Spec.EffectContext.GetHitResult())
 		{
@@ -539,6 +539,33 @@ void UAbilitySystemComponent::InvokeGameplayCueEvent(const FGameplayEffectSpec &
 	}
 }
 
+void UAbilitySystemComponent::InvokeGameplayCueEvent(const FGameplayTag GameplayCueTag, EGameplayCueEvent::Type EventType, FGameplayEffectContextHandle EffectContext)
+{
+	AActor* ActorAvatar = AbilityActorInfo->AvatarActor.Get();
+	AActor* ActorOwner = AbilityActorInfo->OwnerActor.Get();
+	IGameplayCueInterface* GameplayCueInterface = Cast<IGameplayCueInterface>(ActorAvatar);
+	if (!GameplayCueInterface)
+	{
+		return;
+	}
+
+	FGameplayCueParameters CueParameters;
+
+	if (EffectContext.IsValid())
+	{
+		CueParameters.EffectContext = EffectContext;
+	}
+	else
+	{
+		CueParameters.EffectContext.AddInstigator(ActorOwner, ActorAvatar); // By default use the owner and avatar as the instigator and causer
+	}
+
+	CueParameters.NormalizedMagnitude = 1.f;
+	CueParameters.RawMagnitude = 0.f;
+
+	UAbilitySystemGlobals::Get().GetGameplayCueManager()->HandleGameplayCues(ActorAvatar, GameplayCueTag, EventType, CueParameters);
+}
+
 void UAbilitySystemComponent::ExecuteGameplayCue(const FGameplayTag GameplayCueTag, FPredictionKey PredictionKey, FGameplayEffectContextHandle EffectContext)
 {
 	if (IsOwnerActorAuthoritative())
@@ -579,33 +606,6 @@ void UAbilitySystemComponent::RemoveGameplayCue(const FGameplayTag GameplayCueTa
 	}
 }
 
-void UAbilitySystemComponent::InvokeGameplayCueEvent(const FGameplayTag GameplayCueTag, EGameplayCueEvent::Type EventType, FGameplayEffectContextHandle EffectContext)
-{
-	AActor* ActorAvatar = AbilityActorInfo->AvatarActor.Get();
-	AActor* ActorOwner = AbilityActorInfo->OwnerActor.Get();
-	IGameplayCueInterface* GameplayCueInterface = Cast<IGameplayCueInterface>(ActorAvatar);
-	if (!GameplayCueInterface)
-	{
-		return;
-	}
-
-	FGameplayCueParameters CueParameters;
-
-	if (EffectContext.IsValid())
-	{
-		CueParameters.EffectContext = EffectContext;
-	}
-	else
-	{
-		CueParameters.EffectContext.AddInstigator(ActorOwner, ActorAvatar); // By default use the owner and avatar as the instigator and causer
-	}
-	
-	CueParameters.NormalizedMagnitude = 1.f;
-	CueParameters.RawMagnitude = 0.f;
-
-	GameplayCueInterface->HandleGameplayCue(ActorAvatar, GameplayCueTag, EventType, CueParameters);
-}
-
 void UAbilitySystemComponent::NetMulticast_InvokeGameplayCueExecuted_FromSpec_Implementation(const FGameplayEffectSpec Spec, FPredictionKey PredictionKey)
 {
 	if (IsOwnerActorAuthoritative() || PredictionKey.IsValidKey() == false)
diff --git a/Engine/Source/Runtime/GameplayAbilities/Private/AbilitySystemGlobals.cpp b/Engine/Source/Runtime/GameplayAbilities/Private/AbilitySystemGlobals.cpp
--- a/Engine/Source/Runtime/GameplayAbilities/Private/AbilitySystemGlobals.cpp
+++ b/Engine/Source/Runtime/GameplayAbilities/Private/AbilitySystemGlobals.cpp
@@ -12,6 +12,7 @@ UAbilitySystemGlobals::UAbilitySystemGlobals(const FObjectInitializer& ObjectIni
 : Super(ObjectInitializer)
 {
 	AbilitySystemGlobalsClassName = FStringClassReference(TEXT("/Script/GameplayAbilities.AbilitySystemGlobals"));
+	GlobalGameplayCueManagerName = FStringClassReference(TEXT("/Script/GameplayAbilities.GameplayCueManager"));
 
 #if WITH_EDITORONLY_DATA
 	RegisteredReimportCallback = false;
@@ -24,6 +25,8 @@ void UAbilitySystemGlobals::InitGlobalData()
 	GetGlobalAttributeMetaDataTable();
 	
 	InitAtributeDefaults();
+
+	GetGameplayCueManager();
 }
 
 UCurveTable * UAbilitySystemGlobals::GetGlobalCurveTable()
@@ -88,7 +91,7 @@ FGameplayEffectContext* UAbilitySystemGlobals::AllocGameplayEffectContext() cons
 /** This is just some syntax sugar to avoid calling gode to have to IGameplayAbilitiesModule::Get() */
 UAbilitySystemGlobals& UAbilitySystemGlobals::Get()
 {
-	static UAbilitySystemGlobals * GlobalPtr = IGameplayAbilitiesModule::Get().GetAbilitySystemGlobals();
+	static UAbilitySystemGlobals* GlobalPtr = IGameplayAbilitiesModule::Get().GetAbilitySystemGlobals();
 	check(GlobalPtr == IGameplayAbilitiesModule::Get().GetAbilitySystemGlobals());
 	check(GlobalPtr);
 
@@ -183,4 +186,20 @@ void UAbilitySystemGlobals::InitAtributeDefaults()
 		AllocAttributeSetInitter();
 		GlobalAttributeSetInitter->PreloadAttributeSetData(GlobalAttributeDefaultsTable);
 	}
-}
\ No newline at end of file
+}
+
+// --------------------------------------------------------------------
+
+UGameplayCueManager* UAbilitySystemGlobals::GetGameplayCueManager()
+{
+	if (GlobalGameplayCueManager == nullptr)
+	{
+		GlobalGameplayCueManager = LoadObject<UGameplayCueManager>(NULL, *GlobalGameplayCueManagerName.ToString(), NULL, LOAD_None, NULL);
+		if (GameplayCueNotifyPaths.Num() > 0)
+		{
+			GlobalGameplayCueManager->LoadObjectLibraryFromPaths( GameplayCueNotifyPaths, false );
+		}
+	}
+
+	return GlobalGameplayCueManager;
+}
diff --git a/Engine/Source/Runtime/GameplayAbilities/Private/GameplayAbilityTargetTypes.cpp b/Engine/Source/Runtime/GameplayAbilities/Private/GameplayAbilityTargetTypes.cpp
--- a/Engine/Source/Runtime/GameplayAbilities/Private/GameplayAbilityTargetTypes.cpp
+++ b/Engine/Source/Runtime/GameplayAbilities/Private/GameplayAbilityTargetTypes.cpp
@@ -17,15 +17,8 @@ TArray<FActiveGameplayEffectHandle> FGameplayAbilityTargetData::ApplyGameplayEff
 		Level,							// FIXME: Leveling
 		NULL							// FIXME: CurveData override... should we just remove this?
 		);
-	if (HasHitResult())
-	{
-		SpecToApply.EffectContext.AddHitResult(*GetHitResult());
-	}
-
-	if (HasOrigin())
-	{
-		SpecToApply.EffectContext.AddOrigin(GetOrigin().GetLocation());
-	}
+	
+	AddTargetDataToContext(EffectContext);
 
 	TArray<TWeakObjectPtr<AActor> > Actors = GetActors();
 	TArray<FActiveGameplayEffectHandle>	AppliedHandles;
@@ -46,6 +39,19 @@ TArray<FActiveGameplayEffectHandle> FGameplayAbilityTargetData::ApplyGameplayEff
 	return AppliedHandles;
 }
 
+void FGameplayAbilityTargetData::AddTargetDataToContext(FGameplayEffectContextHandle& Context)
+{
+	if (HasHitResult())
+	{
+		Context.AddHitResult(*GetHitResult());
+	}
+
+	if (HasOrigin())
+	{
+		Context.AddOrigin(GetOrigin().GetLocation());
+	}
+}
+
 FString FGameplayAbilityTargetData::ToString() const
 {
 	return TEXT("BASE CLASS");
diff --git a/Engine/Source/Runtime/GameplayAbilities/Private/GameplayCueManager.cpp b/Engine/Source/Runtime/GameplayAbilities/Private/GameplayCueManager.cpp
new file mode 100644
--- /dev/null
+++ b/Engine/Source/Runtime/GameplayAbilities/Private/GameplayCueManager.cpp
@@ -0,0 +1,278 @@
+// Copyright 1998-2014 Epic Games, Inc. All Rights Reserved.
+
+#include "AbilitySystemPrivatePCH.h"
+#include "AssetRegistryModule.h"
+#include "GameplayCueManager.h"
+
+#if WITH_EDITOR
+#include "UnrealEd.h"
+#endif
+
+UGameplayCueManager::UGameplayCueManager(const class FObjectInitializer& PCIP)
+: Super(PCIP)
+{
+	RegisteredEditorCallbacks = false;
+}
+
+
+void UGameplayCueManager::HandleGameplayCues(AActor* TargetActor, const FGameplayTagContainer& GameplayCueTags, EGameplayCueEvent::Type EventType, FGameplayCueParameters Parameters)
+{
+	for (auto It = GameplayCueTags.CreateConstIterator(); It; ++It)
+	{
+		HandleGameplayCue(TargetActor, *It, EventType, Parameters);
+	}
+}
+
+void UGameplayCueManager::HandleGameplayCue(AActor* TargetActor, FGameplayTag GameplayCueTag, EGameplayCueEvent::Type EventType, FGameplayCueParameters Parameters)
+{
+	int32 DataIdx = GameplayCueDataMap.FindChecked(GameplayCueTag);	// We should always find data here. We prepopulate the entire map with every gameplaycue tag when we init.
+	HandleGameplayCueNotify_Internal(TargetActor, DataIdx, EventType, Parameters);
+
+	IGameplayCueInterface* GameplayCueInterface = Cast<IGameplayCueInterface>(TargetActor);
+	if (GameplayCueInterface)
+	{
+		GameplayCueInterface->HandleGameplayCue(TargetActor, GameplayCueTag, EventType, Parameters);
+	}
+}
+
+void UGameplayCueManager::HandleGameplayCueNotify_Internal(AActor* TargetActor, int32 DataIdx, EGameplayCueEvent::Type EventType, FGameplayCueParameters Parameters)
+{	
+	if (DataIdx != INDEX_NONE)
+	{
+		check(GameplayCueData.IsValidIndex(DataIdx));
+
+		FGameplayCueNotifyData& CueData = GameplayCueData[DataIdx];
+
+		if (CueData.LoadedGameplayCueNotify == nullptr)
+		{
+			CueData.LoadedGameplayCueNotify = FindObject<UGameplayCueNotify>(nullptr, *CueData.GameplayCueNotifyObj.ToString());
+			if (CueData.LoadedGameplayCueNotify == nullptr)
+			{
+				StreamableManager.SimpleAsyncLoad(CueData.GameplayCueNotifyObj);
+				ABILITY_LOG(Warning, TEXT("GameplayCueNotify %s was not loaded when GameplayCue was invoked. Starting async loading."), *CueData.GameplayCueNotifyObj.ToString());
+				return;
+			}
+		}
+
+		UGameplayCueNotify* SpawnedInstanced = nullptr;
+
+		if (auto InnerMap = NotifyMap.Find(TargetActor))
+		{
+			if (auto WeakPtrPtr = InnerMap->Find(DataIdx))
+			{
+				SpawnedInstanced = WeakPtrPtr->Get();
+			}
+		}
+
+		if (SpawnedInstanced == nullptr && CueData.LoadedGameplayCueNotify->NeedsInstanceForEvent(EventType))
+		{
+			// We don't have an instance for this, and we need one, so make one
+			SpawnedInstanced = static_cast<UGameplayCueNotify*>(StaticDuplicateObject(CueData.LoadedGameplayCueNotify, this, TEXT("None"), ~RF_RootSet));
+
+			auto& InnerMap = NotifyMap.FindOrAdd(TargetActor);
+
+			InnerMap.Add(DataIdx) = SpawnedInstanced;
+
+			// Fixme: this will grow unbounded
+			InstantiatedObjects.Add(SpawnedInstanced);
+		}
+
+		UGameplayCueNotify* NotifyToExecute = SpawnedInstanced ? SpawnedInstanced : CueData.LoadedGameplayCueNotify;
+		NotifyToExecute->HandleGameplayCue(TargetActor, EventType, Parameters);
+
+		if (NotifyToExecute->IsOverride == false)
+		{
+			HandleGameplayCueNotify_Internal(TargetActor, CueData.ParentDataIdx, EventType, Parameters);
+		}
+	}
+}
+
+void UGameplayCueManager::EndGameplayCuesFor(AActor* TargetActor)
+{
+	TMap<int32, TWeakObjectPtr<UGameplayCueNotify> > FoundMap;	
+	if (NotifyMap.RemoveAndCopyValue(TargetActor, FoundMap))
+	{
+		for (auto It = FoundMap.CreateConstIterator(); It; ++It)
+		{
+			UGameplayCueNotify* InstancedCue =  It.Value().Get();
+			if (InstancedCue)
+			{
+				InstancedCue->OnOwnerDestroyed();
+			}
+		}
+	}
+}
+
+// ------------------------------------------------------------------------
+
+void UGameplayCueManager::LoadObjectLibraryFromPaths(const TArray<FString>& Paths, bool bFullyLoad)
+{
+	DECLARE_SCOPE_CYCLE_COUNTER(TEXT("Loading Library"), STAT_ObjectLibrary, STATGROUP_LoadTime);
+
+#if WITH_EDITOR
+	FFormatNamedArguments Args;
+	FScopedSlowTask SlowTask(0, FText::Format(NSLOCTEXT("AbilitySystemEditor", "BeginLoadingGameplayCueNotify", "Loading GameplayCue Library"), Args));
+	SlowTask.MakeDialog();
+#endif
+
+	if (!GameplayCueNotifyObjectLibrary)
+	{
+		GameplayCueNotifyObjectLibrary = UObjectLibrary::CreateLibrary(UGameplayCueNotify::StaticClass(), true, GIsEditor);
+	}
+
+	FScopeCycleCounterUObject PreloadScope(GameplayCueNotifyObjectLibrary);
+
+	GameplayCueNotifyObjectLibrary->LoadBlueprintAssetDataFromPaths(Paths);
+	GameplayCueNotifyObjectLibrary->LoadAssetDataFromPaths(Paths);	
+
+	if (bFullyLoad)
+	{
+#if STATS
+		FString PerfMessage = FString::Printf(TEXT("Fully Loaded GameplayCueNotify object library"));
+		SCOPE_LOG_TIME_IN_SECONDS(*PerfMessage, nullptr)
+#endif
+		GameplayCueNotifyObjectLibrary->LoadAssetsFromAssetData();
+	}
+
+	// ---------------------------------------------------------
+	// Look for GameplayCueNotifies that handle events
+	// ---------------------------------------------------------
+	
+	TArray<FAssetData> AssetDatas;
+	GameplayCueNotifyObjectLibrary->GetAssetDataList(AssetDatas);
+
+	IGameplayTagsModule& GameplayTagsModule = IGameplayTagsModule::Get();
+
+	GameplayCueData.Empty();
+
+	TSet<FGameplayTag>	FoundTags;
+
+	for (FAssetData Data: AssetDatas)
+	{
+		const FString* FoundGameplayTag = Data.TagsAndValues.Find(GET_MEMBER_NAME_CHECKED(UGameplayCueNotify, GameplayCueName));
+		if (FoundGameplayTag)
+		{
+			ABILITY_LOG(Warning, TEXT("Found: %s"), **FoundGameplayTag);
+
+			FGameplayTag  GameplayCueTag = GameplayTagsModule.GetGameplayTagsManager().RequestGameplayTag(FName(**FoundGameplayTag), false);
+			if (GameplayCueTag.IsValid())
+			{
+				FoundTags.Add(GameplayCueTag);
+
+				// Add a new NotifyData entry to our flat list for this one
+				FGameplayCueNotifyData NewData;
+				FStringAssetReference StringRef;
+				StringRef.AssetLongPathname = Data.ObjectPath.ToString();
+
+				TAssetSubclassOf<UGameplayCueNotify> AssetRef;
+				AssetRef = StringRef;
+				
+				NewData.GameplayCueNotifyObj = StringRef;
+				int32 NewIdx = GameplayCueData.Add(NewData);
+
+				// Add an entry into the acceleration map
+				GameplayCueDataMap.FindOrAdd(GameplayCueTag) = NewIdx;
+			}
+			else
+			{
+				ABILITY_LOG(Warning, TEXT("Found GameplayCue tag %s in asset %s but there is no corresponding tag in the GameplayTagMAnager."), **FoundGameplayTag, *Data.PackageName.ToString());
+			}
+		}
+	}
+
+	// ---------------------------------------------------------
+	//	Build up the rest of the acceleration map: every GameplayCue tag should have an entry in the map that points to the index into GameplayCueData to use when it is invoked.
+	//	(or to -1 if no GameplayCueNotify is associated with that tag)
+	// 
+	// ---------------------------------------------------------
+	
+	FGameplayTagContainer AllGameplayCueTags = GameplayTagsModule.GetGameplayTagsManager().RequestGameplayTagChildren(BaseGameplayCueTag());
+
+	GameplayCueDataMap.Add(BaseGameplayCueTag()) = INDEX_NONE;
+
+	for (FGameplayTag ThisGameplayCueTag : AllGameplayCueTags)
+	{
+		if (GameplayCueDataMap.Contains(ThisGameplayCueTag))
+		{
+			continue;
+		}
+
+		FGameplayTag Parent = ThisGameplayCueTag.RequestDirectParent();
+
+		GameplayCueDataMap.Add(ThisGameplayCueTag) = GameplayCueDataMap.FindChecked(Parent);
+	}
+
+	// PrintGameplayCueNotifyMap();
+
+#if WITH_EDITOR
+	if (!RegisteredEditorCallbacks)
+	{
+		FAssetRegistryModule& AssetRegistryModule = FModuleManager::LoadModuleChecked<FAssetRegistryModule>("AssetRegistry");
+		AssetRegistryModule.Get().OnInMemoryAssetCreated().AddUObject(this, &UGameplayCueManager::HandleAssetAdded);
+		AssetRegistryModule.Get().OnInMemoryAssetDeleted().AddUObject(this, &UGameplayCueManager::HandleAssetDeleted);
+		RegisteredEditorCallbacks = true;
+	}
+#endif
+
+}
+
+int32 UGameplayCueManager::FinishLoadingGameplayCueNotifies()
+{
+	int32 NumLoadeded = 0;
+	return NumLoadeded;
+}
+
+void UGameplayCueManager::BeginLoadingGameplayCueNotify(FGameplayTag GameplayCueTag)
+{
+
+}
+
+#if WITH_EDITOR
+
+void UGameplayCueManager::HandleAssetAdded(UObject *Object)
+{
+
+}
+
+/** Handles cleaning up an object library if it matches the passed in object */
+void UGameplayCueManager::HandleAssetDeleted(UObject *Object)
+{
+
+}
+
+#endif
+
+void UGameplayCueManager::PrintGameplayCueNotifyMap()
+{
+	FGameplayTagContainer AllGameplayCueTags = IGameplayTagsModule::Get().GetGameplayTagsManager().RequestGameplayTagChildren(BaseGameplayCueTag());
+
+	for (FGameplayTag ThisGameplayCueTag : AllGameplayCueTags)
+	{
+		int32 idx = GameplayCueDataMap.FindChecked(ThisGameplayCueTag);
+		if (idx != INDEX_NONE)
+		{
+			ABILITY_LOG(Warning, TEXT("   %s -> %d"), *ThisGameplayCueTag.ToString(), idx);
+		}
+		else
+		{
+			ABILITY_LOG(Warning, TEXT("   %s -> unmapped"), *ThisGameplayCueTag.ToString());
+		}
+	}
+}
+
+FGameplayTag UGameplayCueManager::BaseGameplayCueTag()
+{
+	static FGameplayTag  BaseGameplayCueTag = IGameplayTagsModule::Get().GetGameplayTagsManager().RequestGameplayTag(TEXT("GameplayCue"), false);
+	return BaseGameplayCueTag;
+}
+
+static void	PrintGameplayCueNotifyMapConsoleCommandFunc(UWorld* InWorld)
+{
+	UAbilitySystemGlobals::Get().GetGameplayCueManager()->PrintGameplayCueNotifyMap();
+}
+
+FAutoConsoleCommandWithWorld PrintGameplayCueNotifyMapConsoleCommand(
+	TEXT("GameplayCue.PrintGameplayCueNotifyMap"),
+	TEXT("Displays GameplayCue notify map"),
+	FConsoleCommandWithWorldDelegate::CreateStatic(PrintGameplayCueNotifyMapConsoleCommandFunc)
+	);
\ No newline at end of file
diff --git a/Engine/Source/Runtime/GameplayAbilities/Private/GameplayCueNotify.cpp b/Engine/Source/Runtime/GameplayAbilities/Private/GameplayCueNotify.cpp
new file mode 100644
--- /dev/null
+++ b/Engine/Source/Runtime/GameplayAbilities/Private/GameplayCueNotify.cpp
@@ -0,0 +1,84 @@
+// Copyright 1998-2014 Epic Games, Inc. All Rights Reserved.
+
+#include "AbilitySystemPrivatePCH.h"
+#include "GameplayCueNotify.h"
+
+#pragma optimize( "", off )
+
+UGameplayCueNotify::UGameplayCueNotify(const class FObjectInitializer& PCIP)
+: Super(PCIP)
+{
+	IsOverride = true;
+}
+
+void UGameplayCueNotify::DeriveGameplayCueTagFromAssetName()
+{
+	// In the editor, attempt to infer GameplayCueTag from our asset name (if there is no valid GameplayCueTag already).
+#if WITH_EDITOR
+	if (GIsEditor)
+	{
+		if (GameplayCueTag.IsValid() == false)
+		{
+			FString MyName = GetName();
+
+			MyName.RemoveFromStart(TEXT("Default__"));
+			MyName.RemoveFromEnd(TEXT("_c"));
+
+			MyName.ReplaceInline(TEXT("_"), TEXT("."));
+
+			if (MyName.Contains(TEXT("GameplayCue")) == false)
+			{
+				MyName = FString(TEXT("GameplayCue.")) + MyName;
+			}
+
+			IGameplayTagsModule& GameplayTagsModule = IGameplayTagsModule::Get();
+			GameplayCueTag = GameplayTagsModule.GetGameplayTagsManager().RequestGameplayTag(FName(*MyName), false);
+		}
+
+		GameplayCueName = GameplayCueTag.GetTagName();
+	}
+#endif
+}
+
+void UGameplayCueNotify::HandleGameplayCue(AActor *Self, EGameplayCueEvent::Type EventType, FGameplayCueParameters Parameters)
+{
+
+}
+
+void UGameplayCueNotify::Serialize(FArchive& Ar)
+{
+	if (Ar.IsSaving())
+	{
+		DeriveGameplayCueTagFromAssetName();
+	}
+	
+	Super::Serialize(Ar);
+
+	if (Ar.IsLoading())
+	{
+		DeriveGameplayCueTagFromAssetName();
+	}
+}
+
+void UGameplayCueNotify::PostInitProperties()
+{
+	Super::PostInitProperties();
+	DeriveGameplayCueTagFromAssetName();
+}
+
+bool UGameplayCueNotify::HandlesEvent(EGameplayCueEvent::Type EventType) const
+{
+	return false;
+}
+
+/** Does this GameplayCueNotify need a per source actor instance for this event? */
+bool UGameplayCueNotify::NeedsInstanceForEvent(EGameplayCueEvent::Type EventType) const
+{
+	return (EventType != EGameplayCueEvent::Executed);
+}
+
+void UGameplayCueNotify::OnOwnerDestroyed()
+{
+	// May need to do extra cleanup in child classes
+	MarkPendingKill();
+}
\ No newline at end of file
diff --git a/Engine/Source/Runtime/GameplayAbilities/Private/GameplayCueNotify_Blueprint.cpp b/Engine/Source/Runtime/GameplayAbilities/Private/GameplayCueNotify_Blueprint.cpp
new file mode 100644
--- /dev/null
+++ b/Engine/Source/Runtime/GameplayAbilities/Private/GameplayCueNotify_Blueprint.cpp
@@ -0,0 +1,50 @@
+// Copyright 1998-2014 Epic Games, Inc. All Rights Reserved.
+
+#include "AbilitySystemPrivatePCH.h"
+#include "GameplayCueNotify_Blueprint.h"
+
+UGameplayCueNotify_Blueprint::UGameplayCueNotify_Blueprint(const class FObjectInitializer& PCIP)
+: Super(PCIP)
+{
+	
+}
+
+void UGameplayCueNotify_Blueprint::HandleGameplayCue(AActor* MyTarget, EGameplayCueEvent::Type EventType, FGameplayCueParameters Parameters)
+{
+	Super::HandleGameplayCue(MyTarget, EventType, Parameters);
+
+	if (MyTarget && !MyTarget->IsPendingKill())
+	{
+		K2_HandleGameplayCue(MyTarget, EventType, Parameters);
+
+		switch(EventType)
+		{
+		case EGameplayCueEvent::OnActive:
+			OnActive(MyTarget, Parameters);
+			break;
+
+		case EGameplayCueEvent::Executed:
+			OnExecute(MyTarget, Parameters);
+			break;
+
+		case EGameplayCueEvent::Removed:
+			OnRemove(MyTarget, Parameters);
+			break;
+		};
+	}
+	else
+	{
+		ABILITY_LOG(Warning, TEXT("Null Target"));
+	}
+}
+
+bool UGameplayCueNotify_Blueprint::HandlesEvent(EGameplayCueEvent::Type EventType) const
+{
+	return true;
+}
+
+/** Does this GameplayCueNotify need a per source actor instance for this event? */
+bool UGameplayCueNotify_Blueprint::NeedsInstanceForEvent(EGameplayCueEvent::Type EventType) const
+{
+	return (EventType != EGameplayCueEvent::Executed);
+}
\ No newline at end of file
diff --git a/Engine/Source/Runtime/GameplayAbilities/Private/GameplayCueNotify_HitImpact.cpp b/Engine/Source/Runtime/GameplayAbilities/Private/GameplayCueNotify_HitImpact.cpp
new file mode 100644
--- /dev/null
+++ b/Engine/Source/Runtime/GameplayAbilities/Private/GameplayCueNotify_HitImpact.cpp
@@ -0,0 +1,41 @@
+// Copyright 1998-2014 Epic Games, Inc. All Rights Reserved.
+
+#include "AbilitySystemPrivatePCH.h"
+#include "GameplayCueNotify_HitImpact.h"
+
+UGameplayCueNotify_HitImpact::UGameplayCueNotify_HitImpact(const class FObjectInitializer& PCIP)
+: Super(PCIP)
+{
+
+}
+
+bool UGameplayCueNotify_HitImpact::HandlesEvent(EGameplayCueEvent::Type EventType) const
+{
+	return (EventType == EGameplayCueEvent::Executed);
+}
+
+bool UGameplayCueNotify_HitImpact::NeedsInstanceForEvent(EGameplayCueEvent::Type EventType) const
+{
+	return false;
+}
+
+void UGameplayCueNotify_HitImpact::HandleGameplayCue(AActor* Self, EGameplayCueEvent::Type EventType, FGameplayCueParameters Parameters)
+{
+	check(EventType == EGameplayCueEvent::Executed);
+	check(Self);
+	
+	const FHitResult* HitResult = Parameters.EffectContext.GetHitResult();
+	if (HitResult)
+	{
+		if (ParticleSystem)
+		{
+			UGameplayStatics::SpawnEmitterAtLocation(Self, ParticleSystem, HitResult->ImpactPoint, HitResult->ImpactNormal.Rotation(), true);
+		}
+
+
+	}
+	else
+	{
+		ABILITY_LOG(Warning, TEXT("GameplayCue %s was called on GameplayCueNotify but there was no HitResult to spawn an impact from."), *Parameters.MatchedTagName.ToString(), *GetName() );
+	}
+}
diff --git a/Engine/Source/Runtime/GameplayAbilities/Public/Abilities/GameplayAbility.h b/Engine/Source/Runtime/GameplayAbilities/Public/Abilities/GameplayAbility.h
--- a/Engine/Source/Runtime/GameplayAbilities/Public/Abilities/GameplayAbility.h
+++ b/Engine/Source/Runtime/GameplayAbilities/Public/Abilities/GameplayAbility.h
@@ -351,14 +351,18 @@ protected:
 	// -------------------------------------
 	
 	UFUNCTION(BlueprintCallable, Category = Ability, meta=(GameplayTagFilter="GameplayCue"), FriendlyName="ExecuteGameplayCue")
-	virtual void K2_ExecuteGameplayCue(FGameplayTag GameplayCueTag);
+	virtual void K2_ExecuteGameplayCue(FGameplayTag GameplayCueTag, FGameplayEffectContextHandle Context);
 
 	UFUNCTION(BlueprintCallable, Category = Ability, meta=(GameplayTagFilter="GameplayCue"), FriendlyName="AddGameplayCue")
-	virtual void K2_AddGameplayCue(FGameplayTag GameplayCueTag);
+	virtual void K2_AddGameplayCue(FGameplayTag GameplayCueTag, FGameplayEffectContextHandle Context);
 
 	UFUNCTION(BlueprintCallable, Category = Ability, meta=(GameplayTagFilter="GameplayCue"), FriendlyName="RemoveGameplayCue")
 	virtual void K2_RemoveGameplayCue(FGameplayTag GameplayCueTag);
 
+	/** Generates a GameplayEffectContextHandle from our owner and an optional TargetData.*/
+	UFUNCTION(BlueprintCallable, Category = Ability, meta = (GameplayTagFilter = "GameplayCue"))
+	virtual FGameplayEffectContextHandle GetContextFromOwner(FGameplayAbilityTargetDataHandle OptionalTargetData) const;
+
 
 	// -------------------------------------
 
diff --git a/Engine/Source/Runtime/GameplayAbilities/Public/Abilities/GameplayAbilityTargetTypes.h b/Engine/Source/Runtime/GameplayAbilities/Public/Abilities/GameplayAbilityTargetTypes.h
--- a/Engine/Source/Runtime/GameplayAbilities/Public/Abilities/GameplayAbilityTargetTypes.h
+++ b/Engine/Source/Runtime/GameplayAbilities/Public/Abilities/GameplayAbilityTargetTypes.h
@@ -65,6 +65,8 @@ struct GAMEPLAYABILITIES_API FGameplayAbilityTargetData
 
 	TArray<FActiveGameplayEffectHandle> ApplyGameplayEffect(const UGameplayEffect* GameplayEffect, const FGameplayEffectContextHandle& InEffectContext, float Level, FModifierQualifier Qualifier = FModifierQualifier());
 
+	virtual void AddTargetDataToContext(FGameplayEffectContextHandle& Context);
+
 	virtual TArray<TWeakObjectPtr<AActor> >	GetActors() const
 	{
 		return TArray<TWeakObjectPtr<AActor> >();
diff --git a/Engine/Source/Runtime/GameplayAbilities/Public/AbilitySystemBlueprintLibrary.h b/Engine/Source/Runtime/GameplayAbilities/Public/AbilitySystemBlueprintLibrary.h
--- a/Engine/Source/Runtime/GameplayAbilities/Public/AbilitySystemBlueprintLibrary.h
+++ b/Engine/Source/Runtime/GameplayAbilities/Public/AbilitySystemBlueprintLibrary.h
@@ -103,6 +103,10 @@ class UAbilitySystemBlueprintLibrary : public UBlueprintFunctionLibrary
 	UFUNCTION(BlueprintPure, Category = "Ability|GameplayCue")
 	static FTransform GetInstigatorTransform(FGameplayCueParameters Parameters);
 
+	/** Gets instigating world location */
+	UFUNCTION(BlueprintPure, Category = "Ability|GameplayCue")
+	static FVector GetOrigin(FGameplayCueParameters Parameters);
+
 
 	// -------------------------------------------------------------------------------
 
diff --git a/Engine/Source/Runtime/GameplayAbilities/Public/AbilitySystemComponent.h b/Engine/Source/Runtime/GameplayAbilities/Public/AbilitySystemComponent.h
--- a/Engine/Source/Runtime/GameplayAbilities/Public/AbilitySystemComponent.h
+++ b/Engine/Source/Runtime/GameplayAbilities/Public/AbilitySystemComponent.h
@@ -223,7 +223,7 @@ class GAMEPLAYABILITIES_API UAbilitySystemComponent : public UActorComponent, pu
 	
 	FOnActiveGameplayEffectRemoved* OnGameplayEffectRemovedDelegate(FActiveGameplayEffectHandle Handle);
 
-	FActiveGameplayEffectHandle ApplyGameplayEffectToTarget(UGameplayEffect *GameplayEffect, UAbilitySystemComponent *Target, float Level = FGameplayEffectLevelSpec::INVALID_LEVEL, FModifierQualifier BaseQualifier = FModifierQualifier());
+	FActiveGameplayEffectHandle ApplyGameplayEffectToTarget(UGameplayEffect *GameplayEffect, UAbilitySystemComponent *Target, float Level = FGameplayEffectLevelSpec::INVALID_LEVEL, FGameplayEffectContextHandle Context = FGameplayEffectContextHandle(), FModifierQualifier BaseQualifier = FModifierQualifier());
 
 	UFUNCTION(BlueprintCallable, Category = GameplayEffects, meta=(FriendlyName = "ApplyGameplayEffectToSelf"))
 	FActiveGameplayEffectHandle K2_ApplyGameplayEffectToSelf(const UGameplayEffect *GameplayEffect, float Level, FGameplayEffectContextHandle EffectContext);
diff --git a/Engine/Source/Runtime/GameplayAbilities/Public/AbilitySystemGlobals.h b/Engine/Source/Runtime/GameplayAbilities/Public/AbilitySystemGlobals.h
--- a/Engine/Source/Runtime/GameplayAbilities/Public/AbilitySystemGlobals.h
+++ b/Engine/Source/Runtime/GameplayAbilities/Public/AbilitySystemGlobals.h
@@ -8,6 +8,7 @@ class AActor;
 class UAbilitySystemComponent;
 class UCurveTable;
 class UDataTable;
+class UGameplayCueManager;
 
 struct FGameplayAbilityActorInfo;
 struct FGameplayEffectContext;
@@ -34,6 +35,14 @@ class GAMEPLAYABILITIES_API UAbilitySystemGlobals : public UObject
 	UPROPERTY(config)
 	FString GlobalAttributeSetDefaultsTableName;
 
+	/** The class to instantiate as the global GameplayCue manager. This class is responsible for directing GameplayCue events, loading and cooking assets related to GameplayCues. */
+	UPROPERTY(config)
+	FStringAssetReference GlobalGameplayCueManagerName;
+
+	/** Look in these paths for GameplayCueNotifies */
+	UPROPERTY(config)
+	TArray<FString>	GameplayCueNotifyPaths;
+
 	/** The class to instantiate as the globals object. Defaults to this class but can be overridden */
 	UPROPERTY(config)
 	FStringClassReference AbilitySystemGlobalsClassName;
@@ -71,6 +80,8 @@ class GAMEPLAYABILITIES_API UAbilitySystemGlobals : public UObject
 
 	FAttributeSetInitter* GetAttributeSetInitter() const;
 
+	virtual UGameplayCueManager* GetGameplayCueManager();
+
 protected:
 
 	virtual void InitAtributeDefaults();
@@ -87,6 +98,9 @@ private:
 	UPROPERTY()
 	UDataTable* GlobalAttributeMetaDataTable;
 
+	UPROPERTY()
+	UGameplayCueManager* GlobalGameplayCueManager;
+
 	TSharedPtr<FAttributeSetInitter> GlobalAttributeSetInitter;
 
 	template <class T>
diff --git a/Engine/Source/Runtime/GameplayAbilities/Public/GameplayCueManager.h b/Engine/Source/Runtime/GameplayAbilities/Public/GameplayCueManager.h
new file mode 100644
--- /dev/null
+++ b/Engine/Source/Runtime/GameplayAbilities/Public/GameplayCueManager.h
@@ -0,0 +1,137 @@
+// Copyright 1998-2014 Epic Games, Inc. All Rights Reserved.
+
+#pragma once
+
+#include "GameplayTags.h"
+#include "GameplayEffect.h"
+#include "GameplayCueNotify.h"
+#include "GameplayCueManager.generated.h"
+
+/**
+ *	
+ *	Major TODO:
+ *	
+ *	-Remove LoadObjectLibraryFromPaths from base implementation. Just have project pass in ObjectLibrary, let it figure out how to make it.
+ *		-On Bob's recommendation: in hopes of removing object libraries coupling with directory structure.
+ *	
+ *	-Async loading of all required GameplayCueNotifies
+ *		-Currently async loaded on demand (first GameplayCueNotify is dropped)
+ *		-Need way to enumerate all required GameplayCue tags (Interface? GameplayCueTags can come from GameplayEffects and GameplayAbilities, possibly more?)
+ *		-Implemented UGameplayCueManager::BeginLoadingGameplayCueNotify
+ *		
+ *	-Better figure out how to handle instancing of GameplayCueNotifies
+ *		-How to handle status type GameplayCueNotifies? (Code/blueprint will spawn stuff, needs to keep track and destroy them when GameplayCue is removed)
+ *		-Instanced InstantiatedObjects is growing unbounded!
+ *		 
+ *	
+ *	-Editor Workflow:
+ *		-Make way to create new GameplayCueNotifies from UGameplayCueManager (details customization)
+ *		-Jump to GameplayCueManager entry for GameplayCueTag (details customization)
+ *		-Implement HandleAssetAdded and HandleAssetDeleted
+ *			-Must make sure we update GameplayCueData/GameplayCueDataMap at appropriate times
+ *				-On startup/begin PIE or try to do it as things change?
+ *		
+ *	-Overriding/forwarding: are we doing this right?
+ *		-When can things override, when can they call parent implementations, etc
+ *		-Do GameplayCueNotifies ever override GameplayCue Events or vice versa?
+ *		
+ *	-Take a pass on destruction
+ *		-Make sure we are cleaning up GameplayCues when actors are destroyed
+ *			(Register with Actor delegate or force game code to call EndGameplayCuesFor?)
+ *	
+ */
+
+
+USTRUCT()
+struct FGameplayCueNotifyData
+{
+	GENERATED_USTRUCT_BODY()
+
+	FGameplayCueNotifyData()
+	: LoadedGameplayCueNotify(nullptr)
+	, ParentDataIdx( INDEX_NONE )
+	{
+	}
+
+	UPROPERTY(VisibleDefaultsOnly, Category=GameplayCue, meta=(AllowedClasses="GameplayCueNotify"))
+	FStringAssetReference	GameplayCueNotifyObj;
+
+	UPROPERTY(transient)
+	UGameplayCueNotify*		LoadedGameplayCueNotify;
+
+	int32 ParentDataIdx;
+};
+
+/**
+ *	A self contained handler of a GameplayCue. These are similiar to AnimNotifies in implementation.
+ */
+
+UCLASS()
+class GAMEPLAYABILITIES_API UGameplayCueManager : public UDataAsset
+{
+	GENERATED_UCLASS_BODY()
+
+	// -------------------------------------------------------------
+	// Handling GameplayCues at runtime:
+	// -------------------------------------------------------------
+
+	virtual void HandleGameplayCues(AActor* TargetActor, const FGameplayTagContainer& GameplayCueTags, EGameplayCueEvent::Type EventType, FGameplayCueParameters Parameters);
+
+	virtual void HandleGameplayCue(AActor* TargetActor, FGameplayTag GameplayCueTag, EGameplayCueEvent::Type EventType, FGameplayCueParameters Parameters);
+
+	/** Force any instanced GameplayCueNotifies to stop */
+	virtual void EndGameplayCuesFor(AActor* TargetActor);
+
+	// -------------------------------------------------------------
+	//  Loading GameplayCueNotifies from ObjectLibraries
+	// -------------------------------------------------------------
+
+	/** Loading soft refs to all GameplayCueNotifies */
+	void LoadObjectLibraryFromPaths(const TArray<FString>& Paths, bool bFullyLoad);
+
+	UPROPERTY(EditDefaultsOnly, Category = GameplayCue)
+	TArray<FGameplayCueNotifyData>	GameplayCueData;
+
+	/** Maps GameplayCue Tag to index into above GameplayCues array. */
+	TMap<FGameplayTag, int32>	GameplayCueDataMap;
+	
+	UPROPERTY(transient)
+	UObjectLibrary* GameplayCueNotifyObjectLibrary;
+
+	// -------------------------------------------------------------
+	// Preload GameplayCue tags that we think we will need:
+	// -------------------------------------------------------------
+
+	void	BeginLoadingGameplayCueNotify(FGameplayTag GameplayCueTag);
+
+	int32	FinishLoadingGameplayCueNotifies();
+
+	FStreamableManager	StreamableManager;
+
+	UPROPERTY(transient)
+	TArray<UGameplayCueNotify*>	LoadedObjects;
+
+	UPROPERTY(transient)
+	TArray<UGameplayCueNotify*>	InstantiatedObjects;
+
+	// Fixme: we can combine the AActor* and the int32 into a single struct with a decent hash and avoid double map lookups
+	TMap< TWeakObjectPtr<AActor>, TMap<int32, TWeakObjectPtr<UGameplayCueNotify> > >	NotifyMap;
+
+	static FGameplayTag	BaseGameplayCueTag();
+
+	void PrintGameplayCueNotifyMap();
+
+#if WITH_EDITOR
+	/** Handles updating an object library when a new asset is created */
+	void HandleAssetAdded(UObject *Object);
+
+	/** Handles cleaning up an object library if it matches the passed in object */
+	void HandleAssetDeleted(UObject *Object);
+
+	bool RegisteredEditorCallbacks;
+#endif
+
+private:
+
+	virtual void HandleGameplayCueNotify_Internal(AActor* TargetActor, int32 DataIdx, EGameplayCueEvent::Type EventType, FGameplayCueParameters Parameters);
+};
\ No newline at end of file
diff --git a/Engine/Source/Runtime/GameplayAbilities/Public/GameplayCueNotify.h b/Engine/Source/Runtime/GameplayAbilities/Public/GameplayCueNotify.h
new file mode 100644
--- /dev/null
+++ b/Engine/Source/Runtime/GameplayAbilities/Public/GameplayCueNotify.h
@@ -0,0 +1,51 @@
+// Copyright 1998-2014 Epic Games, Inc. All Rights Reserved.
+
+#pragma once
+
+#include "GameplayTags.h"
+#include "GameplayEffect.h"
+#include "GameplayCueNotify.generated.h"
+
+
+/**
+ *	A self contained handler of a GameplayCue. These are similiar to AnimNotifies in implementation.
+ *	
+ *	These can be instanced or not, depending on ::NeedsInstanceForEvent.
+ *	 
+ */
+
+UCLASS(abstract, BlueprintType, hidecategories=Object, collapsecategories)
+class GAMEPLAYABILITIES_API UGameplayCueNotify : public UDataAsset
+{
+	GENERATED_UCLASS_BODY()
+
+	/** Does this GameplayCueNotify handle this type of GameplayCueEvent? */
+	virtual bool HandlesEvent(EGameplayCueEvent::Type EventType) const;
+
+	/** Does this GameplayCueNotify need a per source actor instance for this event? */
+	virtual bool NeedsInstanceForEvent(EGameplayCueEvent::Type EventType) const;
+
+	virtual void OnOwnerDestroyed();
+	
+
+	virtual void PostInitProperties();
+
+	virtual void Serialize(FArchive& Ar);
+	
+	virtual void HandleGameplayCue(AActor* MyTarget, EGameplayCueEvent::Type EventType, FGameplayCueParameters Parameters);
+
+	UPROPERTY(EditDefaultsOnly, Category=GameplayCue, meta=(Categories="GameplayCue"))
+	FGameplayTag	GameplayCueTag;
+
+	/** Mirrors GameplayCueTag in order to be asset registry searchable */
+	UPROPERTY(AssetRegistrySearchable, EditDefaultsOnly, Category = GameplayCue)
+	FName GameplayCueName;
+
+	/** Does this Cue override other cues, or is it called in addition to them? E.g., If this is Damage.Physical.Slash, we wont call Damage.Physical afer we run this cue. */
+	UPROPERTY(EditDefaultsOnly, Category=GameplayCue)
+	bool IsOverride;
+
+private:
+
+	void	DeriveGameplayCueTagFromAssetName();	
+};
diff --git a/Engine/Source/Runtime/GameplayAbilities/Public/GameplayCueNotify_Blueprint.h b/Engine/Source/Runtime/GameplayAbilities/Public/GameplayCueNotify_Blueprint.h
new file mode 100644
--- /dev/null
+++ b/Engine/Source/Runtime/GameplayAbilities/Public/GameplayCueNotify_Blueprint.h
@@ -0,0 +1,39 @@
+// Copyright 1998-2014 Epic Games, Inc. All Rights Reserved.
+
+#pragma once
+
+#include "GameplayTags.h"
+#include "GameplayEffect.h"
+#include "GameplayCueNotify_Blueprint.generated.h"
+
+
+/**
+ *	Instanced GameplayCueNotify which runs arbitrary blueprint code.
+ */
+
+UCLASS(Blueprintable,meta=(ShowWorldContextPin))
+class GAMEPLAYABILITIES_API UGameplayCueNotify_Blueprint : public UGameplayCueNotify
+{
+	GENERATED_UCLASS_BODY()
+
+	/** Does this GameplayCueNotify handle this type of GameplayCueEvent? */
+	virtual bool HandlesEvent(EGameplayCueEvent::Type EventType) const override;
+
+	/** Does this GameplayCueNotify need a per source actor instance for this event? */
+	virtual bool NeedsInstanceForEvent(EGameplayCueEvent::Type EventType) const override;
+
+	virtual void HandleGameplayCue(AActor* MyTarget, EGameplayCueEvent::Type EventType, FGameplayCueParameters Parameters) override;
+
+	/** Event Graph  */
+	UFUNCTION(BlueprintImplementableEvent, Category="GameplayCueNotify", FriendlyName="HandleGameplayCue")
+	void K2_HandleGameplayCue(TWeakObjectPtr<AActor> MyTarget, EGameplayCueEvent::Type EventType, FGameplayCueParameters Parameters);
+
+	UFUNCTION(BlueprintImplementableEvent)
+	bool OnExecute(TWeakObjectPtr<AActor> MyTarget, FGameplayCueParameters Parameters);
+
+	UFUNCTION(BlueprintImplementableEvent)
+	bool OnActive(TWeakObjectPtr<AActor> MyTarget, FGameplayCueParameters Parameters);
+
+	UFUNCTION(BlueprintImplementableEvent)
+	bool OnRemove(TWeakObjectPtr<AActor> MyTarget, FGameplayCueParameters Parameters);
+};
\ No newline at end of file
diff --git a/Engine/Source/Runtime/GameplayAbilities/Public/GameplayCueNotify_HitImpact.h b/Engine/Source/Runtime/GameplayAbilities/Public/GameplayCueNotify_HitImpact.h
new file mode 100644
--- /dev/null
+++ b/Engine/Source/Runtime/GameplayAbilities/Public/GameplayCueNotify_HitImpact.h
@@ -0,0 +1,32 @@
+// Copyright 1998-2014 Epic Games, Inc. All Rights Reserved.
+
+#pragma once
+
+#include "GameplayTags.h"
+#include "GameplayEffect.h"
+#include "GameplayCueNotify_HitImpact.generated.h"
+
+/**
+ *	Non instanced GAmeplayCueNotify for spawning particle and sound FX
+ */
+
+UCLASS()
+class GAMEPLAYABILITIES_API UGameplayCueNotify_HitImpact : public UGameplayCueNotify
+{
+	GENERATED_UCLASS_BODY()
+
+	/** Does this GameplayCueNotify handle this type of GameplayCueEvent? */
+	virtual bool HandlesEvent(EGameplayCueEvent::Type EventType) const;
+
+	/** Does this GameplayCueNotify need a per source actor instance for this event? */
+	virtual bool NeedsInstanceForEvent(EGameplayCueEvent::Type EventType) const;
+	
+	virtual void HandleGameplayCue(AActor* MyTarget, EGameplayCueEvent::Type EventType, FGameplayCueParameters Parameters) override;
+
+	UPROPERTY(EditAnywhere, BlueprintReadWrite, Category = GameplayCue)
+	USoundBase* Sound;
+
+	/** Effects to play for weapon attacks against specific surfaces */
+	UPROPERTY(EditAnywhere, BlueprintReadWrite, Category = GameplayCue)
+	UParticleSystem* ParticleSystem;
+};
