comments

diff --git a/Engine/Source/Runtime/GameplayAbilities/Public/GameplayCueNotify_Blueprint.h b/Engine/Source/Runtime/GameplayAbilities/Public/GameplayCueNotify_Blueprint.h
--- a/Engine/Source/Runtime/GameplayAbilities/Public/GameplayCueNotify_Blueprint.h
+++ b/Engine/Source/Runtime/GameplayAbilities/Public/GameplayCueNotify_Blueprint.h
@@ -9,6 +9,13 @@
 
 /**
  *	Instanced GameplayCueNotify which runs arbitrary blueprint code.
+ *	
+ *	
+ *	TODO/Fixme:
+ *		-Unsure: Leave K2_HandleGameplayCue in as generic function?
+ *		-OnExecute/Active/Remove are more clear, easy to use. Make it harder to share info between events.
+ *	
+ *	
  */
 
 UCLASS(Blueprintable,meta=(ShowWorldContextPin))
@@ -24,10 +31,10 @@ class GAMEPLAYABILITIES_API UGameplayCueNotify_Blueprint : public UGameplayCueNo
 
 	virtual void HandleGameplayCue(AActor* MyTarget, EGameplayCueEvent::Type EventType, FGameplayCueParameters Parameters) override;
 
-	/** Event Graph  */
+	/** Generic Event Graph event that will get called for every event type */
 	UFUNCTION(BlueprintImplementableEvent, Category="GameplayCueNotify", FriendlyName="HandleGameplayCue")
 	void K2_HandleGameplayCue(TWeakObjectPtr<AActor> MyTarget, EGameplayCueEvent::Type EventType, FGameplayCueParameters Parameters);
-
+		
 	UFUNCTION(BlueprintImplementableEvent)
 	bool OnExecute(TWeakObjectPtr<AActor> MyTarget, FGameplayCueParameters Parameters);
 
diff --git a/Engine/Source/Runtime/GameplayAbilities/Public/GameplayCueNotify_HitImpact.h b/Engine/Source/Runtime/GameplayAbilities/Public/GameplayCueNotify_HitImpact.h
--- a/Engine/Source/Runtime/GameplayAbilities/Public/GameplayCueNotify_HitImpact.h
+++ b/Engine/Source/Runtime/GameplayAbilities/Public/GameplayCueNotify_HitImpact.h
@@ -7,7 +7,8 @@
 #include "GameplayCueNotify_HitImpact.generated.h"
 
 /**
- *	Non instanced GAmeplayCueNotify for spawning particle and sound FX
+ *	Non instanced GAmeplayCueNotify for spawning particle and sound FX.
+ *	Still WIP - needs to be fleshed out more.
  */
 
 UCLASS()
