GameplayAbilities: Fixed non-unity build

diff --git a/Engine/Source/Runtime/GameplayAbilities/Private/AbilitySystemGlobals.cpp b/Engine/Source/Runtime/GameplayAbilities/Private/AbilitySystemGlobals.cpp
--- a/Engine/Source/Runtime/GameplayAbilities/Private/AbilitySystemGlobals.cpp
+++ b/Engine/Source/Runtime/GameplayAbilities/Private/AbilitySystemGlobals.cpp
@@ -7,6 +7,8 @@
 #include "Abilities/GameplayAbilityTypes.h"
 #include "GameplayTagsModule.h"
 #include "GameplayCueInterface.h"
+#include "GameplayCueManager.h"
+
 
 UAbilitySystemGlobals::UAbilitySystemGlobals(const FObjectInitializer& ObjectInitializer)
 : Super(ObjectInitializer)
diff --git a/Engine/Source/Runtime/GameplayAbilities/Private/GameplayCueManager.cpp b/Engine/Source/Runtime/GameplayAbilities/Private/GameplayCueManager.cpp
--- a/Engine/Source/Runtime/GameplayAbilities/Private/GameplayCueManager.cpp
+++ b/Engine/Source/Runtime/GameplayAbilities/Private/GameplayCueManager.cpp
@@ -2,6 +2,7 @@
 
 #include "AbilitySystemPrivatePCH.h"
 #include "AssetRegistryModule.h"
+#include "GameplayCueInterface.h"
 #include "GameplayCueManager.h"
 
 #if WITH_EDITOR
