GameplayAbilities: Fixed non-unity build

diff --git a/Engine/Source/Runtime/GameplayAbilities/Private/GameplayCueManager.cpp b/Engine/Source/Runtime/GameplayAbilities/Private/GameplayCueManager.cpp
--- a/Engine/Source/Runtime/GameplayAbilities/Private/GameplayCueManager.cpp
+++ b/Engine/Source/Runtime/GameplayAbilities/Private/GameplayCueManager.cpp
@@ -4,11 +4,14 @@
 #include "AssetRegistryModule.h"
 #include "GameplayCueInterface.h"
 #include "GameplayCueManager.h"
+#include "GameplayTagsModule.h"
+
 
 #if WITH_EDITOR
 #include "UnrealEd.h"
 #endif
 
+
 UGameplayCueManager::UGameplayCueManager(const class FObjectInitializer& PCIP)
 : Super(PCIP)
 {
diff --git a/Engine/Source/Runtime/GameplayAbilities/Private/GameplayCueNotify.cpp b/Engine/Source/Runtime/GameplayAbilities/Private/GameplayCueNotify.cpp
--- a/Engine/Source/Runtime/GameplayAbilities/Private/GameplayCueNotify.cpp
+++ b/Engine/Source/Runtime/GameplayAbilities/Private/GameplayCueNotify.cpp
@@ -2,9 +2,11 @@
 
 #include "AbilitySystemPrivatePCH.h"
 #include "GameplayCueNotify.h"
+#include "GameplayTagsModule.h"
 
 #pragma optimize( "", off )
 
+
 UGameplayCueNotify::UGameplayCueNotify(const class FObjectInitializer& PCIP)
 : Super(PCIP)
 {
diff --git a/Engine/Source/Runtime/GameplayAbilities/Public/GameplayCueNotify_Blueprint.h b/Engine/Source/Runtime/GameplayAbilities/Public/GameplayCueNotify_Blueprint.h
--- a/Engine/Source/Runtime/GameplayAbilities/Public/GameplayCueNotify_Blueprint.h
+++ b/Engine/Source/Runtime/GameplayAbilities/Public/GameplayCueNotify_Blueprint.h
@@ -4,6 +4,7 @@
 
 #include "GameplayTags.h"
 #include "GameplayEffect.h"
+#include "GameplayCueNotify.h"
 #include "GameplayCueNotify_Blueprint.generated.h"
 
 
diff --git a/Engine/Source/Runtime/GameplayAbilities/Public/GameplayCueNotify_HitImpact.h b/Engine/Source/Runtime/GameplayAbilities/Public/GameplayCueNotify_HitImpact.h
--- a/Engine/Source/Runtime/GameplayAbilities/Public/GameplayCueNotify_HitImpact.h
+++ b/Engine/Source/Runtime/GameplayAbilities/Public/GameplayCueNotify_HitImpact.h
@@ -4,6 +4,7 @@
 
 #include "GameplayTags.h"
 #include "GameplayEffect.h"
+#include "GameplayCueNotify.h"
 #include "GameplayCueNotify_HitImpact.generated.h"
 
 /**
