mono fix

diff --git a/Engine/Source/Runtime/GameplayAbilities/Private/GameplayCueManager.cpp b/Engine/Source/Runtime/GameplayAbilities/Private/GameplayCueManager.cpp
--- a/Engine/Source/Runtime/GameplayAbilities/Private/GameplayCueManager.cpp
+++ b/Engine/Source/Runtime/GameplayAbilities/Private/GameplayCueManager.cpp
@@ -15,7 +15,9 @@
 UGameplayCueManager::UGameplayCueManager(const class FObjectInitializer& PCIP)
 : Super(PCIP)
 {
+#if WITH_EDITOR
 	RegisteredEditorCallbacks = false;
+#endif
 }
 
 
