Reworked WaitInputRelease task to work in all cases. Really detects when the key is up rather than requiring it to go from down to up. Prep work done for WaitInputPress task.

diff --git a/Engine/Source/Runtime/GameplayAbilities/Private/Abilities/GameplayAbility.cpp b/Engine/Source/Runtime/GameplayAbilities/Private/Abilities/GameplayAbility.cpp
--- a/Engine/Source/Runtime/GameplayAbilities/Private/Abilities/GameplayAbility.cpp
+++ b/Engine/Source/Runtime/GameplayAbilities/Private/Abilities/GameplayAbility.cpp
@@ -352,16 +352,6 @@ void UGameplayAbility::ConfirmActivateSucceed()
 	}
 }
 
-void UGameplayAbility::InputPressed(const FGameplayAbilitySpecHandle Handle, const FGameplayAbilityActorInfo* ActorInfo, const FGameplayAbilityActivationInfo ActivationInfo)
-{
-}
-
-void UGameplayAbility::InputReleased(const FGameplayAbilitySpecHandle Handle, const FGameplayAbilityActorInfo* ActorInfo, const FGameplayAbilityActivationInfo ActivationInfo)
-{
-	OnInputRelease.Broadcast(this);
-	OnInputRelease.Clear();
-}
-
 bool UGameplayAbility::CheckCooldown(const FGameplayAbilitySpecHandle Handle, const FGameplayAbilityActorInfo* ActorInfo) const
 {
 	if (CooldownGameplayEffect)
diff --git a/Engine/Source/Runtime/GameplayAbilities/Private/Abilities/Tasks/AbilityTask_WaitInputRelease.cpp b/Engine/Source/Runtime/GameplayAbilities/Private/Abilities/Tasks/AbilityTask_WaitInputRelease.cpp
--- a/Engine/Source/Runtime/GameplayAbilities/Private/Abilities/Tasks/AbilityTask_WaitInputRelease.cpp
+++ b/Engine/Source/Runtime/GameplayAbilities/Private/Abilities/Tasks/AbilityTask_WaitInputRelease.cpp
@@ -10,16 +10,28 @@ UAbilityTask_WaitInputRelease::UAbilityTask_WaitInputRelease(const FObjectInitia
 	RegisteredCallback= false;
 }
 
-void UAbilityTask_WaitInputRelease::OnReleaseCallback(UGameplayAbility* InAbility)
+void UAbilityTask_WaitInputRelease::OnReleaseCallback(int32 InputID/*UGameplayAbility* InAbility*/)
 {
 	float ElapsedTime = GetWorld()->GetTimeSeconds() - StartTime;
 
-	FScopedPredictionWindow	ScopedPrediction(InAbility);
+	if (!Ability.IsValid())
+	{
+		return;
+	}
+	UGameplayAbility* MyAbility = Ability.Get();
+	check(MyAbility->IsInstantiated());
+	if (MyAbility->GetCurrentAbilitySpec()->InputID != InputID)
+	{
+		//This key is for some other ability
+		return;
+	}
+
+	FScopedPredictionWindow	ScopedPrediction(MyAbility);
 
-	if (Ability->GetCurrentActivationInfo().ActivationMode != EGameplayAbilityActivationMode::Authority)
+	if (MyAbility->GetCurrentActivationInfo().ActivationMode != EGameplayAbilityActivationMode::Authority)
 	{
 		// Tell the server we released.
-		AbilitySystemComponent->ServerInputRelease(Ability->GetCurrentAbilitySpecHandle(), ScopedPrediction.ScopedPredictionKey);
+		AbilitySystemComponent->ServerInputRelease(MyAbility->GetCurrentAbilitySpecHandle(), ScopedPrediction.ScopedPredictionKey);
 	}
 
 	// We are done. Kill us so we don't keep getting broadcast messages
@@ -48,7 +60,7 @@ void UAbilityTask_WaitInputRelease::Activate()
 			else
 			{
 				StartTime = GetWorld()->GetTimeSeconds();
-				Ability->OnInputRelease.AddUObject(this, &UAbilityTask_WaitInputRelease::OnReleaseCallback);
+				AbilitySystemComponent->AbilityKeyReleaseCallbacks.AddDynamic(this, &UAbilityTask_WaitInputRelease::OnReleaseCallback);
 				RegisteredCallback = true;
 			}
 		}
@@ -59,7 +71,7 @@ void UAbilityTask_WaitInputRelease::OnDestroy(bool AbilityEnded)
 {
 	if (RegisteredCallback && Ability.IsValid())
 	{
-		Ability->OnInputRelease.RemoveUObject(this, &UAbilityTask_WaitInputRelease::OnReleaseCallback);
+		AbilitySystemComponent->AbilityKeyReleaseCallbacks.RemoveDynamic(this, &UAbilityTask_WaitInputRelease::OnReleaseCallback);
 	}
 
 	Super::OnDestroy(AbilityEnded);
diff --git a/Engine/Source/Runtime/GameplayAbilities/Private/AbilitySystemComponent_Abilities.cpp b/Engine/Source/Runtime/GameplayAbilities/Private/AbilitySystemComponent_Abilities.cpp
--- a/Engine/Source/Runtime/GameplayAbilities/Private/AbilitySystemComponent_Abilities.cpp
+++ b/Engine/Source/Runtime/GameplayAbilities/Private/AbilitySystemComponent_Abilities.cpp
@@ -13,7 +13,6 @@
 #include "MessageLog.h"
 #include "UObjectToken.h"
 #include "MapErrors.h"
-
 #define LOCTEXT_NAMESPACE "AbilitySystemComponent"
 
 /** Enable to log out all render state create, destroy and updatetransform events */
@@ -600,7 +599,6 @@ void UAbilitySystemComponent::ServerTryActivateAbility_Implementation(FGameplayA
 	}
 
 	UGameplayAbility* InstancedAbility = NULL;
-	Spec->InputPressed = true; // Pretend input was pressed. Allows UAbilityTask_WaitInputRelease to work.
 
 	// Attempt to activate the ability (server side) and tell the client if it succeeded or failed.
 	if (TryActivateAbility(Handle, PredictionKey, &InstancedAbility))
@@ -609,7 +607,6 @@ void UAbilitySystemComponent::ServerTryActivateAbility_Implementation(FGameplayA
 	}
 	else
 	{
-		Spec->InputPressed = false;
 		ClientActivateAbilityFailed(Handle, PredictionKey.Current);
 	}
 
@@ -924,6 +921,15 @@ void UAbilitySystemComponent::AbilityInputPressed(int32 InputID)
 							Instance->InputPressed(Spec.Handle, AbilityActorInfo.Get(), Spec.ActivationInfo);
 						}						
 					}
+					FPredictionKey NewKey = FPredictionKey::CreateNewPredictionKey();
+					FScopedPredictionWindow ScopedPrediction(this, NewKey);
+					//We don't require (AbilityKeyPressCallbacks.IsBound() || AbilityKeyReleaseCallbacks.IsBound())), because we don't necessarily know now if we'll need this data later.
+					if (GetOwnerRole() != ROLE_Authority)
+					{
+						// Tell the server we pressed input.
+						ServerSetReplicatedAbilityKeyState(InputID, true, NewKey);
+					}
+					AbilityKeyPressCallbacks.Broadcast(InputID);
 				}
 				else
 				{
@@ -945,6 +951,15 @@ void UAbilitySystemComponent::AbilityInputReleased(int32 InputID)
 			if (Spec.Ability)
 			{
 				AbilitySpecInputReleased(Spec);
+				FPredictionKey NewKey = FPredictionKey::CreateNewPredictionKey();
+				FScopedPredictionWindow ScopedPrediction(this, NewKey);
+				//We don't require (AbilityKeyPressCallbacks.IsBound() || AbilityKeyReleaseCallbacks.IsBound())), because we don't necessarily know now if we'll need this data later.
+				if (GetOwnerRole() != ROLE_Authority)
+				{
+					// Tell the server we released input.
+					ServerSetReplicatedAbilityKeyState(InputID, false, NewKey);
+				}
+				AbilityKeyReleaseCallbacks.Broadcast(InputID);
 			}
 		}
 	}
@@ -1056,6 +1071,38 @@ void UAbilitySystemComponent::TargetCancel()
 
 // --------------------------------------------------------------------------
 
+void UAbilitySystemComponent::ServerSetReplicatedAbilityKeyState_Implementation(int32 InputID, bool Pressed, FPredictionKey PredictionKey)
+{
+	FScopedPredictionWindow ScopedPrediction(this, PredictionKey);
+
+	for (FGameplayAbilitySpec& Spec : ActivatableAbilities)
+	{
+		if (Spec.InputID == InputID)
+		{
+			if (Spec.Ability)
+			{
+				Spec.InputPressed = Pressed;
+			}
+		}
+	}
+
+	if (Pressed)
+	{
+		AbilityKeyPressCallbacks.Broadcast(InputID);
+	}
+	else
+	{
+		AbilityKeyReleaseCallbacks.Broadcast(InputID);
+	}
+}
+
+bool UAbilitySystemComponent::ServerSetReplicatedAbilityKeyState_Validate(int32 InputID, bool Pressed, FPredictionKey PredictionKey)
+{
+	return true;
+}
+
+// --------------------------------------------------------------------------
+
 void UAbilitySystemComponent::ServerSetReplicatedConfirm_Implementation(bool Confirmed, FPredictionKey PredictionKey)
 {
 	FScopedPredictionWindow ScopedPrediction(this, PredictionKey);
diff --git a/Engine/Source/Runtime/GameplayAbilities/Public/Abilities/GameplayAbility.h b/Engine/Source/Runtime/GameplayAbilities/Public/Abilities/GameplayAbility.h
--- a/Engine/Source/Runtime/GameplayAbilities/Public/Abilities/GameplayAbility.h
+++ b/Engine/Source/Runtime/GameplayAbilities/Public/Abilities/GameplayAbility.h
@@ -112,11 +112,11 @@ public:
 	//	
 	// ----------------------------------------------------------------------------------------------------------------
 	
-	/** Input binding. Base implementation calls TryActivateAbility */
-	virtual void InputPressed(const FGameplayAbilitySpecHandle Handle, const FGameplayAbilityActorInfo* ActorInfo, const FGameplayAbilityActivationInfo ActivationInfo);
+	/** Input binding stub. */
+	virtual void InputPressed(const FGameplayAbilitySpecHandle Handle, const FGameplayAbilityActorInfo* ActorInfo, const FGameplayAbilityActivationInfo ActivationInfo) {};
 
-	/** Input binding. Base implementation does nothing */
-	virtual void InputReleased(const FGameplayAbilitySpecHandle Handle, const FGameplayAbilityActorInfo* ActorInfo, const FGameplayAbilityActivationInfo ActivationInfo);
+	/** Input binding stub. */
+	virtual void InputReleased(const FGameplayAbilitySpecHandle Handle, const FGameplayAbilityActorInfo* ActorInfo, const FGameplayAbilityActivationInfo ActivationInfo) {};
 
 	/** Returns true if this ability can be activated right now. Has no side effects */
 	virtual bool CanActivateAbility(const FGameplayAbilitySpecHandle Handle, const FGameplayAbilityActorInfo* ActorInfo) const;
@@ -211,9 +211,6 @@ public:
 
 	void TaskEnded(UAbilityTask* Task);
 
-	/** GenericCallback for Input being released. Auto cleared on broadcast */
-	FGenericAbilityDelegate	OnInputRelease;
-
 	/** Is this ability triggered from TriggerData (or is it triggered explicitly through input/game code) */
 	bool IsTriggered() const;
 
diff --git a/Engine/Source/Runtime/GameplayAbilities/Public/Abilities/Tasks/AbilityTask_WaitInputRelease.h b/Engine/Source/Runtime/GameplayAbilities/Public/Abilities/Tasks/AbilityTask_WaitInputRelease.h
--- a/Engine/Source/Runtime/GameplayAbilities/Public/Abilities/Tasks/AbilityTask_WaitInputRelease.h
+++ b/Engine/Source/Runtime/GameplayAbilities/Public/Abilities/Tasks/AbilityTask_WaitInputRelease.h
@@ -20,7 +20,7 @@ class UAbilityTask_WaitInputRelease : public UAbilityTask
 	FInputReleaseDelegate	OnRelease;
 
 	UFUNCTION()
-	void OnReleaseCallback(UGameplayAbility* InAbility);
+	void OnReleaseCallback(int32 InputID);
 
 	virtual void Activate() override;
 
diff --git a/Engine/Source/Runtime/GameplayAbilities/Public/AbilitySystemComponent.h b/Engine/Source/Runtime/GameplayAbilities/Public/AbilitySystemComponent.h
--- a/Engine/Source/Runtime/GameplayAbilities/Public/AbilitySystemComponent.h
+++ b/Engine/Source/Runtime/GameplayAbilities/Public/AbilitySystemComponent.h
@@ -47,6 +47,9 @@ class GAMEPLAYABILITIES_API UAbilitySystemComponent : public UActorComponent, pu
 {
 	GENERATED_UCLASS_BODY()
 
+	/** Used to register callbacks to ability-key input */
+	DECLARE_DYNAMIC_MULTICAST_DELEGATE_OneParam(FAbilityAbilityKey, /*UGameplayAbility*, Ability, */int32, InputID);
+
 	/** Used to register callbacks to confirm/cancel input */
 	DECLARE_DYNAMIC_MULTICAST_DELEGATE(FAbilityConfirmOrCancel);
 
@@ -427,6 +430,8 @@ class GAMEPLAYABILITIES_API UAbilitySystemComponent : public UActorComponent, pu
 	UFUNCTION(BlueprintCallable, Category="Abilities")
 	void InputCancel();
 
+	FAbilityAbilityKey	AbilityKeyPressCallbacks;
+	FAbilityAbilityKey	AbilityKeyReleaseCallbacks;
 	FAbilityConfirmOrCancel	ConfirmCallbacks;
 	FAbilityConfirmOrCancel	CancelCallbacks;
 
@@ -575,6 +580,9 @@ public:
 	 *	While these appear to be state, these are actually synchronization events w/ some payload data
 	 */
 
+	UFUNCTION(Server, reliable, WithValidation)
+	void ServerSetReplicatedAbilityKeyState(int32 InputID, bool Pressed, FPredictionKey PredictionKey);
+
 	UFUNCTION(Server, reliable, WithValidation)
 	void ServerSetReplicatedConfirm(bool Confirmed, FPredictionKey PredictionKey);
 
@@ -588,6 +596,7 @@ public:
 
 	void ConsumeAbilityTargetData();
 
+	//Note: Ability key state is stored on the ability spec, not here on the ASC with confirm/cancel.
 	bool ReplicatedConfirmAbility;
 	bool ReplicatedCancelAbility;
 
