Found a stray pragma optimize.

diff --git a/Engine/Source/Runtime/GameplayAbilities/Private/GameplayCueNotify.cpp b/Engine/Source/Runtime/GameplayAbilities/Private/GameplayCueNotify.cpp
--- a/Engine/Source/Runtime/GameplayAbilities/Private/GameplayCueNotify.cpp
+++ b/Engine/Source/Runtime/GameplayAbilities/Private/GameplayCueNotify.cpp
@@ -4,7 +4,6 @@
 #include "GameplayCueNotify.h"
 #include "GameplayTagsModule.h"
 
-#pragma optimize( "", off )
 
 
 UGameplayCueNotify::UGameplayCueNotify(const class FObjectInitializer& PCIP)
