WaitInputPress task is working.

diff --git a/Engine/Source/Runtime/GameplayAbilities/Private/Abilities/Tasks/AbilityTask_WaitInputPress.cpp b/Engine/Source/Runtime/GameplayAbilities/Private/Abilities/Tasks/AbilityTask_WaitInputPress.cpp
new file mode 100644
--- /dev/null
+++ b/Engine/Source/Runtime/GameplayAbilities/Private/Abilities/Tasks/AbilityTask_WaitInputPress.cpp
@@ -0,0 +1,78 @@
+
+#include "AbilitySystemPrivatePCH.h"
+#include "Abilities/Tasks/AbilityTask_WaitInputPress.h"
+#include "GameplayPrediction.h"
+
+UAbilityTask_WaitInputPress::UAbilityTask_WaitInputPress(const FObjectInitializer& ObjectInitializer)
+	: Super(ObjectInitializer)
+{
+	StartTime = 0.f;
+	RegisteredCallback= false;
+}
+
+void UAbilityTask_WaitInputPress::OnPressCallback(int32 InputID)
+{
+	float ElapsedTime = GetWorld()->GetTimeSeconds() - StartTime;
+
+	if (!Ability.IsValid())
+	{
+		return;
+	}
+	UGameplayAbility* MyAbility = Ability.Get();
+	check(MyAbility->IsInstantiated());
+	if (MyAbility->GetCurrentAbilitySpec()->InputID != InputID)
+	{
+		//This key is for some other ability
+		return;
+	}
+
+	FScopedPredictionWindow	ScopedPrediction(MyAbility);
+
+	if (MyAbility->GetCurrentActivationInfo().ActivationMode != EGameplayAbilityActivationMode::Authority)
+	{
+		// Tell the server we pressed.
+		AbilitySystemComponent->ServerInputPress(MyAbility->GetCurrentAbilitySpecHandle(), ScopedPrediction.ScopedPredictionKey);
+	}
+
+	// We are done. Kill us so we don't keep getting broadcast messages
+	OnPress.Broadcast(ElapsedTime);
+	EndTask();
+}
+
+UAbilityTask_WaitInputPress* UAbilityTask_WaitInputPress::WaitInputPress(class UObject* WorldContextObject)
+{
+	return NewTask<UAbilityTask_WaitInputPress>(WorldContextObject);
+}
+
+void UAbilityTask_WaitInputPress::Activate()
+{
+	if (Ability.IsValid())
+	{
+		FGameplayAbilitySpec* Spec = Ability->GetCurrentAbilitySpec();
+		if (Spec)
+		{
+			if (Spec->InputPressed)
+			{
+				// Input was already pressed before we got here, we are done.
+				OnPress.Broadcast(0.f);
+				EndTask();
+			}
+			else
+			{
+				StartTime = GetWorld()->GetTimeSeconds();
+				AbilitySystemComponent->AbilityKeyPressCallbacks.AddDynamic(this, &UAbilityTask_WaitInputPress::OnPressCallback);
+				RegisteredCallback = true;
+			}
+		}
+	}
+}
+
+void UAbilityTask_WaitInputPress::OnDestroy(bool AbilityEnded)
+{
+	if (RegisteredCallback && Ability.IsValid())
+	{
+		AbilitySystemComponent->AbilityKeyPressCallbacks.RemoveDynamic(this, &UAbilityTask_WaitInputPress::OnPressCallback);
+	}
+
+	Super::OnDestroy(AbilityEnded);
+}
diff --git a/Engine/Source/Runtime/GameplayAbilities/Private/Abilities/Tasks/AbilityTask_WaitInputRelease.cpp b/Engine/Source/Runtime/GameplayAbilities/Private/Abilities/Tasks/AbilityTask_WaitInputRelease.cpp
--- a/Engine/Source/Runtime/GameplayAbilities/Private/Abilities/Tasks/AbilityTask_WaitInputRelease.cpp
+++ b/Engine/Source/Runtime/GameplayAbilities/Private/Abilities/Tasks/AbilityTask_WaitInputRelease.cpp
@@ -10,7 +10,7 @@ UAbilityTask_WaitInputRelease::UAbilityTask_WaitInputRelease(const FObjectInitia
 	RegisteredCallback= false;
 }
 
-void UAbilityTask_WaitInputRelease::OnReleaseCallback(int32 InputID/*UGameplayAbility* InAbility*/)
+void UAbilityTask_WaitInputRelease::OnReleaseCallback(int32 InputID)
 {
 	float ElapsedTime = GetWorld()->GetTimeSeconds() - StartTime;
 
diff --git a/Engine/Source/Runtime/GameplayAbilities/Private/AbilitySystemComponent_Abilities.cpp b/Engine/Source/Runtime/GameplayAbilities/Private/AbilitySystemComponent_Abilities.cpp
--- a/Engine/Source/Runtime/GameplayAbilities/Private/AbilitySystemComponent_Abilities.cpp
+++ b/Engine/Source/Runtime/GameplayAbilities/Private/AbilitySystemComponent_Abilities.cpp
@@ -965,6 +965,27 @@ void UAbilitySystemComponent::AbilityInputReleased(int32 InputID)
 	}
 }
 
+void UAbilitySystemComponent::AbilitySpecInputPressed(FGameplayAbilitySpec& Spec)
+{
+	Spec.InputPressed = true;
+	if (Spec.IsActive())
+	{
+		// The ability is active, so just pipe the input event to it
+		if (Spec.Ability->GetInstancingPolicy() != EGameplayAbilityInstancingPolicy::InstancedPerExecution)
+		{
+			Spec.Ability->InputPressed(Spec.Handle, AbilityActorInfo.Get(), Spec.ActivationInfo);
+		}
+		else
+		{
+			TArray<UGameplayAbility*> Instances = Spec.GetAbilityInstances();
+			for (UGameplayAbility* Instance : Instances)
+			{
+				Instance->InputPressed(Spec.Handle, AbilityActorInfo.Get(), Spec.ActivationInfo);
+			}
+		}
+	}
+}
+
 void UAbilitySystemComponent::AbilitySpecInputReleased(FGameplayAbilitySpec& Spec)
 {
 	Spec.InputPressed = false;
@@ -1014,6 +1035,22 @@ void UAbilitySystemComponent::InputCancel()
 	CancelCallbacks.Broadcast();
 }
 
+bool UAbilitySystemComponent::ServerInputPress_Validate(FGameplayAbilitySpecHandle Handle, FPredictionKey ScopedPedictionKey)
+{
+	return true;
+}
+
+void UAbilitySystemComponent::ServerInputPress_Implementation(FGameplayAbilitySpecHandle Handle, FPredictionKey ScopedPedictionKey)
+{
+	FScopedPredictionWindow ScopedPrediction(this, ScopedPedictionKey);
+
+	FGameplayAbilitySpec* Spec = FindAbilitySpecFromHandle(Handle);
+	if (Spec)
+	{
+		AbilitySpecInputPressed(*Spec);
+	}
+}
+
 bool UAbilitySystemComponent::ServerInputRelease_Validate(FGameplayAbilitySpecHandle Handle, FPredictionKey ScopedPedictionKey)
 {
 	return true;
diff --git a/Engine/Source/Runtime/GameplayAbilities/Public/Abilities/Tasks/AbilityTask_WaitInputPress.h b/Engine/Source/Runtime/GameplayAbilities/Public/Abilities/Tasks/AbilityTask_WaitInputPress.h
new file mode 100644
--- /dev/null
+++ b/Engine/Source/Runtime/GameplayAbilities/Public/Abilities/Tasks/AbilityTask_WaitInputPress.h
@@ -0,0 +1,37 @@
+// Copyright 1998-2014 Epic Games, Inc. All Rights Reserved.
+#pragma once
+#include "AbilityTask.h"
+#include "Abilities/GameplayAbilityTypes.h"
+#include "AbilitySystemComponent.h"
+#include "AbilityTask_WaitInputPress.generated.h"
+
+/**
+*	Waits until the input is pressed from activating an ability. This should be true immediately upon starting the ability, since the key was pressed to activate it.
+*	We expect server to execute this task in parrallel and keep its own time. We do not keep track of 
+*/
+UCLASS(MinimalAPI)
+class UAbilityTask_WaitInputPress : public UAbilityTask
+{
+	GENERATED_UCLASS_BODY()
+
+	DECLARE_DYNAMIC_MULTICAST_DELEGATE_OneParam(FInputPressDelegate, float, TimeWaited);
+
+	UPROPERTY(BlueprintAssignable)
+	FInputPressDelegate		OnPress;
+
+	UFUNCTION()
+	void OnPressCallback(int32 InputID);
+
+	virtual void Activate() override;
+
+	/** Wait until the user presses the input button for this ability's activation. Returns time this node spent waiting for the press. Will return 0 if input was already down. */
+	UFUNCTION(BlueprintCallable, Category = "Ability|Tasks", meta = (HidePin = "WorldContextObject", DefaultToSelf = "WorldContextObject", BlueprintInternalUseOnly = "TRUE"))
+	static UAbilityTask_WaitInputPress* WaitInputPress(UObject* WorldContextObject);
+
+protected:
+
+	virtual void OnDestroy(bool AbilityEnded) override;
+
+	float StartTime;
+	bool RegisteredCallback;
+};
\ No newline at end of file
diff --git a/Engine/Source/Runtime/GameplayAbilities/Public/AbilitySystemComponent.h b/Engine/Source/Runtime/GameplayAbilities/Public/AbilitySystemComponent.h
--- a/Engine/Source/Runtime/GameplayAbilities/Public/AbilitySystemComponent.h
+++ b/Engine/Source/Runtime/GameplayAbilities/Public/AbilitySystemComponent.h
@@ -418,9 +418,15 @@ class GAMEPLAYABILITIES_API UAbilitySystemComponent : public UActorComponent, pu
 
 	void AbilityInputReleased(int32 InputID);
 
+	void AbilitySpecInputPressed(FGameplayAbilitySpec& Spec);
+
 	void AbilitySpecInputReleased(FGameplayAbilitySpec& Spec);
 
-	/** Sent by abilities when they *need* to tell server when activation input is released. (Not sent by defauly only sent when using AbilityTask_WaitInputRelease) */
+	/** Sent by abilities to tell server when activation input is pressed. (Sent by default in order not to short-circuit WaitInputPress/WaitInputRelease tasks) */
+	UFUNCTION(Server, reliable, WithValidation)
+	void ServerInputPress(FGameplayAbilitySpecHandle AbilityToActivate, FPredictionKey ScopedPedictionKey);
+
+	/** Sent by abilities to tell server when activation input is released. (Sent by default in order not to short-circuit WaitInputPress/WaitInputRelease tasks) */
 	UFUNCTION(Server, reliable, WithValidation)
 	void ServerInputRelease(FGameplayAbilitySpecHandle AbilityToActivate, FPredictionKey ScopedPedictionKey);
 
