Trace actors moved to later in the frame so the 3D reticle doesn't lag behind the camera.

diff --git a/Engine/Source/Runtime/GameplayAbilities/Private/Abilities/GameplayAbilityTargetActor_Trace.cpp b/Engine/Source/Runtime/GameplayAbilities/Private/Abilities/GameplayAbilityTargetActor_Trace.cpp
--- a/Engine/Source/Runtime/GameplayAbilities/Private/Abilities/GameplayAbilityTargetActor_Trace.cpp
+++ b/Engine/Source/Runtime/GameplayAbilities/Private/Abilities/GameplayAbilityTargetActor_Trace.cpp
@@ -16,7 +16,7 @@ AGameplayAbilityTargetActor_Trace::AGameplayAbilityTargetActor_Trace(const FObje
 	: Super(ObjectInitializer)
 {
 	PrimaryActorTick.bCanEverTick = true;
-	PrimaryActorTick.TickGroup = TG_PrePhysics;
+	PrimaryActorTick.TickGroup = TG_PostUpdateWork;
 	StaticTargetFunction = false;
 
 	MaxRange = 999999.0f;
