Unify gameplay tag counter in AbilitySystemComponent: keep single container in AbilitySystemComponent to track GameplayEffect tags and GameplayCue tags (both from GameplayEffect and 'raw')

diff --git a/Engine/Source/Runtime/GameplayAbilities/Private/AbilitySystemComponent.cpp b/Engine/Source/Runtime/GameplayAbilities/Private/AbilitySystemComponent.cpp
--- a/Engine/Source/Runtime/GameplayAbilities/Private/AbilitySystemComponent.cpp
+++ b/Engine/Source/Runtime/GameplayAbilities/Private/AbilitySystemComponent.cpp
@@ -32,6 +32,7 @@ static FAutoConsoleVariableRef CVarDebugGameplayCues(
 
 UAbilitySystemComponent::UAbilitySystemComponent(const FObjectInitializer& ObjectInitializer)
 	: Super(ObjectInitializer)
+	, GameplayTagCountContainer(EGameplayTagMatchType::IncludeParentTags)
 {
 	bWantsInitializeComponent = true;
 
@@ -267,7 +268,7 @@ bool UAbilitySystemComponent::IsGameplayEffectActive(FActiveGameplayEffectHandle
 
 FOnGameplayEffectTagCountChanged& UAbilitySystemComponent::RegisterGameplayTagEvent(FGameplayTag Tag)
 {
-	return ActiveGameplayEffects.RegisterGameplayTagEvent(Tag);
+	return GameplayTagCountContainer.GameplayTagEventMap.FindOrAdd(Tag);
 }
 
 FOnGameplayAttributeChange& UAbilitySystemComponent::RegisterGameplayAttributeEvent(FGameplayAttribute Attribute)
@@ -275,26 +276,62 @@ FOnGameplayAttributeChange& UAbilitySystemComponent::RegisterGameplayAttributeEv
 	return ActiveGameplayEffects.RegisterGameplayAttributeEvent(Attribute);
 }
 
+FGetGameplayTags UAbilitySystemComponent::GetGameplayTagsDelegate() const
+{	
+	return FGetGameplayTags::CreateUObject(this, &UAbilitySystemComponent::GetGameplayTags);
+}
+
+FGameplayTagContainer UAbilitySystemComponent::GetGameplayTags() const
+{
+	FGameplayTagContainer TagContainer;
+	GetOwnedGameplayTags(TagContainer);
+	return TagContainer;
+}
+
+FRegisterGameplayTagChangeDelegate UAbilitySystemComponent::GetRegisterGameplayTagChangeDelegate() const
+{
+	return FRegisterGameplayTagChangeDelegate::CreateUObject(this, &UAbilitySystemComponent::RegisterGameplayTagEvent);
+}
+
 // ------------------------------------------------------------------------
 
 void UAbilitySystemComponent::GetOwnedGameplayTags(FGameplayTagContainer& TagContainer) const
 {
-	return ActiveGameplayEffects.GetOwnedGameplayTags(TagContainer);
+	SCOPE_CYCLE_COUNTER(STAT_GameplayEffectsGetOwnedTags);
+	for (auto It = GameplayTagCountContainer.GameplayTagCountMap.CreateConstIterator(); It; ++It)
+	{
+		if (It.Value() > 0)
+		{
+			TagContainer.AddTagFast(It.Key());
+		}
+	}
 }
 
 bool UAbilitySystemComponent::HasMatchingGameplayTag(FGameplayTag TagToCheck) const
 {
-	return ActiveGameplayEffects.HasMatchingGameplayTag(TagToCheck);
+	return GameplayTagCountContainer.HasMatchingGameplayTag(TagToCheck, EGameplayTagMatchType::Explicit);
 }
 
 bool UAbilitySystemComponent::HasAllMatchingGameplayTags(const FGameplayTagContainer& TagContainer, bool bCountEmptyAsMatch) const
 {
-	return ActiveGameplayEffects.HasAllMatchingGameplayTags(TagContainer, bCountEmptyAsMatch);
+	SCOPE_CYCLE_COUNTER(STAT_GameplayEffectsHasAllTags);
+	return GameplayTagCountContainer.HasAllMatchingGameplayTags(TagContainer, EGameplayTagMatchType::Explicit, bCountEmptyAsMatch);
 }
 
 bool UAbilitySystemComponent::HasAnyMatchingGameplayTags(const FGameplayTagContainer& TagContainer, bool bCountEmptyAsMatch) const
 {
-	return ActiveGameplayEffects.HasAnyMatchingGameplayTags(TagContainer, bCountEmptyAsMatch);
+	SCOPE_CYCLE_COUNTER(STAT_GameplayEffectsHasAnyTag);
+	return GameplayTagCountContainer.HasAnyMatchingGameplayTags(TagContainer, EGameplayTagMatchType::Explicit, bCountEmptyAsMatch);
+}
+
+void UAbilitySystemComponent::UpdateTagMap(const FGameplayTag& BaseTag, int32 CountDelta)
+{
+	GameplayTagCountContainer.UpdateTagMap(BaseTag, CountDelta);
+}
+
+void UAbilitySystemComponent::UpdateTagMap(const FGameplayTagContainer& Container, int32 CountDelta)
+{
+	GameplayTagCountContainer.UpdateTagMap(Container, CountDelta);
 }
 
 // ------------------------------------------------------------------------
@@ -336,6 +373,7 @@ FActiveGameplayEffectHandle UAbilitySystemComponent::ApplyGameplayEffectSpecToSe
 	{
 		return FActiveGameplayEffectHandle();
 	}
+	
 
 	// Clients should treat predicted instant effects as if they have infinite duration. The effects will be cleaned up later.
 	bool bTreatAsInfiniteDuration = GetOwnerRole() != ROLE_Authority && BaseQualifier.PredictionKey().IsValidKey() && Spec.GetDuration() == UGameplayEffect::INSTANT_APPLICATION;
@@ -640,7 +678,7 @@ void UAbilitySystemComponent::NetMulticast_InvokeGameplayCueRemoved_Implementati
 
 bool UAbilitySystemComponent::IsGameplayCueActive(const FGameplayTag GameplayCueTag) const
 {
-	return (ActiveGameplayEffects.HasMatchingGameplayTag(GameplayCueTag) || ActiveGameplayCues.HasMatchingGameplayTag(GameplayCueTag));
+	return HasMatchingGameplayTag(GameplayCueTag);
 }
 
 // ----------------------------------------------------------------------------------------
diff --git a/Engine/Source/Runtime/GameplayAbilities/Private/GameplayEffect.cpp b/Engine/Source/Runtime/GameplayAbilities/Private/GameplayEffect.cpp
--- a/Engine/Source/Runtime/GameplayAbilities/Private/GameplayEffect.cpp
+++ b/Engine/Source/Runtime/GameplayAbilities/Private/GameplayEffect.cpp
@@ -264,7 +264,7 @@ int32 FGameplayEffectSpec::ApplyModifier(const FModifierSpec &InMod, const FModi
 	{
 		case EGameplayModEffect::Magnitude:
 		{
-			for (FModifierSpec &MyMod : Modifiers)
+			for (FModifierSpec& MyMod : Modifiers)
 			{
 				if (InMod.CanModifyModifier(MyMod, QualifierContext))
 				{
@@ -1691,23 +1691,16 @@ FActiveGameplayEffect& FActiveGameplayEffectsContainer::CreateNewActiveGameplayE
 	return NewEffect;
 }
 
-void FActiveGameplayEffectsContainer::UpdateTagMap(const FGameplayTag& BaseTag, int32 CountDelta)
-{
-	GameplayTagCountContainer.UpdateTagMap(BaseTag, CountDelta);
-}
-
-void FActiveGameplayEffectsContainer::UpdateTagMap(const FGameplayTagContainer& Container, int32 CountDelta)
-{
-	GameplayTagCountContainer.UpdateTagMap(Container, CountDelta);
-}
-
 void FActiveGameplayEffectsContainer::InternalOnActiveGameplayEffectAdded(const FActiveGameplayEffect& Effect)
 {
 	// Update gameplaytag count and broadcast delegate if we just added this tag (count=0, prior to increment)
-	UpdateTagMap(Effect.Spec.Def->OwnedTagsContainer, 1);
+	Owner->UpdateTagMap(Effect.Spec.Def->OwnedTagsContainer, 1);
+
+	//Owner.ActiveGameplayCues.AddCue(
+
 	for (const FGameplayEffectCue& Cue : Effect.Spec.Def->GameplayCues)
 	{
-		UpdateTagMap(Cue.GameplayCueTags, 1);
+		Owner->UpdateTagMap(Cue.GameplayCueTags, 1);
 	}	
 }
 
@@ -1787,11 +1780,11 @@ void FActiveGameplayEffectsContainer::InternalOnActiveGameplayEffectRemoved(cons
 	{
 		// Update gameplaytag count and broadcast delegate if we are at 0
 		IGameplayTagsModule& GameplayTagsModule = IGameplayTagsModule::Get();
-		UpdateTagMap(Effect.Spec.Def->OwnedTagsContainer, -1);
+		Owner->UpdateTagMap(Effect.Spec.Def->OwnedTagsContainer, -1);
 
 		for (const FGameplayEffectCue& Cue : Effect.Spec.Def->GameplayCues)
 		{
-			UpdateTagMap(Cue.GameplayCueTags, -1);
+			Owner->UpdateTagMap(Cue.GameplayCueTags, -1);
 		}
 	}
 	else
@@ -1986,35 +1979,6 @@ void FActiveGameplayEffectsContainer::RecalculateStacking()
 	}
 }
 
-void FActiveGameplayEffectsContainer::GetOwnedGameplayTags(FGameplayTagContainer& TagContainer) const
-{
-	SCOPE_CYCLE_COUNTER(STAT_GameplayEffectsGetOwnedTags);
-	for (auto It = GameplayTagCountContainer.GameplayTagCountMap.CreateConstIterator(); It; ++It)
-	{
-		if(It.Value() > 0)
-		{
-			TagContainer.AddTagFast(It.Key());
-		}
-	}
-}
-
-bool FActiveGameplayEffectsContainer::HasMatchingGameplayTag(FGameplayTag TagToCheck) const
-{
-	return GameplayTagCountContainer.HasMatchingGameplayTag(TagToCheck, EGameplayTagMatchType::Explicit);
-}
-
-bool FActiveGameplayEffectsContainer::HasAllMatchingGameplayTags(const FGameplayTagContainer& TagContainer, bool bCountEmptyAsMatch) const
-{
-	SCOPE_CYCLE_COUNTER(STAT_GameplayEffectsHasAllTags);
-	return GameplayTagCountContainer.HasAllMatchingGameplayTags(TagContainer, EGameplayTagMatchType::Explicit, bCountEmptyAsMatch);
-}
-
-bool FActiveGameplayEffectsContainer::HasAnyMatchingGameplayTags(const FGameplayTagContainer& TagContainer, bool bCountEmptyAsMatch) const
-{
-	SCOPE_CYCLE_COUNTER(STAT_GameplayEffectsHasAnyTag);
-	return GameplayTagCountContainer.HasAnyMatchingGameplayTags(TagContainer, EGameplayTagMatchType::Explicit, bCountEmptyAsMatch);
-}
-
 bool FActiveGameplayEffectsContainer::CanApplyAttributeModifiers(const UGameplayEffect *GameplayEffect, float Level, const FGameplayEffectContextHandle& EffectContext)
 {
 	SCOPE_CYCLE_COUNTER(STAT_GameplayEffectsCanApplyAttributeModifiers);
@@ -2099,11 +2063,6 @@ void FActiveGameplayEffectsContainer::RemoveActiveEffects(const FActiveGameplayE
 	}
 }
 
-FOnGameplayEffectTagCountChanged& FActiveGameplayEffectsContainer::RegisterGameplayTagEvent(FGameplayTag Tag)
-{
-	return GameplayTagCountContainer.GameplayTagEventMap.FindOrAdd(Tag);
-}
-
 FOnGameplayAttributeChange& FActiveGameplayEffectsContainer::RegisterGameplayAttributeEvent(FGameplayAttribute Attribute)
 {
 	return AttributeChangeDelegates.FindOrAdd(Attribute);
diff --git a/Engine/Source/Runtime/GameplayAbilities/Public/AbilitySystemComponent.h b/Engine/Source/Runtime/GameplayAbilities/Public/AbilitySystemComponent.h
--- a/Engine/Source/Runtime/GameplayAbilities/Public/AbilitySystemComponent.h
+++ b/Engine/Source/Runtime/GameplayAbilities/Public/AbilitySystemComponent.h
@@ -665,7 +665,15 @@ private:
 
 	// ---------------------------------------------
 	
+	// Acceleration map for all gameplay tags (OwnedGameplayTags from GEs and explicit GameplayCueTags)
+	FGameplayTagCountContainer GameplayTagCountContainer;
+
+	void UpdateTagMap(const FGameplayTag& BaseTag, int32 CountDelta);
 	
+	void UpdateTagMap(const FGameplayTagContainer& Container, int32 CountDelta);
+
+	// Version for delegate callback.
+	FGameplayTagContainer GetOwnedGameplayTagsDel() const;
 	
 	// ---------------------------------------------
 
@@ -682,4 +690,6 @@ protected:
 
 	friend struct FActiveGameplayEffectsContainer;
 	friend struct FActiveGameplayEffect;
+	friend struct FActiveGameplayCue;
+	friend struct FActiveGameplayCueContainer;
 };
diff --git a/Engine/Source/Runtime/GameplayAbilities/Public/GameplayCueInterface.h b/Engine/Source/Runtime/GameplayAbilities/Public/GameplayCueInterface.h
--- a/Engine/Source/Runtime/GameplayAbilities/Public/GameplayCueInterface.h
+++ b/Engine/Source/Runtime/GameplayAbilities/Public/GameplayCueInterface.h
@@ -82,8 +82,6 @@ struct FActiveGameplayCueContainer : public FFastArraySerializer
 	void AddCue(const FGameplayTag& Tag);
 	void RemoveCue(const FGameplayTag& Tag);
 
-	bool HasMatchingGameplayTag(FGameplayTag TagToCheck) const;
-
 	bool NetDeltaSerialize(FNetDeltaSerializeInfo & DeltaParms)
 	{
 		return FastArrayDeltaSerialize<FActiveGameplayCue>(GameplayCues, DeltaParms, *this);
diff --git a/Engine/Source/Runtime/GameplayAbilities/Public/GameplayEffect.h b/Engine/Source/Runtime/GameplayAbilities/Public/GameplayEffect.h
--- a/Engine/Source/Runtime/GameplayAbilities/Public/GameplayEffect.h
+++ b/Engine/Source/Runtime/GameplayAbilities/Public/GameplayEffect.h
@@ -637,8 +637,7 @@ struct FGameplayModifierData
 	FScalableFloat Magnitude;
 
 	// The tags I have
-	FGameplayTagContainer Tags;
-	
+	FGameplayTagContainer Tags;	
 	FGameplayTagContainer RequireTags;
 	FGameplayTagContainer IgnoreTags;
 
@@ -713,7 +712,7 @@ struct FAggregator : public TSharedFromThis<FAggregator>
 	FAggregator(const FAggregator &In);
 	virtual ~FAggregator();
 
-	FAggregator & MarkDirty();
+	FAggregator& MarkDirty();
 	void ClearAllDependancies();
 
 	const FGameplayModifierEvaluatedData& Evaluate() const;
@@ -1110,12 +1109,12 @@ struct FActiveGameplayEffectsContainer : public FFastArraySerializer
 
 	friend struct FActiveGameplayEffect;
 
-	FActiveGameplayEffectsContainer() : bNeedToRecalculateStacks(false), GameplayTagCountContainer(EGameplayTagMatchType::IncludeParentTags) {};
+	FActiveGameplayEffectsContainer() : bNeedToRecalculateStacks(false) {};
 
 	UPROPERTY()
 	TArray< FActiveGameplayEffect >	GameplayEffects;
 
-	UAbilitySystemComponent *Owner;
+	UAbilitySystemComponent* Owner;
 	
 	FActiveGameplayEffect& CreateNewActiveGameplayEffect(const FGameplayEffectSpec &Spec, FPredictionKey InPredictionKey);
 
@@ -1173,16 +1172,6 @@ struct FActiveGameplayEffectsContainer : public FFastArraySerializer
 
 	// ------------------------------------------------
 
-	void GetOwnedGameplayTags(FGameplayTagContainer& TagContainer) const;
-
-	bool HasMatchingGameplayTag(FGameplayTag TagToCheck) const;
-
-	bool HasAllMatchingGameplayTags(const FGameplayTagContainer& TagContainer, bool bCountEmptyAsMatch = true) const;
-
-	bool HasAnyMatchingGameplayTags(const FGameplayTagContainer& TagContainer, bool bCountEmptyAsMatch = true) const;
-
-	// ------------------------------------------------
-
 	bool CanApplyAttributeModifiers(const UGameplayEffect *GameplayEffect, float Level, const FGameplayEffectContextHandle& EffectContext);
 	
 	TArray<float> GetActiveEffectsTimeRemaining(const FActiveGameplayEffectQuery Query) const;
@@ -1197,8 +1186,6 @@ struct FActiveGameplayEffectsContainer : public FFastArraySerializer
 
 	void OnDurationAggregatorDirty(FAggregator* Aggregator, UAbilitySystemComponent* Owner, FActiveGameplayEffectHandle Handle);
 	
-	FOnGameplayEffectTagCountChanged& RegisterGameplayTagEvent(FGameplayTag Tag);
-
 	FOnGameplayAttributeChange& RegisterGameplayAttributeEvent(FGameplayAttribute Attribute);
 
 	bool HasReceivedEffectWithPredictedKey(FPredictionKey PredictionKey) const;
@@ -1217,8 +1204,6 @@ private:
 
 	TMap<FGameplayAttribute, FOnGameplayAttributeChange> AttributeChangeDelegates;
 
-	FGameplayTagCountContainer GameplayTagCountContainer;
-
 	void InternalUpdateNumericalAttribute(FGameplayAttribute Attribute, float NewValue, const FGameplayEffectModCallbackData* ModData);
 
 	bool IsNetAuthority() const;
@@ -1229,9 +1214,6 @@ private:
 	/** Called both in server side creation and replication creation/deletion */
 	void InternalOnActiveGameplayEffectAdded(const FActiveGameplayEffect& Effect);
 	void InternalOnActiveGameplayEffectRemoved(const FActiveGameplayEffect& Effect);
-
-	void UpdateTagMap(const FGameplayTagContainer& Container, int32 CountDelta);
-	void UpdateTagMap(const FGameplayTag& Tag, int32 CountDelta);
 };
 
 template<>
