Added "Repeat" task to engine - this should replace the Orion-specific "Repeated Action" task eventually.

diff --git a/Engine/Source/Runtime/GameplayAbilities/Private/Abilities/Tasks/AbilityTask_Repeat.cpp b/Engine/Source/Runtime/GameplayAbilities/Private/Abilities/Tasks/AbilityTask_Repeat.cpp
new file mode 100644
--- /dev/null
+++ b/Engine/Source/Runtime/GameplayAbilities/Private/Abilities/Tasks/AbilityTask_Repeat.cpp
@@ -0,0 +1,49 @@
+
+#include "AbilitySystemPrivatePCH.h"
+#include "Abilities/Tasks/AbilityTask_Repeat.h"
+#include "GameplayPrediction.h"
+
+UAbilityTask_Repeat::UAbilityTask_Repeat(const FObjectInitializer& ObjectInitializer)
+: Super(ObjectInitializer)
+{
+}
+
+void UAbilityTask_Repeat::PerformAction()
+{		
+	OnPerformAction.Broadcast(ActionCounter);
+	++ActionCounter;
+	if (ActionCounter >= ActionPerformancesDesired)		//Should we allow negative values of ActionPerformancesDesired to mean "unlimited"?
+	{
+		OnFinished.Broadcast(-1);
+		EndTask();
+	}
+}
+
+UAbilityTask_Repeat* UAbilityTask_Repeat::RepeatAction(UObject* WorldContextObject, float InTimeBetweenActions, int32 TotalActionCount)
+{
+	auto MyObj = NewTask<UAbilityTask_Repeat>(WorldContextObject);
+
+	//TODO Validate/fix TimeBetweenActions and TotalActionCount values as needed
+	MyObj->ActionPerformancesDesired = TotalActionCount;
+	MyObj->TimeBetweenActions = InTimeBetweenActions;
+	MyObj->ActionCounter = 0;
+
+	return MyObj;
+}
+
+void UAbilityTask_Repeat::Activate()
+{
+	GetWorld()->GetTimerManager().SetTimer(this, &UAbilityTask_Repeat::PerformAction, TimeBetweenActions, true);
+}
+
+void UAbilityTask_Repeat::OnDestroy(bool AbilityIsEnding)
+{
+	GetWorld()->GetTimerManager().ClearTimer(this, &UAbilityTask_Repeat::PerformAction);
+
+	Super::OnDestroy(AbilityIsEnding);
+}
+
+FString UAbilityTask_Repeat::GetDebugString() const
+{
+	return FString::Printf(TEXT("RepeatAction. TimeBetweenActions: %.2f. ActionCounter: %d"), TimeBetweenActions, ActionCounter);
+}
\ No newline at end of file
diff --git a/Engine/Source/Runtime/GameplayAbilities/Public/Abilities/Tasks/AbilityTask_Repeat.h b/Engine/Source/Runtime/GameplayAbilities/Public/Abilities/Tasks/AbilityTask_Repeat.h
new file mode 100644
--- /dev/null
+++ b/Engine/Source/Runtime/GameplayAbilities/Public/Abilities/Tasks/AbilityTask_Repeat.h
@@ -0,0 +1,40 @@
+// Copyright 1998-2014 Epic Games, Inc. All Rights Reserved.
+#pragma once
+#include "AbilityTask.h"
+#include "Abilities/GameplayAbilityTypes.h"
+#include "AbilitySystemComponent.h"
+#include "AbilityTask_Repeat.generated.h"
+
+DECLARE_DYNAMIC_MULTICAST_DELEGATE_OneParam(FRepeatedActionDelegate, int32, ActionNumber);
+
+/**
+ *	Repeat a task a certain number of times at a given interval.
+ */
+UCLASS()
+class GAMEPLAYABILITIES_API UAbilityTask_Repeat : public UAbilityTask
+{
+	GENERATED_UCLASS_BODY()
+
+	UPROPERTY(BlueprintAssignable)
+	FRepeatedActionDelegate		OnPerformAction;
+
+	UPROPERTY(BlueprintAssignable)
+	FRepeatedActionDelegate		OnFinished;
+
+	virtual FString GetDebugString() const override;
+
+	void PerformAction();
+
+	/** Start a task that repeats an action or set of actions. */
+	UFUNCTION(BlueprintCallable, Category = "Ability|Tasks", meta = (HidePin = "WorldContextObject", DefaultToSelf = "WorldContextObject", BlueprintInternalUseOnly = "TRUE"))
+	static UAbilityTask_Repeat* RepeatAction(UObject* WorldContextObject, float TimeBetweenActions, int32 TotalActionCount);
+
+	virtual void Activate() override;
+
+protected:
+	int32 ActionPerformancesDesired;
+	int32 ActionCounter;
+	float TimeBetweenActions;
+
+	virtual void OnDestroy(bool AbilityIsEnding) override;
+};
\ No newline at end of file
