CIS fix

diff --git a/Engine/Source/Runtime/GameplayAbilities/Private/AbilitySystemComponent.cpp b/Engine/Source/Runtime/GameplayAbilities/Private/AbilitySystemComponent.cpp
--- a/Engine/Source/Runtime/GameplayAbilities/Private/AbilitySystemComponent.cpp
+++ b/Engine/Source/Runtime/GameplayAbilities/Private/AbilitySystemComponent.cpp
@@ -276,23 +276,6 @@ FOnGameplayAttributeChange& UAbilitySystemComponent::RegisterGameplayAttributeEv
 	return ActiveGameplayEffects.RegisterGameplayAttributeEvent(Attribute);
 }
 
-FGetGameplayTags UAbilitySystemComponent::GetGameplayTagsDelegate() const
-{	
-	return FGetGameplayTags::CreateUObject(this, &UAbilitySystemComponent::GetGameplayTags);
-}
-
-FGameplayTagContainer UAbilitySystemComponent::GetGameplayTags() const
-{
-	FGameplayTagContainer TagContainer;
-	GetOwnedGameplayTags(TagContainer);
-	return TagContainer;
-}
-
-FRegisterGameplayTagChangeDelegate UAbilitySystemComponent::GetRegisterGameplayTagChangeDelegate() const
-{
-	return FRegisterGameplayTagChangeDelegate::CreateUObject(this, &UAbilitySystemComponent::RegisterGameplayTagEvent);
-}
-
 // ------------------------------------------------------------------------
 
 void UAbilitySystemComponent::GetOwnedGameplayTags(FGameplayTagContainer& TagContainer) const
