adding in a way for the attribute set to define what should or should not get updated on InitAttributeSetDefaults(), a new bool is passed in for bInitialInit, if it is false it happened during gameplay and a attribute set can override ShouldInitProperty() and use the property being changed and bInitialInit to return if this property should be re-initialized.

- bug fixes to gameplay effect and gameplay abillities for Orion.

diff --git a/Engine/Source/Runtime/GameplayAbilities/Private/AttributeSet.cpp b/Engine/Source/Runtime/GameplayAbilities/Private/AttributeSet.cpp
--- a/Engine/Source/Runtime/GameplayAbilities/Private/AttributeSet.cpp
+++ b/Engine/Source/Runtime/GameplayAbilities/Private/AttributeSet.cpp
@@ -8,6 +8,7 @@
 
 #include "ComponentReregisterContext.h"	
 
+#pragma optimize( "", off )
 void FGameplayAttribute::SetNumericValueChecked(const float NewValue, class UAttributeSet* Dest) const
 {
 	UNumericProperty *NumericProperty = CastChecked<UNumericProperty>(Attribute);
@@ -305,7 +306,7 @@ void FAttributeSetInitter::PreloadAttributeSetData(UCurveTable* CurveData)
 	}
 }
 
-void FAttributeSetInitter::InitAttributeSetDefaults(UAbilitySystemComponent* AbilitySystemComponent, FName GroupName, int32 Level) const
+void FAttributeSetInitter::InitAttributeSetDefaults(UAbilitySystemComponent* AbilitySystemComponent, FName GroupName, int32 Level, bool bInitialInit) const
 {
 	SCOPE_CYCLE_COUNTER(STAT_InitAttributeSetDefaults);
 	
@@ -340,8 +341,13 @@ void FAttributeSetInitter::InitAttributeSetDefaults(UAbilitySystemComponent* Abi
 			{
 				check(DataPair.Property);
 
-				DataPair.Property->SetFloatingPointPropertyValue(DataPair.Property->ContainerPtrToValuePtr<void>(const_cast<UAttributeSet*>(Set)), DataPair.Value);
+				if (Set->ShouldInitProperty(bInitialInit, DataPair.Property))
+				{
+					FGameplayAttribute AttributeToModify(DataPair.Property);
+					AbilitySystemComponent->SetNumericAttribute(AttributeToModify, DataPair.Value);
+				}
 			}
 		}		
 	}
 }
+#pragma optimize( "", on ) 
\ No newline at end of file
diff --git a/Engine/Source/Runtime/GameplayAbilities/Public/AttributeSet.h b/Engine/Source/Runtime/GameplayAbilities/Public/AttributeSet.h
--- a/Engine/Source/Runtime/GameplayAbilities/Public/AttributeSet.h
+++ b/Engine/Source/Runtime/GameplayAbilities/Public/AttributeSet.h
@@ -82,6 +82,8 @@ public:
 		return true;
 	}
 
+	virtual bool ShouldInitProperty(bool FirstInit, UProperty* PropertyToInit) const { return true; }
+
 	/**
 	 *	Called just before modifying the value of an attribute. AttributeSet can make additional modifications here. Return true to continue, or false to throw out the modification.
 	 *	Note this is only called during an 'execute'. E.g., a modification to the 'base value' of an attribute. It is not called during an application of a GameplayEffect, such as a 5 ssecond +10 movement speed buff.
@@ -305,7 +307,7 @@ struct GAMEPLAYABILITIES_API FAttributeSetInitter
 {
 	void PreloadAttributeSetData(UCurveTable* CurveData);
 
-	void InitAttributeSetDefaults(UAbilitySystemComponent* AbilitySystemComponent, FName GroupName, int32 Level) const;
+	void InitAttributeSetDefaults(UAbilitySystemComponent* AbilitySystemComponent, FName GroupName, int32 Level, bool bInitialInit) const;
 
 private:
 
