removing debug code

diff --git a/Engine/Source/Runtime/GameplayAbilities/Private/AttributeSet.cpp b/Engine/Source/Runtime/GameplayAbilities/Private/AttributeSet.cpp
--- a/Engine/Source/Runtime/GameplayAbilities/Private/AttributeSet.cpp
+++ b/Engine/Source/Runtime/GameplayAbilities/Private/AttributeSet.cpp
@@ -8,7 +8,6 @@
 
 #include "ComponentReregisterContext.h"	
 
-#pragma optimize( "", off )
 void FGameplayAttribute::SetNumericValueChecked(const float NewValue, class UAttributeSet* Dest) const
 {
 	UNumericProperty *NumericProperty = CastChecked<UNumericProperty>(Attribute);
@@ -349,5 +348,4 @@ void FAttributeSetInitter::InitAttributeSetDefaults(UAbilitySystemComponent* Abi
 			}
 		}		
 	}
-}
-#pragma optimize( "", on ) 
\ No newline at end of file
+}
\ No newline at end of file
