missed files

diff --git a/Engine/Source/Runtime/GameplayAbilities/Private/GameplayCueInterface.cpp b/Engine/Source/Runtime/GameplayAbilities/Private/GameplayCueInterface.cpp
--- a/Engine/Source/Runtime/GameplayAbilities/Private/GameplayCueInterface.cpp
+++ b/Engine/Source/Runtime/GameplayAbilities/Private/GameplayCueInterface.cpp
@@ -100,12 +100,16 @@ void IGameplayCueInterface::ForwardGameplayCueToParent()
 
 void FActiveGameplayCue::PreReplicatedRemove(const struct FActiveGameplayCueContainer &InArray)
 {
+	// FIXME: Prediction key check is missing here
 	InArray.Owner->InvokeGameplayCueEvent(GameplayCueTag, EGameplayCueEvent::Removed);
+	InArray.Owner->UpdateTagMap(GameplayCueTag, -1);
 }
 
 void FActiveGameplayCue::PostReplicatedAdd(const struct FActiveGameplayCueContainer &InArray)
 {
+	// FIXME: Prediction key check is missing here
 	InArray.Owner->InvokeGameplayCueEvent(GameplayCueTag, EGameplayCueEvent::WhileActive);
+	InArray.Owner->UpdateTagMap(GameplayCueTag, 1);
 }
 
 void FActiveGameplayCueContainer::AddCue(const FGameplayTag& Tag)
@@ -117,6 +121,7 @@ void FActiveGameplayCueContainer::AddCue(const FGameplayTag& Tag)
 	MarkItemDirty(NewCue);
 
 	GameplayCues.Add(NewCue);
+	Owner->UpdateTagMap(Tag, 1);
 }
 
 void FActiveGameplayCueContainer::RemoveCue(const FGameplayTag& Tag)
@@ -129,21 +134,8 @@ void FActiveGameplayCueContainer::RemoveCue(const FGameplayTag& Tag)
 		{
 			GameplayCues.RemoveAt(idx);
 			MarkArrayDirty();
+			Owner->UpdateTagMap(Tag, -11);
 			return;
 		}
 	}
-}
-
-bool FActiveGameplayCueContainer::HasMatchingGameplayTag(FGameplayTag TagToCheck) const
-{
-	for (const FActiveGameplayCue& Cue : GameplayCues)
-	{
-		/** We expect that if TagToCheck=(a.b) and we have a (a.b.c) tag, then we match  */
-		if (Cue.GameplayCueTag.Matches(EGameplayTagMatchType::IncludeParentTags, TagToCheck, EGameplayTagMatchType::Explicit))
-		{
-			return true;
-		}
-	}
-
-	return false;
 }
\ No newline at end of file
