First pass at a server-authoritative EndAbility node created. Ignores calls from non-authority machines, ends non-replicated ability instances on the appropriate client via the normal K2_EndAbility.

diff --git a/Engine/Source/Runtime/GameplayAbilities/Private/Abilities/GameplayAbility.cpp b/Engine/Source/Runtime/GameplayAbilities/Private/Abilities/GameplayAbility.cpp
--- a/Engine/Source/Runtime/GameplayAbilities/Private/Abilities/GameplayAbility.cpp
+++ b/Engine/Source/Runtime/GameplayAbilities/Private/Abilities/GameplayAbility.cpp
@@ -492,6 +492,19 @@ void UGameplayAbility::K2_EndAbility()
 	EndAbility(CurrentSpecHandle, CurrentActorInfo, CurrentActivationInfo);
 }
 
+void UGameplayAbility::K2_EndAbilityFromServer()
+{
+	check(CurrentActorInfo);
+	check(CurrentActorInfo->AvatarActor.IsValid());
+	ENetRole NetMode = CurrentActorInfo->AvatarActor->Role;
+
+	if (NetMode == ROLE_Authority)		//Do we care about ability execution policy?
+	{
+		CurrentActorInfo->AbilitySystemComponent->ClientEndAbility(CurrentSpecHandle);		//Does this need a prediction key?
+		K2_EndAbility();
+	}
+}
+
 // --------------------------------------------------------------------
 
 void UGameplayAbility::MontageJumpToSection(FName SectionName)
diff --git a/Engine/Source/Runtime/GameplayAbilities/Private/AbilitySystemComponent_Abilities.cpp b/Engine/Source/Runtime/GameplayAbilities/Private/AbilitySystemComponent_Abilities.cpp
--- a/Engine/Source/Runtime/GameplayAbilities/Private/AbilitySystemComponent_Abilities.cpp
+++ b/Engine/Source/Runtime/GameplayAbilities/Private/AbilitySystemComponent_Abilities.cpp
@@ -472,7 +472,7 @@ bool UAbilitySystemComponent::TryActivateAbility(FGameplayAbilitySpecHandle Hand
 
 	UGameplayAbility* InstancedAbility = Ability;
 
-	// If we are the server or this is a client authorative 
+	// If we are the server or this is a client authoritative 
 	if (Ability->GetNetExecutionPolicy() == EGameplayAbilityNetExecutionPolicy::Client || (NetMode == ROLE_Authority))
 	{
 		// Create instance of this ability if necessary
@@ -619,7 +619,24 @@ void UAbilitySystemComponent::ServerTryActivateAbility_Implementation(FGameplayA
 	}
 }
 
-bool UAbilitySystemComponent::ServerTryActivateAbility_Validate(FGameplayAbilitySpecHandle AbilityToActivate, FPredictionKey PredictionKey)
+bool UAbilitySystemComponent::ServerTryActivateAbility_Validate(FGameplayAbilitySpecHandle Handle, FPredictionKey PredictionKey)
+{
+	return true;
+}
+
+void UAbilitySystemComponent::ClientEndAbility_Implementation(FGameplayAbilitySpecHandle AbilityToEnd)
+{
+	FGameplayAbilitySpec* AbilitySpec = FindAbilitySpecFromHandle(AbilityToEnd);
+	if (AbilitySpec && AbilitySpec->Ability && AbilitySpec->IsActive())
+	{
+		for (auto Instance : AbilitySpec->NonReplicatedInstances)
+		{
+			Instance->K2_EndAbility();
+		}
+	}
+}
+
+bool UAbilitySystemComponent::ClientEndAbility_Validate(FGameplayAbilitySpecHandle AbilityToActivate)
 {
 	return true;
 }
diff --git a/Engine/Source/Runtime/GameplayAbilities/Public/Abilities/GameplayAbility.h b/Engine/Source/Runtime/GameplayAbilities/Public/Abilities/GameplayAbility.h
--- a/Engine/Source/Runtime/GameplayAbilities/Public/Abilities/GameplayAbility.h
+++ b/Engine/Source/Runtime/GameplayAbilities/Public/Abilities/GameplayAbility.h
@@ -323,6 +323,9 @@ protected:
 	UFUNCTION(BlueprintCallable, Category = Ability, FriendlyName="EndAbility")
 	virtual void K2_EndAbility();
 
+	UFUNCTION(BlueprintCallable, Category = Ability, FriendlyName = "EndAbilityFromServer")
+	virtual void K2_EndAbilityFromServer();
+
 	virtual void EndAbility(const FGameplayAbilitySpecHandle Handle, const FGameplayAbilityActorInfo* ActorInfo, const FGameplayAbilityActivationInfo ActivationInfo);
 
 	// -------------------------------------
diff --git a/Engine/Source/Runtime/GameplayAbilities/Public/AbilitySystemComponent.h b/Engine/Source/Runtime/GameplayAbilities/Public/AbilitySystemComponent.h
--- a/Engine/Source/Runtime/GameplayAbilities/Public/AbilitySystemComponent.h
+++ b/Engine/Source/Runtime/GameplayAbilities/Public/AbilitySystemComponent.h
@@ -384,6 +384,9 @@ class GAMEPLAYABILITIES_API UAbilitySystemComponent : public UActorComponent, pu
 	UFUNCTION(Server, reliable, WithValidation)
 	void	ServerTryActivateAbility(FGameplayAbilitySpecHandle AbilityToActivate, FPredictionKey PredictionKey);
 
+	UFUNCTION(Client, reliable, WithValidation)
+	void	ClientEndAbility(FGameplayAbilitySpecHandle AbilityToEnd);
+
 	UFUNCTION(Client, Reliable)
 	void	ClientActivateAbilityFailed(FGameplayAbilitySpecHandle AbilityToActivate, int16 PredictionKey);
 
