Repeat task no longer waits before firing the first repetition.

diff --git a/Engine/Source/Runtime/GameplayAbilities/Private/Abilities/Tasks/AbilityTask_Repeat.cpp b/Engine/Source/Runtime/GameplayAbilities/Private/Abilities/Tasks/AbilityTask_Repeat.cpp
--- a/Engine/Source/Runtime/GameplayAbilities/Private/Abilities/Tasks/AbilityTask_Repeat.cpp
+++ b/Engine/Source/Runtime/GameplayAbilities/Private/Abilities/Tasks/AbilityTask_Repeat.cpp
@@ -33,7 +33,19 @@ UAbilityTask_Repeat* UAbilityTask_Repeat::RepeatAction(UObject* WorldContextObje
 
 void UAbilityTask_Repeat::Activate()
 {
-	GetWorld()->GetTimerManager().SetTimer(this, &UAbilityTask_Repeat::PerformAction, TimeBetweenActions, true);
+	if (ActionCounter < ActionPerformancesDesired)
+	{
+		PerformAction();
+		if (ActionCounter < ActionPerformancesDesired)
+		{
+			GetWorld()->GetTimerManager().SetTimer(this, &UAbilityTask_Repeat::PerformAction, TimeBetweenActions, true);
+		}
+	}
+	else
+	{
+		OnFinished.Broadcast(-1);
+		EndTask();
+	}
 }
 
 void UAbilityTask_Repeat::OnDestroy(bool AbilityIsEnding)
