GameplayCueManager
-Properly handle adding and deleting assets post load
-Make 'preload all notifies' a project setting

diff --git a/Engine/Source/Runtime/GameplayAbilities/Private/AbilitySystemGlobals.cpp b/Engine/Source/Runtime/GameplayAbilities/Private/AbilitySystemGlobals.cpp
--- a/Engine/Source/Runtime/GameplayAbilities/Private/AbilitySystemGlobals.cpp
+++ b/Engine/Source/Runtime/GameplayAbilities/Private/AbilitySystemGlobals.cpp
@@ -16,6 +16,9 @@ UAbilitySystemGlobals::UAbilitySystemGlobals(const FObjectInitializer& ObjectIni
 	AbilitySystemGlobalsClassName = FStringClassReference(TEXT("/Script/GameplayAbilities.AbilitySystemGlobals"));
 	GlobalGameplayCueManagerName = FStringClassReference(TEXT("/Script/GameplayAbilities.GameplayCueManager"));
 
+	// Temp: once system for loading only necessary GameplaycueNotifies, this can go away.
+	GameplayCueNotifyFullyLoad = true;
+
 #if WITH_EDITORONLY_DATA
 	RegisteredReimportCallback = false;
 #endif // #if WITH_EDITORONLY_DATA
@@ -199,7 +202,7 @@ UGameplayCueManager* UAbilitySystemGlobals::GetGameplayCueManager()
 		GlobalGameplayCueManager = LoadObject<UGameplayCueManager>(NULL, *GlobalGameplayCueManagerName.ToString(), NULL, LOAD_None, NULL);
 		if (GameplayCueNotifyPaths.Num() > 0)
 		{
-			GlobalGameplayCueManager->LoadObjectLibraryFromPaths( GameplayCueNotifyPaths, false );
+			GlobalGameplayCueManager->LoadObjectLibraryFromPaths( GameplayCueNotifyPaths, GameplayCueNotifyFullyLoad );
 		}
 	}
 
diff --git a/Engine/Source/Runtime/GameplayAbilities/Private/GameplayCueManager.cpp b/Engine/Source/Runtime/GameplayAbilities/Private/GameplayCueManager.cpp
--- a/Engine/Source/Runtime/GameplayAbilities/Private/GameplayCueManager.cpp
+++ b/Engine/Source/Runtime/GameplayAbilities/Private/GameplayCueManager.cpp
@@ -6,7 +6,6 @@
 #include "GameplayCueManager.h"
 #include "GameplayTagsModule.h"
 
-
 #if WITH_EDITOR
 #include "UnrealEd.h"
 #endif
@@ -49,46 +48,80 @@ void UGameplayCueManager::HandleGameplayCueNotify_Internal(AActor* TargetActor,
 
 		FGameplayCueNotifyData& CueData = GameplayCueData[DataIdx];
 
+		// If object is not loaded yet
 		if (CueData.LoadedGameplayCueNotify == nullptr)
 		{
-			CueData.LoadedGameplayCueNotify = FindObject<UGameplayCueNotify>(nullptr, *CueData.GameplayCueNotifyObj.ToString());
-			if (CueData.LoadedGameplayCueNotify == nullptr)
+			// Ignore removed events if this wasn't already loaded (only call Removed if we handled OnActive/WhileActive)
+			if (EventType == EGameplayCueEvent::Removed)
 			{
-				StreamableManager.SimpleAsyncLoad(CueData.GameplayCueNotifyObj);
-				ABILITY_LOG(Warning, TEXT("GameplayCueNotify %s was not loaded when GameplayCue was invoked. Starting async loading."), *CueData.GameplayCueNotifyObj.ToString());
 				return;
 			}
-		}
 
-		UGameplayCueNotify* SpawnedInstanced = nullptr;
+			// See if the object is loaded but just not hooked up here
+			UObject* FoundObject = FindObject<UObject>(nullptr, *CueData.GameplayCueNotifyObj.ToString());
+			if (FoundObject == nullptr)
+			{
+				// Not loaded: start async loading and return
+				StreamableManager.SimpleAsyncLoad(CueData.GameplayCueNotifyObj);
 
-		if (auto InnerMap = NotifyMap.Find(TargetActor))
-		{
-			if (auto WeakPtrPtr = InnerMap->Find(DataIdx))
+				ABILITY_LOG(Warning, TEXT("GameplayCueNotify %s was not loaded when GameplayCue was invoked. Starting async loading."), *CueData.GameplayCueNotifyObj.ToString());
+				return;
+			}
+			else
 			{
-				SpawnedInstanced = WeakPtrPtr->Get();
+				// Found it - did we load a data asset (UGameplayCueNotify) or a Blueprint (UBlueprint)
+				CueData.LoadedGameplayCueNotify = Cast<UGameplayCueNotify>(FoundObject);
+				if (CueData.LoadedGameplayCueNotify == nullptr)
+				{
+					// Not a dataasset - maybe a blueprint
+					
+					UBlueprint* GameplayCueBlueprint = Cast<UBlueprint>(FoundObject);
+					if (GameplayCueBlueprint && GameplayCueBlueprint->GeneratedClass)
+					{
+						CueData.LoadedGameplayCueNotify = Cast<UGameplayCueNotify>(GameplayCueBlueprint->GeneratedClass->ClassDefaultObject);
+					}
+					else
+					{
+						ABILITY_LOG(Warning, TEXT("GameplayCueNotify %s loaded object %s that is not a UGameplayCueNotify or UBlueprint"), *CueData.GameplayCueNotifyObj.ToString(), *FoundObject->GetName());
+						return;
+					}
+					
+				}
 			}
 		}
 
-		if (SpawnedInstanced == nullptr && CueData.LoadedGameplayCueNotify->NeedsInstanceForEvent(EventType))
+		// Handle the Notify if we found something
+		if (CueData.LoadedGameplayCueNotify)
 		{
-			// We don't have an instance for this, and we need one, so make one
-			SpawnedInstanced = static_cast<UGameplayCueNotify*>(StaticDuplicateObject(CueData.LoadedGameplayCueNotify, this, TEXT("None"), ~RF_RootSet));
+			UGameplayCueNotify* SpawnedInstanced = nullptr;
+			if (auto InnerMap = NotifyMap.Find(TargetActor))
+			{
+				if (auto WeakPtrPtr = InnerMap->Find(DataIdx))
+				{
+					SpawnedInstanced = WeakPtrPtr->Get();
+				}
+			}
 
-			auto& InnerMap = NotifyMap.FindOrAdd(TargetActor);
+			if (SpawnedInstanced == nullptr && CueData.LoadedGameplayCueNotify->NeedsInstanceForEvent(EventType))
+			{
+				// We don't have an instance for this, and we need one, so make one
+				SpawnedInstanced = static_cast<UGameplayCueNotify*>(StaticDuplicateObject(CueData.LoadedGameplayCueNotify, this, TEXT("None"), ~RF_RootSet));
 
-			InnerMap.Add(DataIdx) = SpawnedInstanced;
+				auto& InnerMap = NotifyMap.FindOrAdd(TargetActor);
 
-			// Fixme: this will grow unbounded
-			InstantiatedObjects.Add(SpawnedInstanced);
-		}
+				InnerMap.Add(DataIdx) = SpawnedInstanced;
 
-		UGameplayCueNotify* NotifyToExecute = SpawnedInstanced ? SpawnedInstanced : CueData.LoadedGameplayCueNotify;
-		NotifyToExecute->HandleGameplayCue(TargetActor, EventType, Parameters);
+				// Fixme: this will grow unbounded
+				InstantiatedObjects.Add(SpawnedInstanced);
+			}
+			
+			UGameplayCueNotify* NotifyToExecute = SpawnedInstanced ? SpawnedInstanced : CueData.LoadedGameplayCueNotify;
+			NotifyToExecute->HandleGameplayCue(TargetActor, EventType, Parameters);
 
-		if (NotifyToExecute->IsOverride == false)
-		{
-			HandleGameplayCueNotify_Internal(TargetActor, CueData.ParentDataIdx, EventType, Parameters);
+			if (NotifyToExecute->IsOverride == false)
+			{
+				HandleGameplayCueNotify_Internal(TargetActor, CueData.ParentDataIdx, EventType, Parameters);
+			}
 		}
 	}
 }
@@ -111,7 +144,30 @@ void UGameplayCueManager::EndGameplayCuesFor(AActor* TargetActor)
 
 // ------------------------------------------------------------------------
 
-void UGameplayCueManager::LoadObjectLibraryFromPaths(const TArray<FString>& Paths, bool bFullyLoad)
+void UGameplayCueManager::LoadObjectLibraryFromPaths(const TArray<FString>& InPaths, bool InFullyLoad)
+{
+	if (!GameplayCueNotifyObjectLibrary)
+	{
+		GameplayCueNotifyObjectLibrary = UObjectLibrary::CreateLibrary(UGameplayCueNotify::StaticClass(), true, GIsEditor);
+	}
+
+	LoadedPaths = InPaths;
+	bFullyLoad = InFullyLoad;
+
+	LoadObjectLibrary_Internal();
+
+#if WITH_EDITOR
+	if (!RegisteredEditorCallbacks)
+	{
+		FAssetRegistryModule& AssetRegistryModule = FModuleManager::LoadModuleChecked<FAssetRegistryModule>("AssetRegistry");
+		AssetRegistryModule.Get().OnInMemoryAssetCreated().AddUObject(this, &UGameplayCueManager::HandleAssetAdded);
+		AssetRegistryModule.Get().OnInMemoryAssetDeleted().AddUObject(this, &UGameplayCueManager::HandleAssetDeleted);
+		RegisteredEditorCallbacks = true;
+	}
+#endif
+}
+
+void UGameplayCueManager::LoadObjectLibrary_Internal()
 {
 	DECLARE_SCOPE_CYCLE_COUNTER(TEXT("Loading Library"), STAT_ObjectLibrary, STATGROUP_LoadTime);
 
@@ -121,15 +177,10 @@ void UGameplayCueManager::LoadObjectLibraryFromPaths(const TArray<FString>& Path
 	SlowTask.MakeDialog();
 #endif
 
-	if (!GameplayCueNotifyObjectLibrary)
-	{
-		GameplayCueNotifyObjectLibrary = UObjectLibrary::CreateLibrary(UGameplayCueNotify::StaticClass(), true, GIsEditor);
-	}
-
 	FScopeCycleCounterUObject PreloadScope(GameplayCueNotifyObjectLibrary);
-
-	GameplayCueNotifyObjectLibrary->LoadBlueprintAssetDataFromPaths(Paths);
-	GameplayCueNotifyObjectLibrary->LoadAssetDataFromPaths(Paths);	
+	
+	GameplayCueNotifyObjectLibrary->LoadBlueprintAssetDataFromPaths(LoadedPaths);
+	GameplayCueNotifyObjectLibrary->LoadAssetDataFromPaths(LoadedPaths);
 
 	if (bFullyLoad)
 	{
@@ -150,34 +201,24 @@ void UGameplayCueManager::LoadObjectLibraryFromPaths(const TArray<FString>& Path
 	IGameplayTagsModule& GameplayTagsModule = IGameplayTagsModule::Get();
 
 	GameplayCueData.Empty();
-
-	TSet<FGameplayTag>	FoundTags;
+	GameplayCueDataMap.Empty();
 
 	for (FAssetData Data: AssetDatas)
 	{
 		const FString* FoundGameplayTag = Data.TagsAndValues.Find(GET_MEMBER_NAME_CHECKED(UGameplayCueNotify, GameplayCueName));
-		if (FoundGameplayTag)
+		if (FoundGameplayTag && FoundGameplayTag->Equals(TEXT("None")) == false)
 		{
 			ABILITY_LOG(Warning, TEXT("Found: %s"), **FoundGameplayTag);
 
 			FGameplayTag  GameplayCueTag = GameplayTagsModule.GetGameplayTagsManager().RequestGameplayTag(FName(**FoundGameplayTag), false);
 			if (GameplayCueTag.IsValid())
 			{
-				FoundTags.Add(GameplayCueTag);
-
 				// Add a new NotifyData entry to our flat list for this one
-				FGameplayCueNotifyData NewData;
 				FStringAssetReference StringRef;
 				StringRef.AssetLongPathname = Data.ObjectPath.ToString();
 
-				TAssetSubclassOf<UGameplayCueNotify> AssetRef;
-				AssetRef = StringRef;
+				AddGameplayCueData_Internal(GameplayCueTag, StringRef);
 				
-				NewData.GameplayCueNotifyObj = StringRef;
-				int32 NewIdx = GameplayCueData.Add(NewData);
-
-				// Add an entry into the acceleration map
-				GameplayCueDataMap.FindOrAdd(GameplayCueTag) = NewIdx;
 			}
 			else
 			{
@@ -186,16 +227,46 @@ void UGameplayCueManager::LoadObjectLibraryFromPaths(const TArray<FString>& Path
 		}
 	}
 
+	BuildAccelerationMap_Internal();
+}
+
+void UGameplayCueManager::AddGameplayCueData_Internal(FGameplayTag  GameplayCueTag, FStringAssetReference StringRef)
+{
+	// Check for duplicates: we may want to remove this eventually.
+	for (FGameplayCueNotifyData Data : GameplayCueData)
+	{
+		if (Data.GameplayCueTag == GameplayCueTag)
+		{
+			ABILITY_LOG(Warning, TEXT("AddGameplayCueData_Internal called for [%s,%s] when it already existed [%s,%s]. Skipping."), *GameplayCueTag.ToString(), *StringRef.AssetLongPathname, *Data.GameplayCueTag.ToString(), *Data.GameplayCueNotifyObj.AssetLongPathname);
+			return;
+		}
+	}
+
+	FGameplayCueNotifyData NewData;
+	NewData.GameplayCueNotifyObj = StringRef;
+	NewData.GameplayCueTag = GameplayCueTag;
+
+	GameplayCueData.Add(NewData);
+}
+
+void UGameplayCueManager::BuildAccelerationMap_Internal()
+{
 	// ---------------------------------------------------------
 	//	Build up the rest of the acceleration map: every GameplayCue tag should have an entry in the map that points to the index into GameplayCueData to use when it is invoked.
 	//	(or to -1 if no GameplayCueNotify is associated with that tag)
 	// 
 	// ---------------------------------------------------------
-	
-	FGameplayTagContainer AllGameplayCueTags = GameplayTagsModule.GetGameplayTagsManager().RequestGameplayTagChildren(BaseGameplayCueTag());
 
+	GameplayCueDataMap.Empty();
 	GameplayCueDataMap.Add(BaseGameplayCueTag()) = INDEX_NONE;
 
+	for (int32 idx = 0; idx < GameplayCueData.Num(); ++idx)
+	{
+		GameplayCueDataMap.FindOrAdd(GameplayCueData[idx].GameplayCueTag) = idx;
+	}
+
+	FGameplayTagContainer AllGameplayCueTags = IGameplayTagsModule::Get().GetGameplayTagsManager().RequestGameplayTagChildren(BaseGameplayCueTag());
+
 	for (FGameplayTag ThisGameplayCueTag : AllGameplayCueTags)
 	{
 		if (GameplayCueDataMap.Contains(ThisGameplayCueTag))
@@ -209,17 +280,6 @@ void UGameplayCueManager::LoadObjectLibraryFromPaths(const TArray<FString>& Path
 	}
 
 	// PrintGameplayCueNotifyMap();
-
-#if WITH_EDITOR
-	if (!RegisteredEditorCallbacks)
-	{
-		FAssetRegistryModule& AssetRegistryModule = FModuleManager::LoadModuleChecked<FAssetRegistryModule>("AssetRegistry");
-		AssetRegistryModule.Get().OnInMemoryAssetCreated().AddUObject(this, &UGameplayCueManager::HandleAssetAdded);
-		AssetRegistryModule.Get().OnInMemoryAssetDeleted().AddUObject(this, &UGameplayCueManager::HandleAssetDeleted);
-		RegisteredEditorCallbacks = true;
-	}
-#endif
-
 }
 
 int32 UGameplayCueManager::FinishLoadingGameplayCueNotifies()
@@ -237,13 +297,65 @@ void UGameplayCueManager::BeginLoadingGameplayCueNotify(FGameplayTag GameplayCue
 
 void UGameplayCueManager::HandleAssetAdded(UObject *Object)
 {
+	UBlueprint* Blueprint = Cast<UBlueprint>(Object);
+	if (Blueprint && Blueprint->GeneratedClass)
+	{
+		UGameplayCueNotify* CDO = Cast<UGameplayCueNotify>(Blueprint->GeneratedClass->ClassDefaultObject);
+		if (CDO && CDO->GameplayCueTag.IsValid())
+		{
+			FStringAssetReference StringRef;
+			StringRef.AssetLongPathname = Blueprint->GetPathName();
 
+			AddGameplayCueData_Internal(CDO->GameplayCueTag, StringRef);
+			BuildAccelerationMap_Internal();
+		}
+	}
+	
+	
+	UGameplayCueNotify* Notify = Cast<UGameplayCueNotify>(Object);
+	if (Notify && Notify->GameplayCueTag.IsValid())
+	{
+		FStringAssetReference StringRef;
+		StringRef.AssetLongPathname = Notify->GetPathName();
+
+		AddGameplayCueData_Internal(Notify->GameplayCueTag, StringRef);
+		BuildAccelerationMap_Internal();
+	}
 }
 
 /** Handles cleaning up an object library if it matches the passed in object */
 void UGameplayCueManager::HandleAssetDeleted(UObject *Object)
 {
+	FGameplayTag TagtoRemove;
 
+	UBlueprint* Blueprint = Cast<UBlueprint>(Object);
+	if (Blueprint && Blueprint->GeneratedClass)
+	{
+		UGameplayCueNotify* CDO = Cast<UGameplayCueNotify>(Blueprint->GeneratedClass->ClassDefaultObject);
+		if (CDO && CDO->GameplayCueTag.IsValid())
+		{
+			TagtoRemove = CDO->GameplayCueTag;
+		}
+	}
+
+	UGameplayCueNotify* Notify = Cast<UGameplayCueNotify>(Object);
+	if (Notify && Notify->GameplayCueTag.IsValid())
+	{
+		TagtoRemove = Notify->GameplayCueTag;
+	}
+
+	if (TagtoRemove.IsValid())
+	{
+		for (int32 idx = 0; idx < GameplayCueData.Num(); ++idx)
+		{
+			if (GameplayCueData[idx].GameplayCueTag == TagtoRemove)
+			{
+				GameplayCueData.RemoveAt(idx);
+				BuildAccelerationMap_Internal();
+				break;
+			}
+		}
+	}
 }
 
 #endif
@@ -281,4 +393,4 @@ FAutoConsoleCommandWithWorld PrintGameplayCueNotifyMapConsoleCommand(
 	TEXT("GameplayCue.PrintGameplayCueNotifyMap"),
 	TEXT("Displays GameplayCue notify map"),
 	FConsoleCommandWithWorldDelegate::CreateStatic(PrintGameplayCueNotifyMapConsoleCommandFunc)
-	);
\ No newline at end of file
+	);
diff --git a/Engine/Source/Runtime/GameplayAbilities/Public/AbilitySystemGlobals.h b/Engine/Source/Runtime/GameplayAbilities/Public/AbilitySystemGlobals.h
--- a/Engine/Source/Runtime/GameplayAbilities/Public/AbilitySystemGlobals.h
+++ b/Engine/Source/Runtime/GameplayAbilities/Public/AbilitySystemGlobals.h
@@ -43,6 +43,10 @@ class GAMEPLAYABILITIES_API UAbilitySystemGlobals : public UObject
 	UPROPERTY(config)
 	TArray<FString>	GameplayCueNotifyPaths;
 
+	/** Fully load all gameplay cues from the object library */
+	UPROPERTY(config)
+	bool	GameplayCueNotifyFullyLoad;
+
 	/** The class to instantiate as the globals object. Defaults to this class but can be overridden */
 	UPROPERTY(config)
 	FStringClassReference AbilitySystemGlobalsClassName;
diff --git a/Engine/Source/Runtime/GameplayAbilities/Public/GameplayCueManager.h b/Engine/Source/Runtime/GameplayAbilities/Public/GameplayCueManager.h
--- a/Engine/Source/Runtime/GameplayAbilities/Public/GameplayCueManager.h
+++ b/Engine/Source/Runtime/GameplayAbilities/Public/GameplayCueManager.h
@@ -13,6 +13,9 @@
  *	
  *	-Remove LoadObjectLibraryFromPaths from base implementation. Just have project pass in ObjectLibrary, let it figure out how to make it.
  *		-On Bob's recommendation: in hopes of removing object libraries coupling with directory structure.
+ *		-This becomes trickier with smaller projects/licensees. It would be nice to offer a path of least resistence for using gameplaycues (without having to implement a load object library function).
+ *		-This gets even trickier though when dealing with having to reload that objecet library when assets are added/deleted/renamed. Ideally the GameplaycueManager handles all of this, but if we are getting
+ *		an objectlibrary passed in, we don't know how to recreate it. Could add a delegate to invoke when we need the object library reloaded, but then ownership of that library feels pretty weird.
  *	
  *	-Async loading of all required GameplayCueNotifies
  *		-Currently async loaded on demand (first GameplayCueNotify is dropped)
@@ -59,6 +62,8 @@ struct FGameplayCueNotifyData
 	UPROPERTY(transient)
 	UGameplayCueNotify*		LoadedGameplayCueNotify;
 
+	FGameplayTag			GameplayCueTag;
+
 	int32 ParentDataIdx;
 };
 
@@ -134,4 +139,14 @@ class GAMEPLAYABILITIES_API UGameplayCueManager : public UDataAsset
 private:
 
 	virtual void HandleGameplayCueNotify_Internal(AActor* TargetActor, int32 DataIdx, EGameplayCueEvent::Type EventType, FGameplayCueParameters Parameters);
+
+	void LoadObjectLibrary_Internal();
+
+	void AddGameplayCueData_Internal(FGameplayTag  GameplayCueTag, FStringAssetReference StringRef);
+
+	void BuildAccelerationMap_Internal();
+
+	TArray<FString>	LoadedPaths;
+
+	bool bFullyLoad;
 };
\ No newline at end of file
