#UE4 - Call UGameplayAbility::InputPressed() / UGameplayAbility::InputReleased() on the server when inputs occur on the owning client

diff --git a/Engine/Source/Runtime/GameplayAbilities/Private/AbilitySystemComponent_Abilities.cpp b/Engine/Source/Runtime/GameplayAbilities/Private/AbilitySystemComponent_Abilities.cpp
--- a/Engine/Source/Runtime/GameplayAbilities/Private/AbilitySystemComponent_Abilities.cpp
+++ b/Engine/Source/Runtime/GameplayAbilities/Private/AbilitySystemComponent_Abilities.cpp
@@ -1136,6 +1136,16 @@ void UAbilitySystemComponent::ServerSetReplicatedAbilityKeyState_Implementation(
 			if (Spec.Ability)
 			{
 				Spec.InputPressed = Pressed;
+
+				// Call InputPressed() / InputReleased() on the affected abilities
+				if (Pressed)
+				{
+					AbilitySpecInputPressed(Spec);
+				}
+				else
+				{
+					AbilitySpecInputReleased(Spec);
+				}
 			}
 		}
 	}
