hide property

diff --git a/Engine/Source/Runtime/GameplayAbilities/Public/GameplayCueNotify.h b/Engine/Source/Runtime/GameplayAbilities/Public/GameplayCueNotify.h
--- a/Engine/Source/Runtime/GameplayAbilities/Public/GameplayCueNotify.h
+++ b/Engine/Source/Runtime/GameplayAbilities/Public/GameplayCueNotify.h
@@ -26,7 +26,6 @@ class GAMEPLAYABILITIES_API UGameplayCueNotify : public UDataAsset
 	virtual bool NeedsInstanceForEvent(EGameplayCueEvent::Type EventType) const;
 
 	virtual void OnOwnerDestroyed();
-	
 
 	virtual void PostInitProperties();
 
@@ -38,7 +37,7 @@ class GAMEPLAYABILITIES_API UGameplayCueNotify : public UDataAsset
 	FGameplayTag	GameplayCueTag;
 
 	/** Mirrors GameplayCueTag in order to be asset registry searchable */
-	UPROPERTY(AssetRegistrySearchable, EditDefaultsOnly, Category = GameplayCue)
+	UPROPERTY(AssetRegistrySearchable)
 	FName GameplayCueName;
 
 	/** Does this Cue override other cues, or is it called in addition to them? E.g., If this is Damage.Physical.Slash, we wont call Damage.Physical afer we run this cue. */
