Ending task from client or server will end the task on the other side.

diff --git a/Engine/Source/Runtime/GameplayAbilities/Private/Abilities/GameplayAbility.cpp b/Engine/Source/Runtime/GameplayAbilities/Private/Abilities/GameplayAbility.cpp
--- a/Engine/Source/Runtime/GameplayAbilities/Private/Abilities/GameplayAbility.cpp
+++ b/Engine/Source/Runtime/GameplayAbilities/Private/Abilities/GameplayAbility.cpp
@@ -138,7 +138,7 @@ bool UGameplayAbility::IsSupportedForNetworking() const
 	 *	an ability that is not described above.
 	 */
 
-	bool Supported = GetReplicationPolicy() != EGameplayAbilityReplicationPolicy::ReplicateNone || GetOuter()->IsA(UPackage::StaticClass());
+	bool Supported = GetReplicationPolicy() != EGameplayAbilityReplicationPolicy::ReplicateNo || GetOuter()->IsA(UPackage::StaticClass());
 	ensureMsgf(Supported, TEXT("Ability %s failed replication, if it is instanced it should have replication enabled"), *GetName());
 
 	return Supported;
@@ -292,6 +292,11 @@ void UGameplayAbility::EndAbility(const FGameplayAbilitySpecHandle Handle, const
 	}
 }
 
+void UGameplayAbility::EndAbility()
+{
+	EndAbility(CurrentSpecHandle, CurrentActorInfo, CurrentActivationInfo);
+}
+
 void UGameplayAbility::ActivateAbility(const FGameplayAbilitySpecHandle Handle, const FGameplayAbilityActorInfo* ActorInfo, const FGameplayAbilityActivationInfo ActivationInfo)
 {
 	SetCurrentInfo(Handle, ActorInfo, ActivationInfo);
@@ -488,21 +493,23 @@ bool UGameplayAbility::K2_CommitAbility()
 void UGameplayAbility::K2_EndAbility()
 {
 	check(CurrentActorInfo);
-	
-	EndAbility(CurrentSpecHandle, CurrentActorInfo, CurrentActivationInfo);
-}
-
-void UGameplayAbility::K2_EndAbilityFromServer()
-{
-	check(CurrentActorInfo);
-	check(CurrentActorInfo->AvatarActor.IsValid());
-	ENetRole NetMode = CurrentActorInfo->AvatarActor->Role;
 
-	if (NetMode == ROLE_Authority)		//Do we care about ability execution policy?
+	if (CurrentActorInfo->AvatarActor.IsValid())
 	{
-		CurrentActorInfo->AbilitySystemComponent->ClientEndAbility(CurrentSpecHandle);		//Does this need a prediction key?
-		K2_EndAbility();
+		check(CurrentActorInfo->AbilitySystemComponent.IsValid());
+		ENetRole NetMode = CurrentActorInfo->AvatarActor->Role;
+		UAbilitySystemComponent* MyASC = CurrentActorInfo->AbilitySystemComponent.Get();
+		//Do we care about ability execution policy? Do we need a prediction key?
+		if (NetMode == ROLE_Authority)
+		{
+			CurrentActorInfo->AbilitySystemComponent->ClientEndAbility(CurrentSpecHandle);
+		}
+		else
+		{
+			CurrentActorInfo->AbilitySystemComponent->ServerEndAbility(CurrentSpecHandle);
+		}
 	}
+	EndAbility();
 }
 
 // --------------------------------------------------------------------
diff --git a/Engine/Source/Runtime/GameplayAbilities/Private/AbilitySystemComponent_Abilities.cpp b/Engine/Source/Runtime/GameplayAbilities/Private/AbilitySystemComponent_Abilities.cpp
--- a/Engine/Source/Runtime/GameplayAbilities/Private/AbilitySystemComponent_Abilities.cpp
+++ b/Engine/Source/Runtime/GameplayAbilities/Private/AbilitySystemComponent_Abilities.cpp
@@ -278,7 +278,7 @@ UGameplayAbility* UAbilitySystemComponent::CreateNewInstanceOfAbility(FGameplayA
 	check(AbilityInstance);
 
 	// Add it to one of our instance lists so that it doesn't GC.
-	if (AbilityInstance->GetReplicationPolicy() != EGameplayAbilityReplicationPolicy::ReplicateNone)
+	if (AbilityInstance->GetReplicationPolicy() != EGameplayAbilityReplicationPolicy::ReplicateNo)
 	{
 		Spec.ReplicatedInstances.Add(AbilityInstance);
 		AllReplicatedInstancedAbilities.Add(AbilityInstance);
@@ -310,7 +310,7 @@ void UAbilitySystemComponent::NotifyAbilityEnded(FGameplayAbilitySpecHandle Hand
 	{
 		check(Ability->HasAnyFlags(RF_ClassDefaultObject) == false);	// Should never be calling this on a CDO for an instanced ability!
 
-		if (Ability->GetReplicationPolicy() != EGameplayAbilityReplicationPolicy::ReplicateNone)
+		if (Ability->GetReplicationPolicy() != EGameplayAbilityReplicationPolicy::ReplicateNo)
 		{
 			if (GetOwnerRole() == ROLE_Authority)
 			{
@@ -453,7 +453,7 @@ bool UAbilitySystemComponent::TryActivateAbility(FGameplayAbilitySpecHandle Hand
 	}
 
 	// Always do a non instanced CanActivate check
-	if (!Ability->CanActivateAbility(Handle, ActorInfo))
+ 	if (!Ability->CanActivateAbility(Handle, ActorInfo))
 	{
 		return false;
 	}
@@ -513,7 +513,7 @@ bool UAbilitySystemComponent::TryActivateAbility(FGameplayAbilitySpecHandle Hand
 			// We lack the code to predict spawning an instance of the execution and then merge/combine
 			// with the server spawned version when it arrives.
 
-			if (Ability->GetReplicationPolicy() == EGameplayAbilityReplicationPolicy::ReplicateNone)
+			if (Ability->GetReplicationPolicy() == EGameplayAbilityReplicationPolicy::ReplicateNo)
 			{
 				InstancedAbility = CreateNewInstanceOfAbility(*Spec, Ability);
 				InstancedAbility->CallActivateAbility(Handle, ActorInfo, ActivationInfo, OnGameplayAbilityEndedDelegate);
@@ -574,15 +574,21 @@ void UAbilitySystemComponent::ServerTryActivateAbility_Implementation(FGameplayA
 			FExecutingAbilityInfo& ExecutingAbilityInfo = ExecutingServerAbilities[ExecutedIdx];
 			if (ExecutingAbilityInfo.PredictionKey.Current == PredictionKey.Base && ExecutingAbilityInfo.Handle == Handle)
 			{
-				switch(ExecutingAbilityInfo.State)
+				switch (ExecutingAbilityInfo.State)
 				{
-					case EAbilityExecutionState::Failed:
-						ClientActivateAbilityFailed(Handle, PredictionKey.Current);
-						break;
-					case EAbilityExecutionState::Executing:
-					case EAbilityExecutionState::Succeeded:
-						ClientActivateAbilitySucceed(Handle, PredictionKey.Current);
-						break;
+				case EAbilityExecutionState::Failed:
+					ClientActivateAbilityFailed(Handle, PredictionKey.Current);
+					break;
+				case EAbilityExecutionState::Executing:
+				case EAbilityExecutionState::Succeeded:
+					ClientActivateAbilitySucceed(Handle, PredictionKey.Current);
+					//Client commands to end the ability that come in after this point are considered for this instance
+					if (AbilityToActivate->HasAnyFlags(RF_ClassDefaultObject) == false)
+					{
+						//Only applies to instanced abilities
+						AbilityToActivate->CurrentActivationInfo.bCanBeEndedByOtherInstance = true;
+					}
+					break;
 				}
 
 				ExecutingServerAbilities.RemoveAtSwap(ExecutedIdx);
@@ -604,6 +610,12 @@ void UAbilitySystemComponent::ServerTryActivateAbility_Implementation(FGameplayA
 	if (TryActivateAbility(Handle, PredictionKey, &InstancedAbility))
 	{
 		ClientActivateAbilitySucceed(Handle, PredictionKey.Current);
+		//Client commands to end the ability that come in after this point are considered for this instance
+		if (InstancedAbility)
+		{
+			//Only applies to instanced abilities
+			InstancedAbility->CurrentActivationInfo.bCanBeEndedByOtherInstance = true;
+		}
 	}
 	else
 	{
@@ -624,23 +636,34 @@ bool UAbilitySystemComponent::ServerTryActivateAbility_Validate(FGameplayAbility
 	return true;
 }
 
-void UAbilitySystemComponent::ClientEndAbility_Implementation(FGameplayAbilitySpecHandle AbilityToEnd)
+//This is only called when ending an ability in response to a remote instruction.
+void UAbilitySystemComponent::EndAbility(FGameplayAbilitySpecHandle AbilityToEnd)
 {
 	FGameplayAbilitySpec* AbilitySpec = FindAbilitySpecFromHandle(AbilityToEnd);
-	if (AbilitySpec && AbilitySpec->Ability && AbilitySpec->IsActive())
+	if (AbilitySpec && AbilitySpec->Ability && AbilitySpec->IsActive() && AbilitySpec->ActivationInfo.bCanBeEndedByOtherInstance)
 	{
 		for (auto Instance : AbilitySpec->NonReplicatedInstances)
 		{
-			Instance->K2_EndAbility();
+			Instance->EndAbility();
 		}
 	}
 }
 
-bool UAbilitySystemComponent::ClientEndAbility_Validate(FGameplayAbilitySpecHandle AbilityToActivate)
+void UAbilitySystemComponent::ServerEndAbility_Implementation(FGameplayAbilitySpecHandle AbilityToEnd)
+{
+	EndAbility(AbilityToEnd);
+}
+
+bool UAbilitySystemComponent::ServerEndAbility_Validate(FGameplayAbilitySpecHandle AbilityToActivate)
 {
 	return true;
 }
 
+void UAbilitySystemComponent::ClientEndAbility_Implementation(FGameplayAbilitySpecHandle AbilityToEnd)
+{
+	EndAbility(AbilityToEnd);
+}
+
 static_assert(sizeof(int16) == sizeof(FPredictionKey::KeyType), "Sizeof PredictionKey::KeyType does not match RPC parameters in AbilitySystemComponent ClientActivateAbilityFailed_Implementation");
 
 void UAbilitySystemComponent::ClientActivateAbilityFailed_Implementation(FGameplayAbilitySpecHandle Handle, int16 PredictionKey)
diff --git a/Engine/Source/Runtime/GameplayAbilities/Private/GameplayAbilityTypes.cpp b/Engine/Source/Runtime/GameplayAbilities/Private/GameplayAbilityTypes.cpp
--- a/Engine/Source/Runtime/GameplayAbilities/Private/GameplayAbilityTypes.cpp
+++ b/Engine/Source/Runtime/GameplayAbilities/Private/GameplayAbilityTypes.cpp
@@ -102,6 +102,8 @@ void FGameplayAbilityActivationInfo::GenerateNewPredictionKey() const
 void FGameplayAbilityActivationInfo::SetActivationConfirmed()
 {
 	ActivationMode = EGameplayAbilityActivationMode::Confirmed;
+	//Remote (server) commands to end the ability that come in after this point are considered for this instance
+	bCanBeEndedByOtherInstance = true;
 }
 
 void FGameplayAbilityActivationInfo::SetPredictionStale()
diff --git a/Engine/Source/Runtime/GameplayAbilities/Public/Abilities/GameplayAbility.h b/Engine/Source/Runtime/GameplayAbilities/Public/Abilities/GameplayAbility.h
--- a/Engine/Source/Runtime/GameplayAbilities/Public/Abilities/GameplayAbility.h
+++ b/Engine/Source/Runtime/GameplayAbilities/Public/Abilities/GameplayAbility.h
@@ -323,10 +323,8 @@ protected:
 	UFUNCTION(BlueprintCallable, Category = Ability, FriendlyName="EndAbility")
 	virtual void K2_EndAbility();
 
-	UFUNCTION(BlueprintCallable, Category = Ability, FriendlyName = "EndAbilityFromServer")
-	virtual void K2_EndAbilityFromServer();
-
 	virtual void EndAbility(const FGameplayAbilitySpecHandle Handle, const FGameplayAbilityActorInfo* ActorInfo, const FGameplayAbilityActivationInfo ActivationInfo);
+	virtual void EndAbility();		//Convenience function for calling with default parameters
 
 	// -------------------------------------
 	//	GameplayEffects
diff --git a/Engine/Source/Runtime/GameplayAbilities/Public/Abilities/GameplayAbilityTypes.h b/Engine/Source/Runtime/GameplayAbilities/Public/Abilities/GameplayAbilityTypes.h
--- a/Engine/Source/Runtime/GameplayAbilities/Public/Abilities/GameplayAbilityTypes.h
+++ b/Engine/Source/Runtime/GameplayAbilities/Public/Abilities/GameplayAbilityTypes.h
@@ -26,9 +26,14 @@ namespace EGameplayAbilityInstancingPolicy
 
 	enum Type
 	{
-		NonInstanced,			// This ability is never instanced. Anything that executes the ability is operating on the CDO.
-		InstancedPerActor,		// Each actor gets their own instance of this ability. State can be saved, replication is possible.
-		InstancedPerExecution,	// We instance this ability each time it is executed. Replication possible but not recommended.
+		// This ability is never instanced. Anything that executes the ability is operating on the CDO.
+		NonInstanced,
+
+		// Each actor gets their own instance of this ability. State can be saved, replication is possible.
+		InstancedPerActor,
+
+		// We instance this ability each time it is executed. Replication possible but not recommended.
+		InstancedPerExecution,
 	};
 }
 
@@ -41,9 +46,14 @@ namespace EGameplayAbilityNetExecutionPolicy
 
 	enum Type
 	{
-		Predictive		UMETA(DisplayName = "Predictive"),	// Part of this ability runs predictively on the client.
-		Server			UMETA(DisplayName = "Server"),		// This ability must be OK'd by the server before doing anything on a client.
-		Client			UMETA(DisplayName = "Client"),		// This ability runs as long the client says it does.
+		// Part of this ability runs predictively on the client.
+		Predictive		UMETA(DisplayName = "Predictive"),
+
+		// This ability must be OK'd by the server before doing anything on a client.
+		Server			UMETA(DisplayName = "Server"),
+
+		// This ability runs as long the client says it does.
+		Client			UMETA(DisplayName = "Client"),
 	};
 }
 
@@ -56,9 +66,11 @@ namespace EGameplayAbilityReplicationPolicy
 
 	enum Type
 	{
-		ReplicateNone		UMETA(DisplayName = "Replicate None"),	// We don't replicate the instance of the ability to anyone.
-		ReplicateAll		UMETA(DisplayName = "Replicate All"),	// We replicate the instance of the ability to everyone (even simulating clients).
-		ReplicateOwner		UMETA(DisplayName = "Replicate Owner"),	// Only replicate the instance of the ability to the owner.
+		// We don't replicate the instance of the ability to anyone.
+		ReplicateNo			UMETA(DisplayName = "Do Not Replicate"),
+
+		// We replicate the instance of the ability to the owner.
+		ReplicateYes		UMETA(DisplayName = "Replicate"),
 	};
 }
 
@@ -68,11 +80,17 @@ namespace EGameplayAbilityActivationMode
 {
 	enum Type
 	{
-		Authority,			// We are the authority activating this ability
-		NonAuthority,		// We are not the authority but aren't predicting yet. This is a mostly invalid state to be in.
+		// We are the authority activating this ability
+		Authority,
 
-		Predicting,			// We are predicting the activation of this ability
-		Confirmed,			// We are not the authority, but the authority has confirmed this activation
+		// We are not the authority but aren't predicting yet. This is a mostly invalid state to be in.
+		NonAuthority,
+
+		// We are predicting the activation of this ability
+		Predicting,
+
+		// We are not the authority, but the authority has confirmed this activation
+		Confirmed,
 	};
 }
 
@@ -194,12 +212,14 @@ struct GAMEPLAYABILITIES_API FGameplayAbilityActivationInfo
 
 	FGameplayAbilityActivationInfo()
 		: ActivationMode(EGameplayAbilityActivationMode::Authority)
+		, bCanBeEndedByOtherInstance(false)
 	{
 
 	}
 
 	FGameplayAbilityActivationInfo(AActor* InActor, FPredictionKey InPredictionKey)
 		: PredictionKey(InPredictionKey)
+		, bCanBeEndedByOtherInstance(false)
 	{
 		// On Init, we are either Authority or NonAuthority. We haven't been given a PredictionKey and we haven't been confirmed.
 		// NonAuthority essentially means 'I'm not sure what how I'm going to do this yet'.
@@ -209,6 +229,7 @@ struct GAMEPLAYABILITIES_API FGameplayAbilityActivationInfo
 	FGameplayAbilityActivationInfo(EGameplayAbilityActivationMode::Type InType, FPredictionKey InPredictionKey = FPredictionKey())
 		: ActivationMode(InType)
 		, PredictionKey(InPredictionKey)
+		, bCanBeEndedByOtherInstance(false)
 	{
 	}
 
@@ -236,6 +257,10 @@ struct GAMEPLAYABILITIES_API FGameplayAbilityActivationInfo
 	UPROPERTY(BlueprintReadOnly, Category = "ActorInfo")
 	mutable TEnumAsByte<EGameplayAbilityActivationMode::Type>	ActivationMode;
 
+	/** An ability that runs on multiple game instances can be canceled by a remote instance, but only if that remote instance has already confirmed starting it. */
+	UPROPERTY()
+	bool bCanBeEndedByOtherInstance;
+
 private:
 
 	UPROPERTY()
diff --git a/Engine/Source/Runtime/GameplayAbilities/Public/AbilitySystemComponent.h b/Engine/Source/Runtime/GameplayAbilities/Public/AbilitySystemComponent.h
--- a/Engine/Source/Runtime/GameplayAbilities/Public/AbilitySystemComponent.h
+++ b/Engine/Source/Runtime/GameplayAbilities/Public/AbilitySystemComponent.h
@@ -384,7 +384,13 @@ class GAMEPLAYABILITIES_API UAbilitySystemComponent : public UActorComponent, pu
 	UFUNCTION(Server, reliable, WithValidation)
 	void	ServerTryActivateAbility(FGameplayAbilitySpecHandle AbilityToActivate, FPredictionKey PredictionKey);
 
-	UFUNCTION(Client, reliable, WithValidation)
+	/** Called by ServerEndAbility and ClientEndAbility; avoids code duplication. */
+	void EndAbility(FGameplayAbilitySpecHandle AbilityToEnd);
+
+	UFUNCTION(Server, reliable, WithValidation)
+	void	ServerEndAbility(FGameplayAbilitySpecHandle AbilityToEnd);
+
+	UFUNCTION(Client, reliable)
 	void	ClientEndAbility(FGameplayAbilitySpecHandle AbilityToEnd);
 
 	UFUNCTION(Client, Reliable)
