GUBP fix - variable init order.

diff --git a/Engine/Source/Runtime/GameplayAbilities/Public/Abilities/GameplayAbilityTypes.h b/Engine/Source/Runtime/GameplayAbilities/Public/Abilities/GameplayAbilityTypes.h
--- a/Engine/Source/Runtime/GameplayAbilities/Public/Abilities/GameplayAbilityTypes.h
+++ b/Engine/Source/Runtime/GameplayAbilities/Public/Abilities/GameplayAbilityTypes.h
@@ -218,8 +218,8 @@ struct GAMEPLAYABILITIES_API FGameplayAbilityActivationInfo
 	}
 
 	FGameplayAbilityActivationInfo(AActor* InActor, FPredictionKey InPredictionKey)
-		: PredictionKey(InPredictionKey)
-		, bCanBeEndedByOtherInstance(false)
+		: bCanBeEndedByOtherInstance(false)
+		, PredictionKey(InPredictionKey)
 	{
 		// On Init, we are either Authority or NonAuthority. We haven't been given a PredictionKey and we haven't been confirmed.
 		// NonAuthority essentially means 'I'm not sure what how I'm going to do this yet'.
@@ -228,8 +228,8 @@ struct GAMEPLAYABILITIES_API FGameplayAbilityActivationInfo
 
 	FGameplayAbilityActivationInfo(EGameplayAbilityActivationMode::Type InType, FPredictionKey InPredictionKey = FPredictionKey())
 		: ActivationMode(InType)
-		, PredictionKey(InPredictionKey)
 		, bCanBeEndedByOtherInstance(false)
+		, PredictionKey(InPredictionKey)
 	{
 	}
 
