Fixed InputPress bool having the wrong value immediately upon ability startup.

diff --git a/Engine/Source/Runtime/GameplayAbilities/Private/AbilitySystemComponent_Abilities.cpp b/Engine/Source/Runtime/GameplayAbilities/Private/AbilitySystemComponent_Abilities.cpp
--- a/Engine/Source/Runtime/GameplayAbilities/Private/AbilitySystemComponent_Abilities.cpp
+++ b/Engine/Source/Runtime/GameplayAbilities/Private/AbilitySystemComponent_Abilities.cpp
@@ -492,7 +492,7 @@ bool UAbilitySystemComponent::TryActivateAbility(FGameplayAbilitySpecHandle Hand
 	}
 	else if (Ability->GetNetExecutionPolicy() == EGameplayAbilityNetExecutionPolicy::Server)
 	{
-		ServerTryActivateAbility(Handle, FPredictionKey());
+		ServerTryActivateAbility(Handle, Spec->InputPressed, FPredictionKey());
 	}
 	else if (Ability->GetNetExecutionPolicy() == EGameplayAbilityNetExecutionPolicy::Predictive)
 	{
@@ -502,7 +502,7 @@ bool UAbilitySystemComponent::TryActivateAbility(FGameplayAbilitySpecHandle Hand
 		FPredictionKey ThisPredictionKey = ActivationInfo.GetPredictionKeyForNewAction();
 		
 		// This must be called immediately after GeneratePredictionKey to prevent problems with recursively activating abilities
-		ServerTryActivateAbility(Handle, ThisPredictionKey);
+		ServerTryActivateAbility(Handle, Spec->InputPressed, ThisPredictionKey);
 
 		// If this PredictionKey is rejected, we will call OnClientActivateAbilityFailed.
 		ThisPredictionKey.NewRejectedDelegate().BindUObject(this, &UAbilitySystemComponent::OnClientActivateAbilityFailed, Handle, ThisPredictionKey.Current);
@@ -542,7 +542,7 @@ bool UAbilitySystemComponent::TryActivateAbility(FGameplayAbilitySpecHandle Hand
 	return true;
 }
 
-void UAbilitySystemComponent::ServerTryActivateAbility_Implementation(FGameplayAbilitySpecHandle Handle, FPredictionKey PredictionKey)
+void UAbilitySystemComponent::ServerTryActivateAbility_Implementation(FGameplayAbilitySpecHandle Handle, bool InputPressed, FPredictionKey PredictionKey)
 {
 #if !(UE_BUILD_SHIPPING || UE_BUILD_TEST)
 	if (DenyClientActivation > 0)
@@ -605,6 +605,7 @@ void UAbilitySystemComponent::ServerTryActivateAbility_Implementation(FGameplayA
 	}
 
 	UGameplayAbility* InstancedAbility = NULL;
+	Spec->InputPressed = InputPressed; // Tasks that check input may need this to be right immediately on ability startup.
 
 	// Attempt to activate the ability (server side) and tell the client if it succeeded or failed.
 	if (TryActivateAbility(Handle, PredictionKey, &InstancedAbility))
@@ -631,7 +632,7 @@ void UAbilitySystemComponent::ServerTryActivateAbility_Implementation(FGameplayA
 	}
 }
 
-bool UAbilitySystemComponent::ServerTryActivateAbility_Validate(FGameplayAbilitySpecHandle Handle, FPredictionKey PredictionKey)
+bool UAbilitySystemComponent::ServerTryActivateAbility_Validate(FGameplayAbilitySpecHandle Handle, bool InputPressed, FPredictionKey PredictionKey)
 {
 	return true;
 }
diff --git a/Engine/Source/Runtime/GameplayAbilities/Public/AbilitySystemComponent.h b/Engine/Source/Runtime/GameplayAbilities/Public/AbilitySystemComponent.h
--- a/Engine/Source/Runtime/GameplayAbilities/Public/AbilitySystemComponent.h
+++ b/Engine/Source/Runtime/GameplayAbilities/Public/AbilitySystemComponent.h
@@ -382,7 +382,7 @@ class GAMEPLAYABILITIES_API UAbilitySystemComponent : public UActorComponent, pu
 	void	OnRep_ActivateAbilities();
 
 	UFUNCTION(Server, reliable, WithValidation)
-	void	ServerTryActivateAbility(FGameplayAbilitySpecHandle AbilityToActivate, FPredictionKey PredictionKey);
+	void	ServerTryActivateAbility(FGameplayAbilitySpecHandle AbilityToActivate, bool InputPressed, FPredictionKey PredictionKey);
 
 	/** Called by ServerEndAbility and ClientEndAbility; avoids code duplication. */
 	void EndAbility(FGameplayAbilitySpecHandle AbilityToEnd);
