new file checkin to make passing shelved CLs around easier
-Contents are #if out until full changes are ready.

diff --git a/Engine/Source/Runtime/GameplayAbilities/Private/GameplayEffectAggregator.cpp b/Engine/Source/Runtime/GameplayAbilities/Private/GameplayEffectAggregator.cpp
new file mode 100644
--- /dev/null
+++ b/Engine/Source/Runtime/GameplayAbilities/Private/GameplayEffectAggregator.cpp
@@ -0,0 +1,168 @@
+// Copyright 1998-2013 Epic Games, Inc. All Rights Reserved.
+
+#include "AbilitySystemPrivatePCH.h"
+#include "GameplayEffectAggregator.h"
+
+#if GE_REFACTOR
+#pragma optimize( "", off )
+
+bool FAggregatorMod::Qualifies(const FAggregatorEvaluateParameters& Parameters) const
+{
+	bool SourceMet = (!SourceTagReqs || SourceTagReqs->RequirementsMet(Parameters.SourceTags));
+	bool TargetMet = (!TargetTagReqs || TargetTagReqs->RequirementsMet(Parameters.TargetTags));
+
+	return SourceMet && TargetMet;
+}
+
+float FAggregator::Evaluate(const FAggregatorEvaluateParameters& Parameters) const
+{
+	for (const FAggregatorMod& Mod : Mods[EGameplayModOp::Override])
+	{
+		if (Mod.Qualifies(Parameters))
+		{
+			return Mod.EvaluatedMagnitude;
+		}
+	}
+
+	float Additive = SumMods(Mods[EGameplayModOp::Additive], 0.f, Parameters);
+	float Multiplicitive = SumMods(Mods[EGameplayModOp::Multiplicitive], 1.f, Parameters);
+	float Division = SumMods(Mods[EGameplayModOp::Division], 1.f, Parameters);
+
+	if (FMath::IsNearlyZero(Division))
+	{
+		ABILITY_LOG(Warning, TEXT("Division summation was 0.0f in FAggregator."));
+		Division = 1.f;
+	}
+
+	return ((BaseValue + Additive) * Multiplicitive) / Division;
+}
+
+void FAggregator::SetBaseValue(float NewBaseValue)
+{
+	BaseValue = NewBaseValue;
+	OnDirty.Broadcast(this);
+}
+
+float FAggregator::StaticExecModOnBaseValue(float BaseValue, const struct FModifierSpec& Mod)
+{
+	float EvaluatedMagnitude = Mod.Info.Magnitude.GetValueChecked(); // fixme ?
+
+	switch (Mod.Info.ModifierOp)
+	{
+	case EGameplayModOp::Override:
+	{
+		BaseValue = EvaluatedMagnitude;
+		break;
+	}
+	case EGameplayModOp::Additive:
+	{
+		BaseValue += EvaluatedMagnitude;
+		break;
+	}
+	case EGameplayModOp::Multiplicitive:
+	{
+		BaseValue *= EvaluatedMagnitude;
+		break;
+	}
+	case EGameplayModOp::Division:
+	{
+		if (FMath::IsNearlyZero(EvaluatedMagnitude) == false)
+		{
+			BaseValue /= EvaluatedMagnitude;
+		}
+		break;
+	}
+	}
+
+	return BaseValue;
+}
+
+void FAggregator::ExecModOnBaseValue(const struct FModifierSpec& Mod)
+{
+	BaseValue = StaticExecModOnBaseValue(BaseValue, Mod);
+	OnDirty.Broadcast(this);
+}
+
+float FAggregator::SumMods(const TArray<FAggregatorMod> &Mods, float Bias, const FAggregatorEvaluateParameters& Parameters) const
+{
+	float Sum = Bias;
+
+	for (const FAggregatorMod& Mod : Mods)
+	{
+		if (Mod.Qualifies(Parameters))
+		{
+			Sum += (Mod.EvaluatedMagnitude - Bias);
+		}
+	}
+
+	return Sum;
+}
+
+void FAggregator::AddMod(const FGameplayTagRequirements* SourceTagReqs, const FGameplayTagRequirements* TargetTagReqs, float EvaluatedMagnitude, FActiveGameplayEffectHandle ActiveHandle)
+{
+	check(Data);
+	check(Data->Modifiers.IsValidIndex(ModIdx));
+
+	TArray<FAggregatorMod> &ModList = Mods[Data->Modifiers[ModIdx].ModifierOp];
+
+	int32 NewIdx = ModList.AddUninitialized();
+	FAggregatorMod& NewMod = ModList[NewIdx];
+
+	NewMod.Data = Data;
+	NewMod.ModIdx = ModIdx;
+	NewMod.EvaluatedMagnitude = EvaluatedMagnitude;
+	NewMod.ActiveHandle = ActiveHandle;
+
+	OnDirty.Broadcast(this );
+}
+
+void FAggregator::RemoveMod(FActiveGameplayEffectHandle ActiveHandle)
+{
+	if (ActiveHandle.IsValid())
+	{
+		RemoveModsWithActiveHandle(Mods[EGameplayModOp::Additive], ActiveHandle);
+		RemoveModsWithActiveHandle(Mods[EGameplayModOp::Multiplicitive], ActiveHandle);
+		RemoveModsWithActiveHandle(Mods[EGameplayModOp::Division], ActiveHandle);
+		RemoveModsWithActiveHandle(Mods[EGameplayModOp::Override], ActiveHandle);
+	}
+
+	OnDirty.Broadcast(this);
+}
+
+void FAggregator::RemoveModsWithActiveHandle(TArray<FAggregatorMod>& Mods, FActiveGameplayEffectHandle ActiveHandle)
+{
+	check(ActiveHandle.IsValid());
+
+	for(int32 idx=Mods.Num()-1; idx >= 0; --idx)
+	{
+		if (Mods[idx].ActiveHandle == ActiveHandle)
+		{
+			Mods.RemoveAtSwap(idx, 1, false);
+		}
+	}
+}
+
+void FAggregator::TakeSnapshotOf(const FAggregator& AggToSnapshot)
+{
+	BaseValue = AggToSnapshot.BaseValue;
+	for (int32 idx=0; idx < ARRAY_COUNT(Mods); ++idx)
+	{
+		Mods[idx] = AggToSnapshot.Mods[idx];
+	}
+}
+
+void FAggregatorRef::TakeSnapshotOf(const FAggregatorRef& RefToSnapshot)
+{
+	if (RefToSnapshot.Data.IsValid())
+	{
+		FAggregator* SrcData = RefToSnapshot.Data.Get();
+
+		Data = TSharedPtr<FAggregator>(new FAggregator());
+		Data->TakeSnapshotOf(*SrcData);
+	}
+	else
+	{
+		Data.Reset();
+	}
+}
+#endif
\ No newline at end of file
diff --git a/Engine/Source/Runtime/GameplayAbilities/Private/GameplayEffectCalculation.cpp b/Engine/Source/Runtime/GameplayAbilities/Private/GameplayEffectCalculation.cpp
new file mode 100644
--- /dev/null
+++ b/Engine/Source/Runtime/GameplayAbilities/Private/GameplayEffectCalculation.cpp
@@ -0,0 +1,24 @@
+// Copyright 1998-2014 Epic Games, Inc. All Rights Reserved.
+
+#include "AbilitySystemPrivatePCH.h"
+
+#if GE_REFACTOR
+#include "GameplayEffectCalculation.h"
+
+UGameplayEffectCalculation::UGameplayEffectCalculation(const FObjectInitializer& ObjectInitializer)
+	: Super(ObjectInitializer)
+{
+
+}
+
+void UGameplayEffectCalculation::Execute(FGameplayEffectSpec& Spec, int32 ExecutionIdx, UAbilitySystemComponent* TargetAbilitySystemComponent) const
+{
+	
+}
+
+void UGameplayEffectCalculation::GetCaptureDefinitions(OUT TArray<FGameplayEffectAttributeCaptureDefinition>& Definitions) const
+{
+
+}
+
+#endif
\ No newline at end of file
diff --git a/Engine/Source/Runtime/GameplayAbilities/Public/GameplayEffectAggregator.h b/Engine/Source/Runtime/GameplayAbilities/Public/GameplayEffectAggregator.h
new file mode 100644
--- /dev/null
+++ b/Engine/Source/Runtime/GameplayAbilities/Public/GameplayEffectAggregator.h
@@ -0,0 +1,73 @@
+// Copyright 1998-2014 Epic Games, Inc. All Rights Reserved.
+
+#pragma once
+
+#if GE_REFACTOR
+
+#include "GameplayEffectTypes.h"
+
+// #include "GameplayEffectAggregator.generated.h"
+
+struct GAMEPLAYABILITIES_API FAggregatorEvaluateParameters
+{
+	FGameplayTagContainer	SourceTags;
+	FGameplayTagContainer	TargetTags;
+};
+
+struct GAMEPLAYABILITIES_API FAggregatorMod
+{
+	const FGameplayTagRequirements*	SourceTagReqs;
+	const FGameplayTagRequirements*	TargetTagReqs;
+
+	float EvaluatedMagnitude;		// Magnitude this mod was last evaluated at
+
+	FActiveGameplayEffectHandle ActiveHandle;	// Handle of the active GameplayEffect we are tied to (if any)
+
+	bool Qualifies(const FAggregatorEvaluateParameters& Parameters) const;
+};
+
+
+struct GAMEPLAYABILITIES_API FAggregator : public TSharedFromThis<FAggregator>
+{
+	DECLARE_MULTICAST_DELEGATE_OneParam(FOnAggregatorDirty, FAggregator*);
+
+	FAggregator(float InBaseValue=0.f) : BaseValue(InBaseValue) { }
+
+	void SetBaseValue(float NewBaseValue);
+	void ExecModOnBaseValue(const struct FModifierSpec& Mod);
+
+	static float StaticExecModOnBaseValue(float BaseValue, const struct FModifierSpec& Mod);
+
+	void AddMod(float EvaluatedData, TEnumAsByte<EGameplayModOp::Type> ModifierOp, const FGameplayTagRequirements*	SourceTagReqs, const FGameplayTagRequirements* TargetTagReqs, FActiveGameplayEffectHandle ActiveHandle = FActiveGameplayEffectHandle());
+	void RemoveMod(FActiveGameplayEffectHandle ActiveHandle);
+
+	FOnAggregatorDirty OnDirty;
+
+	float Evaluate(const FAggregatorEvaluateParameters& Parameters) const;
+
+	void TakeSnapshotOf(const FAggregator& AggToSnapshot);
+
+private:
+
+	float	BaseValue;
+	TArray<FAggregatorMod>	Mods[EGameplayModOp::Max];
+
+	float SumMods(const TArray<FAggregatorMod> &Mods, float Bias, const FAggregatorEvaluateParameters& Parameters) const;
+	void RemoveModsWithActiveHandle(TArray<FAggregatorMod>& Mods, FActiveGameplayEffectHandle ActiveHandle);
+
+	friend struct FAggregator;
+};
+
+struct GAMEPLAYABILITIES_API FAggregatorRef
+{
+	FAggregatorRef() { }
+	FAggregatorRef(FAggregator* InData) : Data(InData) { }
+
+	FAggregator* Get() const { return Data.Get(); }
+
+	TSharedPtr<FAggregator>	Data;
+
+	void TakeSnapshotOf(const FAggregatorRef& RefToSnapshot);
+};
+
+#endif
\ No newline at end of file
diff --git a/Engine/Source/Runtime/GameplayAbilities/Public/GameplayEffectCalculation.h b/Engine/Source/Runtime/GameplayAbilities/Public/GameplayEffectCalculation.h
new file mode 100644
--- /dev/null
+++ b/Engine/Source/Runtime/GameplayAbilities/Public/GameplayEffectCalculation.h
@@ -0,0 +1,40 @@
+// Copyright 1998-2014 Epic Games, Inc. All Rights Reserved.
+
+#pragma once
+
+
+#include "GameplayEffectCalculation.generated.h"
+
+#if GE_REFACTOR
+
+struct FGameplayEffectSpec;
+struct FGameplayEffectAttributeCaptureDefinition;
+
+class UAbilitySystemComponent;
+
+UCLASS(abstract)
+class GAMEPLAYABILITIES_API UGameplayEffectCalculation : public UObject
+{
+	GENERATED_UCLASS_BODY()
+
+	virtual void Execute(FGameplayEffectSpec& Spec, int32 ExecutionIdx, UAbilitySystemComponent* TargetAbilitySystemComponent) const;
+
+	virtual void GetCaptureDefinitions(OUT TArray<FGameplayEffectAttributeCaptureDefinition>& Definitions) const;
+};
+
+// yuck?
+
+#define GE_REQ_SRC(SET, PROP, SNAPSHOT) \
+{ \
+	static UProperty* ThisProperty = FindFieldChecked<UProperty>(SET::StaticClass(), GET_MEMBER_NAME_CHECKED(SET, PROP)); \
+	Definitions.Add(FGameplayEffectAttributeCaptureDefinition(FGameplayAttribute(ThisProperty), EGameplayEffectAttributeCaptureSource::Source, SNAPSHOT)); \
+}
+
+#define GE_REQ_TARGET(SET, PROP, SNAPSHOT) \
+{ \
+	static UProperty* ThisProperty = FindFieldChecked<UProperty>(SET::StaticClass(), GET_MEMBER_NAME_CHECKED(SET, PROP)); \
+	Definitions.Add(FGameplayEffectAttributeCaptureDefinition(FGameplayAttribute(ThisProperty), EGameplayEffectAttributeCaptureSource::Target, SNAPSHOT)); \
+}
+
+
+#endif
\ No newline at end of file
