compile fix again

diff --git a/Engine/Source/Runtime/GameplayAbilities/Private/GameplayEffectCalculation.cpp b/Engine/Source/Runtime/GameplayAbilities/Private/GameplayEffectCalculation.cpp
--- a/Engine/Source/Runtime/GameplayAbilities/Private/GameplayEffectCalculation.cpp
+++ b/Engine/Source/Runtime/GameplayAbilities/Private/GameplayEffectCalculation.cpp
@@ -1,24 +1,4 @@
 // Copyright 1998-2014 Epic Games, Inc. All Rights Reserved.
 
 #include "AbilitySystemPrivatePCH.h"
-
-#if GE_REFACTOR
 #include "GameplayEffectCalculation.h"
-
-UGameplayEffectCalculation::UGameplayEffectCalculation(const FObjectInitializer& ObjectInitializer)
-	: Super(ObjectInitializer)
-{
-
-}
-
-void UGameplayEffectCalculation::Execute(FGameplayEffectSpec& Spec, int32 ExecutionIdx, UAbilitySystemComponent* TargetAbilitySystemComponent) const
-{
-	
-}
-
-void UGameplayEffectCalculation::GetCaptureDefinitions(OUT TArray<FGameplayEffectAttributeCaptureDefinition>& Definitions) const
-{
-
-}
-
-#endif
\ No newline at end of file
diff --git a/Engine/Source/Runtime/GameplayAbilities/Public/GameplayEffectCalculation.h b/Engine/Source/Runtime/GameplayAbilities/Public/GameplayEffectCalculation.h
--- a/Engine/Source/Runtime/GameplayAbilities/Public/GameplayEffectCalculation.h
+++ b/Engine/Source/Runtime/GameplayAbilities/Public/GameplayEffectCalculation.h
@@ -3,38 +3,3 @@
 #pragma once
 
 
-#include "GameplayEffectCalculation.generated.h"
-
-#if GE_REFACTOR
-
-struct FGameplayEffectSpec;
-struct FGameplayEffectAttributeCaptureDefinition;
-
-class UAbilitySystemComponent;
-
-UCLASS(abstract)
-class GAMEPLAYABILITIES_API UGameplayEffectCalculation : public UObject
-{
-	GENERATED_UCLASS_BODY()
-
-	virtual void Execute(FGameplayEffectSpec& Spec, int32 ExecutionIdx, UAbilitySystemComponent* TargetAbilitySystemComponent) const;
-
-	virtual void GetCaptureDefinitions(OUT TArray<FGameplayEffectAttributeCaptureDefinition>& Definitions) const;
-};
-
-// yuck?
-
-#define GE_REQ_SRC(SET, PROP, SNAPSHOT) \
-{ \
-	static UProperty* ThisProperty = FindFieldChecked<UProperty>(SET::StaticClass(), GET_MEMBER_NAME_CHECKED(SET, PROP)); \
-	Definitions.Add(FGameplayEffectAttributeCaptureDefinition(FGameplayAttribute(ThisProperty), EGameplayEffectAttributeCaptureSource::Source, SNAPSHOT)); \
-}
-
-#define GE_REQ_TARGET(SET, PROP, SNAPSHOT) \
-{ \
-	static UProperty* ThisProperty = FindFieldChecked<UProperty>(SET::StaticClass(), GET_MEMBER_NAME_CHECKED(SET, PROP)); \
-	Definitions.Add(FGameplayEffectAttributeCaptureDefinition(FGameplayAttribute(ThisProperty), EGameplayEffectAttributeCaptureSource::Target, SNAPSHOT)); \
-}
-
-
-#endif
\ No newline at end of file
