Gameplay Effect inital refactor
-Remove concept of outgoing/incoming/active GE modifiers
-Initial work for custom calculations and scoped attribute modifiers

diff --git a/Engine/Source/Editor/GameplayAbilitiesEditor/Private/AttributeDetails.cpp b/Engine/Source/Editor/GameplayAbilitiesEditor/Private/AttributeDetails.cpp
--- a/Engine/Source/Editor/GameplayAbilitiesEditor/Private/AttributeDetails.cpp
+++ b/Engine/Source/Editor/GameplayAbilitiesEditor/Private/AttributeDetails.cpp
@@ -32,6 +32,9 @@ void FAttributePropertyDetails::CustomizeHeader( TSharedRef<IPropertyHandle> Str
 
 	TArray<UProperty*> PropertiesToAdd;
 
+
+	
+/*
 	const bool bIsExtensionAttribute = (StructPropertyHandle->GetProperty()->GetOuter() == FExtensionAttributeModifierInfo::StaticStruct());
 	if (bIsExtensionAttribute)
 	{
@@ -86,6 +89,7 @@ void FAttributePropertyDetails::CustomizeHeader( TSharedRef<IPropertyHandle> Str
 		}
 	}
 	else
+*/
 	{
 		// Gather all UAttriubute classes
 		for (TObjectIterator<UClass> ClassIt; ClassIt; ++ClassIt)
diff --git a/Engine/Source/Editor/GameplayAbilitiesEditor/Private/GameplayModifierInfoDetails.cpp b/Engine/Source/Editor/GameplayAbilitiesEditor/Private/GameplayModifierInfoDetails.cpp
--- a/Engine/Source/Editor/GameplayAbilitiesEditor/Private/GameplayModifierInfoDetails.cpp
+++ b/Engine/Source/Editor/GameplayAbilitiesEditor/Private/GameplayModifierInfoDetails.cpp
@@ -40,16 +40,6 @@ void FGameplayModifierInfoCustomization::CustomizeChildren(TSharedRef<IPropertyH
 		{
 			const FName PropName = ChildHandle->GetProperty()->GetFName();
 
-			if ( PropName == GET_MEMBER_NAME_CHECKED(FGameplayModifierInfo, ModifierType) )
-			{
-				ChildHandle->SetOnPropertyValueChanged( FSimpleDelegate::CreateSP(this, &FGameplayModifierInfoCustomization::OnModifierTypeChanged, StructPropertyHandle) );
-			}
-
-			if ( PropName == GET_MEMBER_NAME_CHECKED(FGameplayModifierInfo, EffectType) )
-			{
-				ChildHandle->SetOnPropertyValueChanged( FSimpleDelegate::CreateSP(this, &FGameplayModifierInfoCustomization::OnEffectTypeChanged, StructPropertyHandle) );
-			}
-
 			if ( PropName == GET_MEMBER_NAME_CHECKED(FGameplayModifierInfo, ModifierOp) )
 			{
 				ChildHandle->SetOnPropertyValueChanged( FSimpleDelegate::CreateSP(this, &FGameplayModifierInfoCustomization::OnModifierOpChanged, StructPropertyHandle) );
@@ -61,16 +51,6 @@ void FGameplayModifierInfoCustomization::CustomizeChildren(TSharedRef<IPropertyH
 	}
 }
 
-void FGameplayModifierInfoCustomization::OnModifierTypeChanged(TSharedRef<IPropertyHandle> StructPropertyHandle)
-{
-	UpdateHiddenProperties(StructPropertyHandle);
-}
-
-void FGameplayModifierInfoCustomization::OnEffectTypeChanged(TSharedRef<IPropertyHandle> StructPropertyHandle)
-{
-	UpdateHiddenProperties(StructPropertyHandle);
-}
-
 void FGameplayModifierInfoCustomization::OnModifierOpChanged(TSharedRef<IPropertyHandle> StructPropertyHandle)
 {
 	// When changing to a callback op, automatically add an initial entry in the callbacks list.
@@ -113,64 +93,7 @@ EVisibility FGameplayModifierInfoCustomization::GetPropertyVisibility(FName Prop
 
 void FGameplayModifierInfoCustomization::UpdateHiddenProperties(const TSharedRef<IPropertyHandle>& StructPropertyHandle)
 {
-	HiddenProperties.Empty();
-
-	TArray<const void*> StructPtrs;
-	StructPropertyHandle->AccessRawData(StructPtrs);
-	if ( StructPtrs.Num() > 0 )
-	{
-		const FGameplayModifierInfo& FirstModifier = *reinterpret_cast<const FGameplayModifierInfo*>(StructPtrs[0]);
-		EGameplayMod::Type CommonModType = FirstModifier.ModifierType;
-		EGameplayModEffect::Type CommonModEffect = FirstModifier.EffectType;
-		EGameplayModOp::Type CommonModOp = FirstModifier.ModifierOp;
-		bool bAllSameModType = true;
-		bool bAllSameModEffect = true;
-		bool bAllSameModOp = true;
-
-		for ( const void* Ptr : StructPtrs )
-		{
-			const FGameplayModifierInfo& Modifier = *reinterpret_cast<const FGameplayModifierInfo*>(Ptr);
-
-			if ( Modifier.ModifierType != CommonModType )
-			{
-				bAllSameModType = false;
-			}
 
-			if ( Modifier.EffectType != CommonModEffect )
-			{
-				bAllSameModEffect = false;
-			}
-
-			if ( Modifier.ModifierOp != CommonModOp )
-			{
-				bAllSameModOp = false;
-			}
-		}
-		
-		if ( bAllSameModType && CommonModType == EGameplayMod::Attribute )
-		{
-			HiddenProperties.Add(GET_MEMBER_NAME_CHECKED(FGameplayModifierInfo, EffectType));
-			HiddenProperties.Add(GET_MEMBER_NAME_CHECKED(FGameplayModifierInfo, TargetEffect));
-		}
-
-		if ( bAllSameModEffect && CommonModEffect != EGameplayModEffect::LinkedGameplayEffect )
-		{
-			HiddenProperties.Add(GET_MEMBER_NAME_CHECKED(FGameplayModifierInfo, TargetEffect));
-		}
-
-		if ( bAllSameModOp )
-		{
-			if ( CommonModOp == EGameplayModOp::Callback )
-			{
-				HiddenProperties.Add(GET_MEMBER_NAME_CHECKED(FGameplayModifierInfo, Magnitude));
-				HiddenProperties.Add(GET_MEMBER_NAME_CHECKED(FGameplayModifierInfo, LevelInfo));
-			}
-			else
-			{
-				HiddenProperties.Add(GET_MEMBER_NAME_CHECKED(FGameplayModifierInfo, Callbacks));
-			}
-		}
-	}
 }
 
 
diff --git a/Engine/Source/Editor/GameplayAbilitiesEditor/Private/GameplayModifierInfoDetails.h b/Engine/Source/Editor/GameplayAbilitiesEditor/Private/GameplayModifierInfoDetails.h
--- a/Engine/Source/Editor/GameplayAbilitiesEditor/Private/GameplayModifierInfoDetails.h
+++ b/Engine/Source/Editor/GameplayAbilitiesEditor/Private/GameplayModifierInfoDetails.h
@@ -19,8 +19,7 @@ public:
 	virtual void CustomizeChildren(TSharedRef<IPropertyHandle> StructPropertyHandle, IDetailChildrenBuilder& StructBuilder, IPropertyTypeCustomizationUtils& StructCustomizationUtils) override;
 
 private:
-	void OnModifierTypeChanged(TSharedRef<IPropertyHandle> StructPropertyHandle);
-	void OnEffectTypeChanged(TSharedRef<IPropertyHandle> StructPropertyHandle);
+	
 	void OnModifierOpChanged(TSharedRef<IPropertyHandle> StructPropertyHandle);
 
 	EVisibility GetPropertyVisibility(FName PropName) const;
diff --git a/Engine/Source/Editor/GameplayAbilitiesEditor/Private/SGraphNodeK2GameplayEffectVar.cpp b/Engine/Source/Editor/GameplayAbilitiesEditor/Private/SGraphNodeK2GameplayEffectVar.cpp
--- a/Engine/Source/Editor/GameplayAbilitiesEditor/Private/SGraphNodeK2GameplayEffectVar.cpp
+++ b/Engine/Source/Editor/GameplayAbilitiesEditor/Private/SGraphNodeK2GameplayEffectVar.cpp
@@ -129,11 +129,6 @@ void SGraphNodeK2GameplayEffectVar::UpdateGraphNode()
 	const int32 ModifierCount = GameplayEffect ? GameplayEffect->Modifiers.Num() : 0;
 	for (int32 Idx = 0; Idx < ModifierCount; ++Idx)
 	{
-		if (GameplayEffect->Modifiers[Idx].ModifierType != EGameplayMod::Attribute)
-		{
-			continue;
-		}
-
 		RightNodeBox->AddSlot()
 			.FillHeight(0.3f)
 			[
diff --git a/Engine/Source/Runtime/GameplayAbilities/Private/Abilities/GameplayAbility.cpp b/Engine/Source/Runtime/GameplayAbilities/Private/Abilities/GameplayAbility.cpp
--- a/Engine/Source/Runtime/GameplayAbilities/Private/Abilities/GameplayAbility.cpp
+++ b/Engine/Source/Runtime/GameplayAbilities/Private/Abilities/GameplayAbility.cpp
@@ -740,7 +740,7 @@ FActiveGameplayEffectHandle UGameplayAbility::ApplyGameplayEffectToOwner(const F
 {
 	if (ActivationInfo.ActivationMode == EGameplayAbilityActivationMode::Authority || ActivationInfo.ActivationMode == EGameplayAbilityActivationMode::Predicting)
 	{
-		return ActorInfo->AbilitySystemComponent->ApplyGameplayEffectToSelf(GameplayEffect, 1.f, GetEffectContext(ActorInfo), FModifierQualifier().PredictionKey(ActivationInfo.GetPredictionKeyForNewAction()));
+		return ActorInfo->AbilitySystemComponent->ApplyGameplayEffectToSelf(GameplayEffect, 1.f, GetEffectContext(ActorInfo), ActivationInfo.GetPredictionKeyForNewAction());
 	}
 
 	// We cannot apply GameplayEffects in this context. Return an empty handle.
@@ -771,7 +771,7 @@ FActiveGameplayEffectHandle UGameplayAbility::ApplyGameplayEffectToTarget(const
 	{
 		for (auto Data : Target.Data)
 		{
-			TArray<FActiveGameplayEffectHandle> EffectHandles = Data->ApplyGameplayEffect(GameplayEffect, GetEffectContext(ActorInfo), (float)GameplayEffectLevel, FModifierQualifier().PredictionKey(ActivationInfo.GetPredictionKeyForNewAction()));
+			TArray<FActiveGameplayEffectHandle> EffectHandles = Data->ApplyGameplayEffect(GameplayEffect, GetEffectContext(ActorInfo), (float)GameplayEffectLevel, ActivationInfo.GetPredictionKeyForNewAction());
 			if (EffectHandles.Num() > 0)
 			{
 				EffectHandle = EffectHandles[0];
diff --git a/Engine/Source/Runtime/GameplayAbilities/Private/Abilities/Tasks/AbilityTask_WaitAttributeChange.cpp b/Engine/Source/Runtime/GameplayAbilities/Private/Abilities/Tasks/AbilityTask_WaitAttributeChange.cpp
--- a/Engine/Source/Runtime/GameplayAbilities/Private/Abilities/Tasks/AbilityTask_WaitAttributeChange.cpp
+++ b/Engine/Source/Runtime/GameplayAbilities/Private/Abilities/Tasks/AbilityTask_WaitAttributeChange.cpp
@@ -43,8 +43,8 @@ void UAbilityTask_WaitAttributeChange::OnAttributeChange(float NewValue, const F
 	}
 	else
 	{
-		if ((WithTag.IsValid() && !Data->EvaluatedData.Tags.HasTag(WithTag, EGameplayTagMatchType::IncludeParentTags, EGameplayTagMatchType::Explicit)) ||
-			(WithoutTag.IsValid() && Data->EvaluatedData.Tags.HasTag(WithoutTag, EGameplayTagMatchType::IncludeParentTags, EGameplayTagMatchType::Explicit)))
+		if ((WithTag.IsValid() && !Data->EvaluatedData.SourceTags.HasTag(WithTag, EGameplayTagMatchType::IncludeParentTags, EGameplayTagMatchType::Explicit)) ||
+			(WithoutTag.IsValid() && Data->EvaluatedData.SourceTags.HasTag(WithoutTag, EGameplayTagMatchType::IncludeParentTags, EGameplayTagMatchType::Explicit)))
 		{
 			// Failed tag check
 			return;
diff --git a/Engine/Source/Runtime/GameplayAbilities/Private/AbilitySystemComponent.cpp b/Engine/Source/Runtime/GameplayAbilities/Private/AbilitySystemComponent.cpp
--- a/Engine/Source/Runtime/GameplayAbilities/Private/AbilitySystemComponent.cpp
+++ b/Engine/Source/Runtime/GameplayAbilities/Private/AbilitySystemComponent.cpp
@@ -1,5 +1,4 @@
-// Copyright 1998-2013 Epic Games, Inc. All Rights Reserved.
-// ActorComponent.cpp: Actor component implementation.
+// Copyright 1998-2014 Epic Games, Inc. All Rights Reserved.
 
 #include "AbilitySystemPrivatePCH.h"
 #include "AbilitySystemComponent.h"
@@ -14,6 +13,8 @@
 #include "UObjectToken.h"
 #include "MapErrors.h"
 
+#pragma optimize( "", off )
+
 DEFINE_LOG_CATEGORY(LogAbilitySystemComponent);
 
 #define LOCTEXT_NAMESPACE "AbilitySystemComponent"
@@ -40,7 +41,6 @@ UAbilitySystemComponent::UAbilitySystemComponent(const FObjectInitializer& Objec
 	PrimaryComponentTick.bStartWithTickEnabled = true; // FIXME! Just temp until timer manager figured out
 	PrimaryComponentTick.bCanEverTick = true;
 	
-	ActiveGameplayEffects.Owner = this;
 	ActiveGameplayCues.Owner = this;
 
 	bReplicates = true;
@@ -124,30 +124,8 @@ void UAbilitySystemComponent::OnRegister()
 			Attributes->InitFromMetaDataTable(DefaultStartingData[i].DefaultStartingTable);
 		}
 	}
-}
-
-// ---------------------------------------------------------
-
-bool UAbilitySystemComponent::AreGameplayEffectApplicationRequirementsSatisfied(const class UGameplayEffect* EffectToAdd, const FGameplayEffectContextHandle& EffectContext) const
-{
-	bool bReqsSatisfied = false;
-	if (EffectToAdd)
-	{
-		// Collect gameplay tags from instigator and target to see if requirements are satisfied
-		FGameplayTagContainer InstigatorTags;
-		EffectContext.GetOwnedGameplayTags(InstigatorTags);
-
-		FGameplayTagContainer TargetTags;
-		IGameplayTagAssetInterface* OwnerGTA = Cast<IGameplayTagAssetInterface>(AbilityActorInfo->OwnerActor.Get());
-		if (OwnerGTA)
-		{
-			OwnerGTA->GetOwnedGameplayTags(TargetTags);
-		}
-
-		bReqsSatisfied = EffectToAdd->AreApplicationTagRequirementsSatisfied(InstigatorTags, TargetTags);
-	}
 
-	return bReqsSatisfied;
+	ActiveGameplayEffects.RegisterWithOwner(this);
 }
 
 // ---------------------------------------------------------
@@ -157,9 +135,9 @@ bool UAbilitySystemComponent::IsOwnerActorAuthoritative() const
 	return !IsNetSimulating();
 }
 
-bool UAbilitySystemComponent::HasNetworkAuthorityToApplyGameplayEffect(const FModifierQualifier QualifierContext) const
+bool UAbilitySystemComponent::HasNetworkAuthorityToApplyGameplayEffect(FPredictionKey PredictionKey) const
 {
-	return (IsOwnerActorAuthoritative() || QualifierContext.PredictionKey().IsValidForMorePrediction());
+	return (IsOwnerActorAuthoritative() || PredictionKey.IsValidForMorePrediction());
 }
 
 void UAbilitySystemComponent::SetNumericAttribute(const FGameplayAttribute &Attribute, float NewFloatValue)
@@ -180,14 +158,8 @@ FGameplayEffectSpecHandle UAbilitySystemComponent::GetOutgoingSpec(UGameplayEffe
 	// Fixme: we should build a map and cache these off. We can invalidate the map when an OutgoingGE modifier is applied or removed from us.
 
 	// By default use the owner and avatar as the instigator and causer
-	FGameplayEffectSpec* NewSpec = new FGameplayEffectSpec(GameplayEffect, GetEffectContext(), Level, GetCurveDataOverride());
-	if (ActiveGameplayEffects.ApplyActiveEffectsTo(*NewSpec, FModifierQualifier().Type(EGameplayMod::OutgoingGE)))
-	{
-		return FGameplayEffectSpecHandle(NewSpec);
-	}
-
-	delete NewSpec;
-	return FGameplayEffectSpecHandle(nullptr);
+	FGameplayEffectSpec* NewSpec = new FGameplayEffectSpec(GameplayEffect, GetEffectContext(), Level);
+	return FGameplayEffectSpecHandle(NewSpec);
 }
 
 FGameplayEffectContextHandle UAbilitySystemComponent::GetEffectContext() const
@@ -199,18 +171,18 @@ FGameplayEffectContextHandle UAbilitySystemComponent::GetEffectContext() const
 }
 
 /** This is a helper function used in automated testing, I'm not sure how useful it will be to gamecode or blueprints */
-FActiveGameplayEffectHandle UAbilitySystemComponent::ApplyGameplayEffectToTarget(UGameplayEffect *GameplayEffect, UAbilitySystemComponent *Target, float Level, FGameplayEffectContextHandle Context, FModifierQualifier BaseQualifier)
+FActiveGameplayEffectHandle UAbilitySystemComponent::ApplyGameplayEffectToTarget(UGameplayEffect *GameplayEffect, UAbilitySystemComponent *Target, float Level, FGameplayEffectContextHandle Context, FPredictionKey PredictionKey)
 {
 	check(GameplayEffect);
-	if (HasNetworkAuthorityToApplyGameplayEffect(BaseQualifier))
+	if (HasNetworkAuthorityToApplyGameplayEffect(PredictionKey))
 	{
 		if (!Context.IsValid())
 		{
 			Context = GetEffectContext();
 		}
 
-		FGameplayEffectSpec	Spec(GameplayEffect, Context, Level, GetCurveDataOverride());
-		return ApplyGameplayEffectSpecToTarget(Spec, Target, BaseQualifier);
+		FGameplayEffectSpec	Spec(GameplayEffect, Context, Level);
+		return ApplyGameplayEffectSpecToTarget(Spec, Target, PredictionKey);
 	}
 
 	return FActiveGameplayEffectHandle();
@@ -223,7 +195,7 @@ FActiveGameplayEffectHandle UAbilitySystemComponent::K2_ApplyGameplayEffectToSel
 }
 
 /** This is a helper function - it seems like this will be useful as a blueprint interface at the least, but Level parameter may need to be expanded */
-FActiveGameplayEffectHandle UAbilitySystemComponent::ApplyGameplayEffectToSelf(const UGameplayEffect *GameplayEffect, float Level, const FGameplayEffectContextHandle& EffectContext, FModifierQualifier BaseQualifier)
+FActiveGameplayEffectHandle UAbilitySystemComponent::ApplyGameplayEffectToSelf(const UGameplayEffect *GameplayEffect, float Level, const FGameplayEffectContextHandle& EffectContext, FPredictionKey PredictionKey)
 {
 	if (GameplayEffect == nullptr)
 	{
@@ -231,20 +203,15 @@ FActiveGameplayEffectHandle UAbilitySystemComponent::ApplyGameplayEffectToSelf(c
 		return FActiveGameplayEffectHandle();
 	}
 
-	if (HasNetworkAuthorityToApplyGameplayEffect(BaseQualifier))
+	if (HasNetworkAuthorityToApplyGameplayEffect(PredictionKey))
 	{
-		FGameplayEffectSpec	Spec(GameplayEffect, EffectContext, Level, GetCurveDataOverride());
-		return ApplyGameplayEffectSpecToSelf(Spec, BaseQualifier);
+		FGameplayEffectSpec	Spec(GameplayEffect, EffectContext, Level);
+		return ApplyGameplayEffectSpecToSelf(Spec, PredictionKey);
 	}
 
 	return FActiveGameplayEffectHandle();
 }
 
-float UAbilitySystemComponent::GetGameplayEffectMagnitudeByTag(FActiveGameplayEffectHandle InHandle, const FGameplayTag& InTag) const
-{
-	return ActiveGameplayEffects.GetGameplayEffectMagnitudeByTag(InHandle, InTag);
-}
-
 FOnActiveGameplayEffectRemoved* UAbilitySystemComponent::OnGameplayEffectRemovedDelegate(FActiveGameplayEffectHandle Handle)
 {
 	FActiveGameplayEffect* ActiveEffect = ActiveGameplayEffects.GetActiveGameplayEffect(Handle);
@@ -266,11 +233,26 @@ bool UAbilitySystemComponent::IsGameplayEffectActive(FActiveGameplayEffectHandle
 	return ActiveGameplayEffects.IsGameplayEffectActive(InHandle);
 }
 
+void UAbilitySystemComponent::CaptureAttributeForGameplayEffect(OUT FGameplayEffectAttributeCaptureSpec& OutCaptureSpec)
+{
+	// Verify the capture is happening on an attribute the component actually has a set for; if not, can't capture the value
+	const FGameplayAttribute& AttributeToCapture = OutCaptureSpec.BackingDefinition.AttributeToCapture;
+	if (AttributeToCapture.IsValid() && GetAttributeSubobject(AttributeToCapture.GetAttributeSetClass()))
+	{
+		ActiveGameplayEffects.CaptureAttributeForGameplayEffect(OutCaptureSpec);
+	}
+}
+
 FOnGameplayEffectTagCountChanged& UAbilitySystemComponent::RegisterGameplayTagEvent(FGameplayTag Tag)
 {
 	return GameplayTagCountContainer.GameplayTagEventMap.FindOrAdd(Tag);
 }
 
+FOnGameplayEffectTagCountChanged& UAbilitySystemComponent::RegisterGenericGameplayTagEvent()
+{
+	return GameplayTagCountContainer.OnAnyTagChangeDelegate;
+}
+
 FOnGameplayAttributeChange& UAbilitySystemComponent::RegisterGameplayAttributeEvent(FGameplayAttribute Attribute)
 {
 	return ActiveGameplayEffects.RegisterGameplayAttributeEvent(Attribute);
@@ -307,59 +289,62 @@ bool UAbilitySystemComponent::HasAnyMatchingGameplayTags(const FGameplayTagConta
 	return GameplayTagCountContainer.HasAnyMatchingGameplayTags(TagContainer, EGameplayTagMatchType::Explicit, bCountEmptyAsMatch);
 }
 
-void UAbilitySystemComponent::UpdateTagMap(const FGameplayTag& BaseTag, int32 CountDelta)
+void UAbilitySystemComponent::AddLooseGameplayTag(const FGameplayTag& GameplayTag, int32 Count)
 {
-	GameplayTagCountContainer.UpdateTagMap(BaseTag, CountDelta);
+	UpdateTagMap(GameplayTag, Count);
 }
 
-void UAbilitySystemComponent::UpdateTagMap(const FGameplayTagContainer& Container, int32 CountDelta)
+void UAbilitySystemComponent::AddLooseGameplayTags(const FGameplayTagContainer& GameplayTags, int32 Count)
 {
-	GameplayTagCountContainer.UpdateTagMap(Container, CountDelta);
+	UpdateTagMap(GameplayTags, Count);
 }
 
-// ------------------------------------------------------------------------
+void UAbilitySystemComponent::RemoveLooseGameplayTag(const FGameplayTag& GameplayTag, int32 Count)
+{
+	UpdateTagMap(GameplayTag, Count);
+}
 
-void UAbilitySystemComponent::TEMP_ApplyActiveGameplayEffects()
+void UAbilitySystemComponent::RemoveLooseGameplayTags(const FGameplayTagContainer& GameplayTags, int32 Count)
 {
-	for (int32 idx=0; idx < ActiveGameplayEffects.GameplayEffects.Num(); ++idx)
-	{
-		FActiveGameplayEffect& ActiveEffect = ActiveGameplayEffects.GameplayEffects[idx];
+	UpdateTagMap(GameplayTags, Count);
+}
 
-		ExecuteGameplayEffect(ActiveEffect.Spec, FModifierQualifier().IgnoreHandle(ActiveEffect.Handle));
+// These are functionally redundant but are called by GEs and GameplayCues that add tags that are not 'loose' (but are handled the same way in practice)
 
-		ABILITY_LOG(Log, TEXT("ActiveEffect[%d] %s - Duration: %.2f]"), idx, *ActiveEffect.Spec.ToSimpleString(), ActiveEffect.Spec.GetDuration());
-	}
+void UAbilitySystemComponent::UpdateTagMap(const FGameplayTag& BaseTag, int32 CountDelta)
+{
+	GameplayTagCountContainer.UpdateTagMap(BaseTag, CountDelta);
 }
 
-FActiveGameplayEffectHandle UAbilitySystemComponent::ApplyGameplayEffectSpecToTarget(OUT FGameplayEffectSpec &Spec, UAbilitySystemComponent *Target, FModifierQualifier BaseQualifier)
+void UAbilitySystemComponent::UpdateTagMap(const FGameplayTagContainer& Container, int32 CountDelta)
 {
-	if (HasNetworkAuthorityToApplyGameplayEffect(BaseQualifier))
+	GameplayTagCountContainer.UpdateTagMap(Container, CountDelta);
+}
+
+// ------------------------------------------------------------------------
+
+FActiveGameplayEffectHandle UAbilitySystemComponent::ApplyGameplayEffectSpecToTarget(OUT FGameplayEffectSpec &Spec, UAbilitySystemComponent *Target, FPredictionKey PredictionKey)
+{
+	if (HasNetworkAuthorityToApplyGameplayEffect(PredictionKey))
 	{
 		// Apply outgoing Effects to the Spec.
 		// Outgoing immunity may stop the outgoing effect from being applied to the target
-		if (ActiveGameplayEffects.ApplyActiveEffectsTo(Spec, FModifierQualifier(BaseQualifier).Type(EGameplayMod::OutgoingGE)))
-		{
-			return Target->ApplyGameplayEffectSpecToSelf(Spec, BaseQualifier);
-		}
+		return Target->ApplyGameplayEffectSpecToSelf(Spec, PredictionKey);
 	}
 
 	return FActiveGameplayEffectHandle();
 }
 
-FActiveGameplayEffectHandle UAbilitySystemComponent::ApplyGameplayEffectSpecToSelf(OUT FGameplayEffectSpec &Spec, FModifierQualifier BaseQualifier)
+FActiveGameplayEffectHandle UAbilitySystemComponent::ApplyGameplayEffectSpecToSelf(OUT FGameplayEffectSpec &Spec, FPredictionKey PredictionKey)
 {
 	// Temp, only non instant, non periodic GEs can be predictive
 	// Effects with other effects may be a mix so go with non-predictive
-	check((BaseQualifier.PredictionKey().IsValidKey() == false) || (Spec.GetPeriod() == UGameplayEffect::NO_PERIOD));
+	check((PredictionKey.IsValidKey() == false) || (Spec.GetPeriod() == UGameplayEffect::NO_PERIOD));
 
-	if (!HasNetworkAuthorityToApplyGameplayEffect(BaseQualifier))
+	if (!HasNetworkAuthorityToApplyGameplayEffect(PredictionKey))
 	{
 		return FActiveGameplayEffectHandle();
 	}
-	
-
-	// Clients should treat predicted instant effects as if they have infinite duration. The effects will be cleaned up later.
-	bool bTreatAsInfiniteDuration = GetOwnerRole() != ROLE_Authority && BaseQualifier.PredictionKey().IsValidKey() && Spec.GetDuration() == UGameplayEffect::INSTANT_APPLICATION;
 
 	// check if the effect being applied actually succeeds
 	float ChanceToApply = Spec.GetChanceToApplyToTarget();
@@ -368,7 +353,24 @@ FActiveGameplayEffectHandle UAbilitySystemComponent::ApplyGameplayEffectSpecToSe
 		return FActiveGameplayEffectHandle();
 	}
 
-	// Make sure we create our copy of the spec in the right place first...
+	// Get MyTags.
+	//	We may want to cache off a GameplayTagContainer instead of rebuilding it everytime.
+	//	But this will also be where we need to merge in context tags? (Headshot, executing ability, etc?)
+	//	Or do we push these tags into (our copy of the spec)?
+
+	FGameplayTagContainer MyTags;
+	GetOwnedGameplayTags(MyTags);
+
+	if (Spec.Def->ApplicationTagRequirements.RequirementsMet(MyTags) == false)
+	{
+		return FActiveGameplayEffectHandle();
+	}
+	
+
+	// Clients should treat predicted instant effects as if they have infinite duration. The effects will be cleaned up later.
+	bool bTreatAsInfiniteDuration = GetOwnerRole() != ROLE_Authority && PredictionKey.IsValidKey() && Spec.GetDuration() == UGameplayEffect::INSTANT_APPLICATION;
+
+	// Make sure we create our copy of the spec in the right place
 	FActiveGameplayEffectHandle	MyHandle;
 	bool bInvokeGameplayCueApplied = UGameplayEffect::INSTANT_APPLICATION != Spec.GetDuration(); // Cache this now before possibly modifying predictive instant effect to infinite duration effect.
 
@@ -383,7 +385,7 @@ FActiveGameplayEffectHandle UAbilitySystemComponent::ApplyGameplayEffectSpecToSe
 			{
 				ActiveGameplayEffects.StacksNeedToRecalculate();
 			}
-			FActiveGameplayEffect &NewActiveEffect = ActiveGameplayEffects.CreateNewActiveGameplayEffect(Spec, BaseQualifier.PredictionKey());
+			FActiveGameplayEffect& NewActiveEffect = ActiveGameplayEffects.CreateNewActiveGameplayEffect(Spec, PredictionKey);
 			MyHandle = NewActiveEffect.Handle;
 			OurCopyOfSpec = &NewActiveEffect.Spec;
 		}
@@ -394,6 +396,8 @@ FActiveGameplayEffectHandle UAbilitySystemComponent::ApplyGameplayEffectSpecToSe
 			OurCopyOfSpec = StackSpec.Get();
 		}
 
+		// GE_REMOVE: Pretty sure this completely goes away? 
+
 		// Do a 1st order copy of the spec so that we can modify it
 		// (the one passed in is owned by the caller, we can't apply our incoming GEs to it)
 		// Note that at this point the spec has a bunch of modifiers. Those modifiers may
@@ -401,40 +405,19 @@ FActiveGameplayEffectHandle UAbilitySystemComponent::ApplyGameplayEffectSpecToSe
 		//
 		// In theory, we don't modify 2nd order modifiers after they are 'attached'
 		// Long complex chains can be created but we never say 'Modify a GE that is modding another GE'
-		OurCopyOfSpec->MakeUnique();
+		// 
+		// 
+		// OurCopyOfSpec->MakeUnique();
 
 		// if necessary add a modifier to OurCopyOfSpec to force it to have an infinite duration
 		if (bTreatAsInfiniteDuration)
-		{
-			FGameplayModifierInfo ModInfo;
-			ModInfo.ModifierOp = EGameplayModOp::Override;
-			ModInfo.Magnitude.SetValue(UGameplayEffect::INFINITE_DURATION);
-			ModInfo.EffectType = EGameplayModEffect::Duration;
-			ModInfo.ModifierType = EGameplayMod::ActiveGE;
-			TSharedPtr<FGameplayEffectLevelSpec> Level(new FGameplayEffectLevelSpec(0.f, OurCopyOfSpec->Def->LevelInfo, NULL));
-			FModifierSpec Mod(ModInfo, Level, NULL);
-			FModifierQualifier Qualifier;
-			Qualifier.Type(EGameplayMod::ActiveGE);
-
-			OurCopyOfSpec->ApplyModifier(Mod, Qualifier, true);
+		{			
+			// This should just be a straight set of the duration float now
 		}
 	}
 
-	// Now that we have our own copy, apply our GEs that modify IncomingGEs
-	if (!ActiveGameplayEffects.ApplyActiveEffectsTo(*OurCopyOfSpec, FModifierQualifier(BaseQualifier).Type(EGameplayMod::IncomingGE).IgnoreHandle(MyHandle)))
-	{
-		// We're immune to this effect
-		return FActiveGameplayEffectHandle();
-	}	
-	
-	// Now that we have the final version of this effect, actually apply it if its going to be hanging around
-	if (Duration != UGameplayEffect::INSTANT_APPLICATION )
-	{
-		if (Spec.GetPeriod() == UGameplayEffect::NO_PERIOD)
-		{
-			ActiveGameplayEffects.ApplySpecToActiveEffectsAndAttributes(*OurCopyOfSpec, FModifierQualifier(BaseQualifier).IgnoreHandle(MyHandle));
-		}
-	}
+	// Capture Our attributes (this may snapshot or link to our AttributeAggregators)
+	OurCopyOfSpec->CapturedRelevantAttributes.CaptureAttributes(this, EGameplayEffectAttributeCaptureSource::Target);
 	
 	// We still probably want to apply tags and stuff even if instant?
 	if (bInvokeGameplayCueApplied)
@@ -453,7 +436,7 @@ FActiveGameplayEffectHandle UAbilitySystemComponent::ApplyGameplayEffectSpecToSe
 	// Execute if this is an instant application effect
 	if (Duration == UGameplayEffect::INSTANT_APPLICATION)
 	{
-		ExecuteGameplayEffect(*OurCopyOfSpec, FModifierQualifier(BaseQualifier).IgnoreHandle(MyHandle));
+		ExecuteGameplayEffect(*OurCopyOfSpec, PredictionKey);
 	}
 	else if (bTreatAsInfiniteDuration)
 	{
@@ -470,7 +453,7 @@ FActiveGameplayEffectHandle UAbilitySystemComponent::ApplyGameplayEffectSpecToSe
 	// todo: this is ignoring the returned handles, should we put them into a TArray and return all of the handles?
 	for (const TSharedRef<FGameplayEffectSpec> TargetSpec : Spec.TargetEffectSpecs)
 	{
-		ApplyGameplayEffectSpecToSelf(TargetSpec.Get(), BaseQualifier);
+		ApplyGameplayEffectSpecToSelf(TargetSpec.Get(), PredictionKey);
 	}
 
 	return MyHandle;
@@ -481,13 +464,13 @@ void UAbilitySystemComponent::ExecutePeriodicEffect(FActiveGameplayEffectHandle
 	ActiveGameplayEffects.ExecutePeriodicGameplayEffect(Handle);
 }
 
-void UAbilitySystemComponent::ExecuteGameplayEffect(FGameplayEffectSpec &Spec, const FModifierQualifier &QualifierContext)
+void UAbilitySystemComponent::ExecuteGameplayEffect(FGameplayEffectSpec &Spec, FPredictionKey PredictionKey)
 {
 	// Should only ever execute effects that are instant application or periodic application
 	// Effects with no period and that aren't instant application should never be executed
 	check( (Spec.GetDuration() == UGameplayEffect::INSTANT_APPLICATION || Spec.GetPeriod() != UGameplayEffect::NO_PERIOD) );
 	
-	ActiveGameplayEffects.ExecuteActiveEffectsFrom(Spec, QualifierContext);
+	ActiveGameplayEffects.ExecuteActiveEffectsFrom(Spec, PredictionKey);
 }
 
 void UAbilitySystemComponent::CheckDurationExpired(FActiveGameplayEffectHandle Handle)
@@ -523,9 +506,8 @@ void UAbilitySystemComponent::InvokeGameplayCueEvent(const FGameplayEffectSpec &
 	{
 		ABILITY_LOG(Warning, TEXT("InvokeGameplayCueEvent: %s"), *Spec.ToSimpleString());
 	}
-
-	// FIXME: Replication of level not finished
-	float ExecuteLevel =  (Spec.ModifierLevel.IsValid() && Spec.ModifierLevel.Get()->IsValid()) ? Spec.ModifierLevel.Get()->GetLevel() : 1.f;
+	
+	float ExecuteLevel = Spec.GetLevel();
 
 	FGameplayCueParameters CueParameters;
 	CueParameters.EffectContext = Spec.EffectContext;
@@ -666,11 +648,6 @@ bool UAbilitySystemComponent::IsGameplayCueActive(const FGameplayTag GameplayCue
 
 // ----------------------------------------------------------------------------------------
 
-void UAbilitySystemComponent::AddDependancyToAttribute(FGameplayAttribute Attribute, const TWeakPtr<FAggregator> InDependant)
-{
-	ActiveGameplayEffects.AddDependancyToAttribute(Attribute, InDependant);
-}
-
 void UAbilitySystemComponent::SetBaseAttributeValueFromReplication(float NewValue, FGameplayAttribute Attribute)
 {
 	ActiveGameplayEffects.SetBaseAttributeValueFromReplication(Attribute, NewValue);
@@ -828,133 +805,11 @@ void UAbilitySystemComponent::PrintAllGameplayEffects() const
 	ActiveGameplayEffects.PrintAllGameplayEffects();
 }
 
-void FActiveGameplayEffectsContainer::PrintAllGameplayEffects() const
-{
-	ABILITY_LOG_SCOPE(TEXT("ActiveGameplayEffects. Num: %d"), GameplayEffects.Num());
-	for (const FActiveGameplayEffect& Effect : GameplayEffects)
-	{
-		Effect.PrintAll();
-	}
-}
-
-void FActiveGameplayEffect::PrintAll() const
-{
-	ABILITY_LOG(Log, TEXT("Handle: %s"), *Handle.ToString());
-	ABILITY_LOG(Log, TEXT("StartWorldTime: %.2f"), StartWorldTime);
-	Spec.PrintAll();
-}
-
-void FGameplayEffectSpec::PrintAll() const
-{
-	ABILITY_LOG_SCOPE(TEXT("GameplayEffectSpec"));
-	ABILITY_LOG(Log, TEXT("Def: %s"), *Def->GetName());
-	
-	ABILITY_LOG(Log, TEXT("Duration: "));
-	Duration.PrintAll();
-
-	ABILITY_LOG(Log, TEXT("Period:"));
-	Period.PrintAll();
-
-	ABILITY_LOG(Log, TEXT("Modifiers:"));
-	for (const FModifierSpec &Mod : Modifiers)
-	{
-		Mod.PrintAll();
-	}
-}
-
-void FModifierSpec::PrintAll() const
-{
-	ABILITY_LOG_SCOPE(TEXT("ModifierSpec"));
-	ABILITY_LOG(Log, TEXT("Attribute: %s"), *Info.Attribute.GetName());
-	ABILITY_LOG(Log, TEXT("ModifierType: %s"), *EGameplayModToString(Info.ModifierType));
-	ABILITY_LOG(Log, TEXT("ModifierOp: %s"), *EGameplayModOpToString(Info.ModifierOp));
-	ABILITY_LOG(Log, TEXT("EffectType: %s"), *EGameplayModEffectToString(Info.EffectType));
-	ABILITY_LOG(Log, TEXT("RequiredTags: %s"), *Info.RequiredTags.ToString());
-	ABILITY_LOG(Log, TEXT("OwnedTags: %s"), *Info.OwnedTags.ToString());
-	ABILITY_LOG(Log, TEXT("(Base) Magnitude: %s"), *Info.Magnitude.ToSimpleString());
-
-	Aggregator.PrintAll();
-}
-
-void FAggregatorRef::PrintAll() const
-{
-	if (!WeakPtr.IsValid())
-	{
-		ABILITY_LOG(Log, TEXT("Invalid AggregatorRef"));
-		return;
-	}
-
-	if (SharedPtr.IsValid())
-	{
-		ABILITY_LOG(Log, TEXT("HardRef AggregatorRef"));
-	}
-	else
-	{
-		ABILITY_LOG(Log, TEXT("SoftRef AggregatorRef"));
-
-	}
-	
-	Get()->PrintAll();
-}
-
-void FAggregator::PrintAll() const
-{
-	ABILITY_LOG_SCOPE(TEXT("FAggregator 0x%X"), this);
-
-#if SKILL_SYSTEM_AGGREGATOR_DEBUG
-	ABILITY_LOG(Log, TEXT("DebugStr: %s"), *DebugString);
-	ABILITY_LOG(Log, TEXT("Copies (of me): %d"), CopiesMade);
-#endif
-
-	if (Level.IsValid())
-	{
-		ABILITY_LOG_SCOPE(TEXT("LevelInfo"));
-		Level->PrintAll();
-	}
-	else
-	{
-		ABILITY_LOG(Log, TEXT("No Level Data"));
-	}
-
-	{
-		ABILITY_LOG_SCOPE(TEXT("BaseData"));
-		BaseData.PrintAll();
-	}
-
-	{
-		ABILITY_LOG_SCOPE(TEXT("CachedData"));
-		CachedData.PrintAll();
-	}
-	
-	for (int32 i=0; i < EGameplayModOp::Max; ++i)
-	{
-		if (Mods[i].Num() > 0)
-		{
-			ABILITY_LOG_SCOPE(TEXT("%s Mods"), *EGameplayModOpToString(i));
-			for (const FAggregatorRef &Ref : Mods[i])
-			{
-				Ref.PrintAll();
-			}
-		}
-	}
-}
-
-void FGameplayModifierData::PrintAll() const
-{
-	ABILITY_LOG(Log, TEXT("Magnitude: %s"), *Magnitude.ToSimpleString());
-	ABILITY_LOG(Log, TEXT("Tags: %s"), *Tags.ToString());
-}
-
-void FGameplayModifierEvaluatedData::PrintAll() const
-{
-	ABILITY_LOG(Log, TEXT("IsValid: %d"), IsValid); 
-	ABILITY_LOG(Log, TEXT("Magnitude: %.2f"), Magnitude);
-	ABILITY_LOG(Log, TEXT("Tags: %s"), *Tags.ToString());
-}
+// ------------------------------------------------------------------------
 
-void FGameplayEffectLevelSpec::PrintAll() const
+void UAbilitySystemComponent::OnAttributeAggregatorDirty(FAggregator* Aggregator, FGameplayAttribute Attribute)
 {
-	ABILITY_LOG(Log, TEXT("ConstantLevel: %.2f"), ConstantLevel);
+	ActiveGameplayEffects.OnAttributeAggregatorDirty(Aggregator, Attribute);
 }
 
 #undef LOCTEXT_NAMESPACE
diff --git a/Engine/Source/Runtime/GameplayAbilities/Private/AbilitySystemComponent_Abilities.cpp b/Engine/Source/Runtime/GameplayAbilities/Private/AbilitySystemComponent_Abilities.cpp
--- a/Engine/Source/Runtime/GameplayAbilities/Private/AbilitySystemComponent_Abilities.cpp
+++ b/Engine/Source/Runtime/GameplayAbilities/Private/AbilitySystemComponent_Abilities.cpp
@@ -42,6 +42,19 @@ void UAbilitySystemComponent::InitializeComponent()
 	}
 }
 
+void UAbilitySystemComponent::UninitializeComponent()
+{
+	Super::UninitializeComponent();
+	
+	for (UAttributeSet* Set : SpawnedAttributes)
+	{
+		if (Set)
+		{
+			Set->MarkPendingKill();
+		}
+	}
+}
+
 void UAbilitySystemComponent::TickComponent(float DeltaTime, enum ELevelTick TickType, FActorComponentTickFunction *ThisTickFunction)
 {
 	SCOPE_CYCLE_COUNTER(STAT_TickAbilityTasks);
diff --git a/Engine/Source/Runtime/GameplayAbilities/Private/AbilitySystemTestAttributeSet.cpp b/Engine/Source/Runtime/GameplayAbilities/Private/AbilitySystemTestAttributeSet.cpp
--- a/Engine/Source/Runtime/GameplayAbilities/Private/AbilitySystemTestAttributeSet.cpp
+++ b/Engine/Source/Runtime/GameplayAbilities/Private/AbilitySystemTestAttributeSet.cpp
@@ -25,6 +25,7 @@ UAbilitySystemTestAttributeSet::UAbilitySystemTestAttributeSet(const FObjectInit
 
 bool UAbilitySystemTestAttributeSet::PreGameplayEffectExecute(struct FGameplayEffectModCallbackData &Data)
 {
+#if 0
 	static UProperty *HealthProperty = FindFieldChecked<UProperty>(UAbilitySystemTestAttributeSet::StaticClass(), GET_MEMBER_NAME_CHECKED(UAbilitySystemTestAttributeSet, Health));
 	static UProperty *DamageProperty = FindFieldChecked<UProperty>(UAbilitySystemTestAttributeSet::StaticClass(), GET_MEMBER_NAME_CHECKED(UAbilitySystemTestAttributeSet, Damage));
 
@@ -80,6 +81,8 @@ bool UAbilitySystemTestAttributeSet::PreGameplayEffectExecute(struct FGameplayEf
 		// At this point, the Magnitude of the applied damage may have been modified by us. We still do the translation to Health in UAbilitySystemTestAttributeSet::PostAttributeModify.
 	}
 
+#endif
+
 	return true;
 }
 
@@ -94,7 +97,7 @@ void UAbilitySystemTestAttributeSet::PostGameplayEffectExecute(const struct FGam
 	if (DamageProperty == ModifiedProperty)
 	{
 		// Anytime Damage is applied with 'Damage.Fire' tag, there is a chance to apply a burning DOT
-		if (Data.EvaluatedData.Tags.HasTag( IGameplayTagsModule::RequestGameplayTag(FName(TEXT("FireDamage"))), EGameplayTagMatchType::IncludeParentTags, EGameplayTagMatchType::Explicit))
+		if (Data.EvaluatedData.SourceTags.HasTag( IGameplayTagsModule::RequestGameplayTag(FName(TEXT("FireDamage"))), EGameplayTagMatchType::IncludeParentTags, EGameplayTagMatchType::Explicit))
 		{
 			// Logic to rand() a burning DOT, if successful, apply DOT GameplayEffect to the target
 		}
diff --git a/Engine/Source/Runtime/GameplayAbilities/Private/GameplayAbilityTargetTypes.cpp b/Engine/Source/Runtime/GameplayAbilities/Private/GameplayAbilityTargetTypes.cpp
--- a/Engine/Source/Runtime/GameplayAbilities/Private/GameplayAbilityTargetTypes.cpp
+++ b/Engine/Source/Runtime/GameplayAbilities/Private/GameplayAbilityTargetTypes.cpp
@@ -7,16 +7,12 @@
 #include "AbilitySystemBlueprintLibrary.h"
 
 
-TArray<FActiveGameplayEffectHandle> FGameplayAbilityTargetData::ApplyGameplayEffect(const UGameplayEffect* GameplayEffect, const FGameplayEffectContextHandle& InEffectContext, float Level, FModifierQualifier Qualifier)
+TArray<FActiveGameplayEffectHandle> FGameplayAbilityTargetData::ApplyGameplayEffect(const UGameplayEffect* GameplayEffect, const FGameplayEffectContextHandle& InEffectContext, float Level, FPredictionKey PredictionKey)
 {
 	// Clone the effect context per target to avoid modifying original context
 	FGameplayEffectContextHandle EffectContext = InEffectContext.Duplicate();
 	
-	FGameplayEffectSpec	SpecToApply(GameplayEffect,					// The UGameplayEffect data asset
-		EffectContext,		// The actor who instigated this
-		Level,							// FIXME: Leveling
-		NULL							// FIXME: CurveData override... should we just remove this?
-		);
+	FGameplayEffectSpec	SpecToApply(GameplayEffect, EffectContext, Level);
 	
 	AddTargetDataToContext(EffectContext);
 
@@ -31,7 +27,7 @@ TArray<FActiveGameplayEffectHandle> FGameplayAbilityTargetData::ApplyGameplayEff
 			UAbilitySystemComponent* TargetComponent = UAbilitySystemBlueprintLibrary::GetAbilitySystemComponent(TargetActor.Get());
 			if (TargetComponent && EffectContext.GetInstigatorAbilitySystemComponent())
 			{
-				AppliedHandles.Add(EffectContext.GetInstigatorAbilitySystemComponent()->ApplyGameplayEffectSpecToTarget(SpecToApply, TargetComponent, Qualifier));
+				AppliedHandles.Add(EffectContext.GetInstigatorAbilitySystemComponent()->ApplyGameplayEffectSpecToTarget(SpecToApply, TargetComponent, PredictionKey));
 			}
 		}
 	}
