Gameplay Effect inital refactor

diff --git a/Engine/Source/Runtime/GameplayAbilities/Private/GameplayEffect.cpp b/Engine/Source/Runtime/GameplayAbilities/Private/GameplayEffect.cpp
--- a/Engine/Source/Runtime/GameplayAbilities/Private/GameplayEffect.cpp
+++ b/Engine/Source/Runtime/GameplayAbilities/Private/GameplayEffect.cpp
@@ -8,15 +8,14 @@
 #include "GameplayTagsModule.h"
 #include "AbilitySystemComponent.h"
 #include "AbilitySystemInterface.h"
+#include "GameplayModCalculation.h"
+
+#pragma optimize( "", off )
 
 const float UGameplayEffect::INFINITE_DURATION = -1.f;
 const float UGameplayEffect::INSTANT_APPLICATION = 0.f;
 const float UGameplayEffect::NO_PERIOD = 0.f;
-const float FGameplayEffectLevelSpec::INVALID_LEVEL = -1.f;
-
-#if SKILL_SYSTEM_AGGREGATOR_DEBUG
-FAggregator::FAllocationStats FAggregator::AllocationStats;
-#endif
+const float UGameplayEffect::INVALID_LEVEL = -1.f;
 
 // --------------------------------------------------------------------------------------------------------------------------------------------------------
 //
@@ -45,16 +44,6 @@ void UGameplayEffect::GetOwnedGameplayTags(FGameplayTagContainer& TagContainer)
 	TagContainer.AppendTags(OwnedTagsContainer);
 }
 
-void UGameplayEffect::GetClearGameplayTags(FGameplayTagContainer& TagContainer) const
-{
-	TagContainer.AppendTags(ClearTagsContainer);
-}
-
-bool UGameplayEffect::AreApplicationTagRequirementsSatisfied(const FGameplayTagContainer& InstigatorTags, const FGameplayTagContainer& TargetTags) const
-{
-	return (InstigatorTags.MatchesAll(ApplicationRequiredInstigatorTags, true) && TargetTags.MatchesAll(ApplicationRequiredTargetTags, true));
-}
-
 void UGameplayEffect::PostLoad()
 {
 	Super::PostLoad();
@@ -64,306 +53,239 @@ void UGameplayEffect::PostLoad()
 
 void UGameplayEffect::ValidateGameplayEffect()
 {
-	ValidateStacking();
-
 	// todo: add a check here for instant effects that modify incoming/outgoing GEs
 }
 
-void UGameplayEffect::ValidateStacking()
+bool FGameplayEffectAttributeCaptureDefinition::operator==(const FGameplayEffectAttributeCaptureDefinition& Other) const
 {
-	if (StackingPolicy != EGameplayEffectStackingPolicy::Unlimited)
-	{
-		StackingPolicy = StackingPolicy;
-
-		if (StackedAttribName == NAME_None)
-		{
-			ABILITY_LOG(Warning, TEXT("%s has a stacking rule but does not have an attribute to apply it to. Removing stacking rule."), *GetPathName());
-			StackingPolicy = EGameplayEffectStackingPolicy::Unlimited;
-		}
-
-		UDataTable* DataTable = IGameplayAbilitiesModule::Get().GetAbilitySystemGlobals()->GetGlobalAttributeMetaDataTable();
-
-		if (DataTable && DataTable->RowMap.Contains(StackedAttribName))
-		{
-			FAttributeMetaData* Row = (FAttributeMetaData*)DataTable->RowMap[StackedAttribName];
-			if (Row && Row->bCanStack == false)
-			{
-				ABILITY_LOG(Warning, TEXT("%s has a stacking rule but modifies attribute %s. %s is not allowed to stack. Removing stacking rule."), *GetPathName(), *Modifiers[0].Attribute.GetName(), *Modifiers[0].Attribute.GetName());
-				StackingPolicy = EGameplayEffectStackingPolicy::Unlimited;
-			}
-		}
-		else
-		{
-			ABILITY_LOG(Warning, TEXT("%s has a stacking rule but modifies attribute %s. %s was not found in the global attribute data table. Removing stacking rule."), *GetPathName(), *Modifiers[0].Attribute.GetName(), *Modifiers[0].Attribute.GetName());
-			StackingPolicy = EGameplayEffectStackingPolicy::Unlimited;
-		}
-
-		if (Modifiers.Num() == 0)
-		{
-			ABILITY_LOG(Warning, TEXT("%s has a stacking rule but does not have any modifiers to apply it to. Removing stacking rule."), *GetPathName());
-			StackingPolicy = EGameplayEffectStackingPolicy::Unlimited;
-		}
-		else
-		{
-			// look for modifiers to attributes that are not the stacking attribute and warn that they may not apply in a predictable fashion
-			for (int32 Idx = 0; Idx < Modifiers.Num(); ++Idx)
-			{
-				if (Modifiers[Idx].ModifierType == EGameplayMod::Attribute)
-				{
-					FName AttributeName(*Modifiers[Idx].Attribute.GetName());
-
-					if (AttributeName != StackedAttribName)
-					{
-						ABILITY_LOG(Warning, TEXT("%s has a stacking rule and modifies attribute %s but one of the modifiers modifies attribute %s. Changes to %s may not behave in a predictable fashion."), *GetPathName(), *StackedAttribName.ToString(), *AttributeName.ToString(), *AttributeName.ToString());
-					}
-				}
-				else if (Modifiers[Idx].ModifierType == EGameplayMod::ActiveGE || Modifiers[Idx].ModifierType == EGameplayMod::IncomingGE || Modifiers[Idx].ModifierType == EGameplayMod::OutgoingGE)
-				{
-					ABILITY_LOG(Warning, TEXT("%s has a stacking rule and modifies attribute %s but one of its modifiers modifies gameplay effects. Changes to other gameplay effects may not behave in a predictable fashion."), *GetPathName(), *StackedAttribName.ToString());
-				}
-			}
-		}
-	}
+	return ((AttributeToCapture == Other.AttributeToCapture) && (AttributeSource == Other.AttributeSource) && (bSnapshot == Other.bSnapshot));
 }
 
-// --------------------------------------------------------------------------------------------------------------------------------------------------------
-//
-//	FGameplayEffectSpec
-//
-// --------------------------------------------------------------------------------------------------------------------------------------------------------
-
-FGameplayEffectSpec::FGameplayEffectSpec(const UGameplayEffect* InDef, const FGameplayEffectContextHandle& InEffectContext, float Level, const FGlobalCurveDataOverride* CurveData)
-: Def(InDef)
-, ModifierLevel(TSharedPtr<FGameplayEffectLevelSpec>(new FGameplayEffectLevelSpec(Level, InDef->LevelInfo, InEffectContext.GetInstigator())))
-, EffectContext(InEffectContext)
-, Duration(new FAggregator(InDef->Duration.MakeFinalizedCopy(CurveData), ModifierLevel, SKILL_AGG_DEBUG(TEXT("%s Duration"), *InDef->GetName())))
-, Period(new FAggregator(InDef->Period.MakeFinalizedCopy(CurveData), ModifierLevel, SKILL_AGG_DEBUG(TEXT("%s Period"), *InDef->GetName())))
-, ChanceToApplyToTarget(new FAggregator(InDef->ChanceToApplyToTarget.MakeFinalizedCopy(CurveData), ModifierLevel, SKILL_AGG_DEBUG(TEXT("%s ChanceToApplyToTarget"), *InDef->GetName())))
-, ChanceToExecuteOnGameplayEffect(new FAggregator(InDef->ChanceToExecuteOnGameplayEffect.MakeFinalizedCopy(CurveData), ModifierLevel, SKILL_AGG_DEBUG(TEXT("%s ChanceToExecuteOnGameplayEffect"), *InDef->GetName())))
+bool FGameplayEffectAttributeCaptureDefinition::operator!=(const FGameplayEffectAttributeCaptureDefinition& Other) const
 {
-	Duration.Get()->RegisterLevelDependancies();
-	Period.Get()->RegisterLevelDependancies();
-	ChanceToApplyToTarget.Get()->RegisterLevelDependancies();
-	ChanceToExecuteOnGameplayEffect.Get()->RegisterLevelDependancies();
-
-	InitModifiers(CurveData, EffectContext, Level);
-
-	if (InDef)
-	{
-		for (UGameplayEffect* TargetDef : InDef->TargetEffects)
-		{
-			TargetEffectSpecs.Add(TSharedRef<FGameplayEffectSpec>(new FGameplayEffectSpec(TargetDef, EffectContext, Level, CurveData)));
-		}
-	}
+	return ((AttributeToCapture != Other.AttributeToCapture) || (AttributeSource != Other.AttributeSource) || (bSnapshot != Other.bSnapshot));
 }
 
-void FGameplayEffectSpec::InitModifiers(const FGlobalCurveDataOverride* CurveData, const FGameplayEffectContextHandle& InEffectContext, float Level)
+FString FGameplayEffectAttributeCaptureDefinition::ToSimpleString() const
 {
-	check(Def);
+	return FString::Printf(TEXT("Attribute: %s, Capture Point: %s, Snapshot: %d"), *AttributeToCapture.GetName(), AttributeSource == EGameplayEffectAttributeCaptureSource::Source ? TEXT("Source") : TEXT("Target"), bSnapshot);
+}
 
-	Modifiers.Reserve(Def->Modifiers.Num());
-	
-	for (const FGameplayModifierInfo &ModInfo : Def->Modifiers)
-	{
-		// Pass down the LevelInfo into this Modifier. 
-		// ModifierSpec may want to override how leveling will work (for this modifier only).
-		// Or it may use the GameplayEffectSpec's level. 
-		// ApplyNewDef will update NewLevelSpec appropriately.
+FGameplayEffectAttributeCaptureSpec::FGameplayEffectAttributeCaptureSpec()
+{
+}
 
-		TSharedPtr<FGameplayEffectLevelSpec> NewLevelSpec = ModifierLevel;
-		ModifierLevel->ApplyNewDef(ModInfo.LevelInfo, NewLevelSpec);
+FGameplayEffectAttributeCaptureSpec::FGameplayEffectAttributeCaptureSpec(const FGameplayEffectAttributeCaptureDefinition& InDefinition)
+	: BackingDefinition(InDefinition)
+{
+}
 
-		// This creates a new FModifierSpec that we own.
-		new (Modifiers)FModifierSpec(ModInfo, NewLevelSpec, CurveData, InEffectContext, Level);
-	}	
+const FGameplayEffectAttributeCaptureDefinition& FGameplayEffectAttributeCaptureSpec::GetBackingDefinition() const
+{
+	return BackingDefinition;
 }
 
-void FGameplayEffectSpec::MakeUnique()
+const FAggregator& FGameplayEffectAttributeCaptureSpec::GetAggregator() const
 {
-	for (FModifierSpec &ModSpec : Modifiers)
+	FAggregator* Aggregator = AttributeAggregator.Get();
+	if (ensure(Aggregator))
 	{
-		ModSpec.Aggregator.MakeUnique();
+		return *Aggregator;
 	}
 
-	Duration.MakeUnique();
-	Period.MakeUnique();
-	ChanceToApplyToTarget.MakeUnique();
-	ChanceToExecuteOnGameplayEffect.MakeUnique();
+	ABILITY_LOG(Error, TEXT("FGameplayEffectAttributeCaptureSpec::GetAggregator called with no valid aggregator"));
+	static FAggregator InvalidAggregator;
+	return InvalidAggregator;
 }
 
-int32 FGameplayEffectSpec::ApplyModifiersFrom(const FGameplayEffectSpec &InSpec, const FModifierQualifier &QualifierContext)
+void FGameplayEffectAttributeCaptureSpecContainer::AddCaptureDefinition(const FGameplayEffectAttributeCaptureDefinition& InCaptureDefinition)
 {
-	ABILITY_LOG_SCOPE(TEXT("FGameplayEffectSpec::ApplyModifiersFrom %s. InSpec: %s"), *this->ToSimpleString(), *InSpec.ToSimpleString());
+	const bool bSourceAttribute = (InCaptureDefinition.AttributeSource == EGameplayEffectAttributeCaptureSource::Source);
+	TArray<FGameplayEffectAttributeCaptureSpec>& AttributeArray = (bSourceAttribute ? SourceAttributes : TargetAttributes);
 
-	int32 NumApplied = 0;
-
-	if (!InSpec.Def || !InSpec.Def->AreGameplayEffectTagRequirementsSatisfied(Def))
+	// Only add additional captures if this exact capture definition isn't already being handled
+	if (!AttributeArray.ContainsByPredicate(
+			[&](const FGameplayEffectAttributeCaptureSpec& Element)
+			{
+				return Element.GetBackingDefinition() == InCaptureDefinition;
+			}))
 	{
-		// InSpec doesn't match this GameplayEffectSpec but if InSpec provides immunity we also need to check the modifiers because they can be canceled independent of the GameplayEffectSpec
-		if ((int32)InSpec.Def->AppliesImmunityTo == (int32)QualifierContext.Type())
+		AttributeArray.Add(FGameplayEffectAttributeCaptureSpec(InCaptureDefinition));
+
+		if (!InCaptureDefinition.bSnapshot)
 		{
-			for (int ii = 0; ii < Modifiers.Num(); ++ii)
-			{
-				const FGameplayModifierEvaluatedData& Data = Modifiers[ii].Aggregator.Get()->Evaluate();
-				if (InSpec.Def->AreGameplayEffectTagRequirementsSatisfied(Data.Tags))
-				{
-					Modifiers.RemoveAtSwap(ii);
-					--ii;
-				}
-			}
+			bHasNonSnapshottedAttributes = true;
 		}
-
-		return 0;
 	}
+}
 
-	if ((int32)InSpec.Def->AppliesImmunityTo == (int32)QualifierContext.Type())
+void FGameplayEffectAttributeCaptureSpecContainer::CaptureAttributes(UAbilitySystemComponent* InAbilitySystemComponent, EGameplayEffectAttributeCaptureSource InCaptureSource)
+{
+	if (InAbilitySystemComponent)
 	{
-		return -1;
-	}
-
-	// Fixme: Use acceleration structures to speed up these types of lookups
-	// The called functions below are reliant on the InSpecs evaluated data. We should ideally only call evaluate once per mod.
+		const bool bSourceComponent = (InCaptureSource == EGameplayEffectAttributeCaptureSource::Source);
+		TArray<FGameplayEffectAttributeCaptureSpec>& AttributeArray = (bSourceComponent ? SourceAttributes : TargetAttributes);
 
-	for (const FModifierSpec &InMod : InSpec.Modifiers)
-	{
-		NumApplied += ApplyModifier(InMod, QualifierContext, InSpec.ShouldApplyAsSnapshot(QualifierContext));
+		// Capture every spec's requirements from the specified component
+		for (FGameplayEffectAttributeCaptureSpec& CurCaptureSpec : AttributeArray)
+		{
+			InAbilitySystemComponent->CaptureAttributeForGameplayEffect(CurCaptureSpec);
+		}
 	}
-
-	return NumApplied;
 }
 
-int32 FGameplayEffectSpec::ExecuteModifiersFrom(const FGameplayEffectSpec &InSpec, const FModifierQualifier &QualifierContext)
+const FAggregator& FGameplayEffectAttributeCaptureSpecContainer::GetAggregator(const FGameplayEffectAttributeCaptureDefinition& InDefinition) const
 {
-	int32 NumExecuted = 0;
+	const bool bSourceAttribute = (InDefinition.AttributeSource == EGameplayEffectAttributeCaptureSource::Source);
+	const TArray<FGameplayEffectAttributeCaptureSpec>& AttributeArray = (bSourceAttribute ? SourceAttributes : TargetAttributes);
 
-	// Fixme: Use acceleration structures to speed up these types of lookups
-	for (FModifierSpec &MyMod : Modifiers)
-	{
-		for (const FModifierSpec &InMod : InSpec.Modifiers)
+	const FGameplayEffectAttributeCaptureSpec* MatchingSpec = AttributeArray.FindByPredicate(
+		[&](const FGameplayEffectAttributeCaptureSpec& Element)
 		{
-			if (InMod.CanModifyModifier(MyMod, QualifierContext))
-			{
-				InMod.ExecuteModOn(MyMod);
-				NumExecuted++;
-			}
-		}
+			return Element.GetBackingDefinition() == InDefinition;
+		});
+
+	if (MatchingSpec)
+	{
+		return MatchingSpec->GetAggregator();
 	}
 
-	return NumExecuted;
+	ABILITY_LOG(Error, TEXT("FGameplayEffectAttributeCaptureSpecContainer::GetAggregator called for attribute definition %s when that attribute was not captured."), *InDefinition.ToSimpleString());
+	static FAggregator InvalidAggregator;
+	return InvalidAggregator;
 }
 
-int32 FGameplayEffectSpec::ApplyModifier(const FModifierSpec &InMod, const FModifierQualifier &QualifierContext, bool bApplyAsSnapshot)
+bool FGameplayEffectAttributeCaptureSpecContainer::HasNonSnapshottedAttributes() const
+{
+	return bHasNonSnapshottedAttributes;
+}
+
+// --------------------------------------------------------------------------------------------------------------------------------------------------------
+//
+//	FGameplayEffectSpec
+//
+// --------------------------------------------------------------------------------------------------------------------------------------------------------
+
+FGameplayEffectSpec::FGameplayEffectSpec(const UGameplayEffect* InDef, const FGameplayEffectContextHandle& InEffectContext, float Level)
+: Def(InDef)
+, EffectContext(InEffectContext)
 {
-	int32 NumApplied = 0;
-	if (!InMod.CanModifyInContext(QualifierContext))
+	check(Def);	
+	SetLevel(Level);
+
+	// Init our ModifierSpecs
+	Modifiers.Reserve(Def->Modifiers.Num());
+	for (const FGameplayModifierInfo &ModInfo : Def->Modifiers)
 	{
-		return 0;
+		// This creates a new FModifierSpec that we own.
+		new (Modifiers)FModifierSpec(ModInfo);
 	}
 
-	switch (InMod.Info.EffectType)
+	// Gather all capture definitions from modifiers
+	for (const FModifierSpec& CurModSpec : Modifiers)
 	{
-		case EGameplayModEffect::Magnitude:
+		for (const FGameplayEffectAttributeCaptureDefinition& CurCaptureDef : CurModSpec.Info.MagnitudeV2.CustomMagnitudeCalculationAttributes)
 		{
-			for (FModifierSpec& MyMod : Modifiers)
-			{
-				if (InMod.CanModifyModifier(MyMod, QualifierContext))
-				{
-					InMod.ApplyModTo(MyMod, bApplyAsSnapshot);
-					NumApplied = 1;
-				}
-			}
-			break;
+			CapturedRelevantAttributes.AddCaptureDefinition(CurCaptureDef);
 		}
+	}
 
-		/**
-		 *	We do not check for Modifier qualification in the below cases, since these modifiers are independant of other modifiers!
-		 *	We can modify the duration of a GameplayEffect which only gives a tag but does not modify any attributes. (E.g., a cooldown)
-		 *	The downside is that you must tag your GameplayEffects rigorously in order to use Incoming/Outgoing/Active modififcation
-		 */
-
-		case EGameplayModEffect::Duration:
-		{
-			// don't modify infinite duration or instant unless we're overriding it
-			if (GetDuration() > 0.f || InMod.Info.ModifierOp == EGameplayModOp::Override)
-			{
-				Duration.Get()->ApplyMod(InMod.Info.ModifierOp, InMod.Aggregator, bApplyAsSnapshot);
-				NumApplied = 1;
-			}
+	// Gather all capture definitions from executions
+	for (const FGameplayEffectExecutionDefinition& Exec : Def->Executions)
+	{
+		if (Exec.CalculationClass == nullptr)
+			continue;
 
-			break;
-		}
+		TArray<FGameplayEffectAttributeCaptureDefinition> Reqs;
+		const UGameplayEffectCalculation* CDO = Cast<UGameplayEffectCalculation>(Exec.CalculationClass->ClassDefaultObject);
+		
+		CDO->GetCaptureDefinitions(Reqs);
 
-		case EGameplayModEffect::ChanceApplyTarget:
+		// Maybe it's better to pass in CapturedRelevantAttributes.AddCaptureDefinition(CurCaptureDef); to GetCaptureDefinitions?
+		for (const FGameplayEffectAttributeCaptureDefinition& Req : Reqs)
 		{
-			ChanceToApplyToTarget.Get()->ApplyMod(InMod.Info.ModifierOp, InMod.Aggregator, bApplyAsSnapshot);
-			NumApplied = 1;
-			break;
+			CapturedRelevantAttributes.AddCaptureDefinition(Req);
 		}
+	}
 
-		case EGameplayModEffect::ChanceExecuteEffect:
-		{
-			ChanceToExecuteOnGameplayEffect.Get()->ApplyMod(InMod.Info.ModifierOp, InMod.Aggregator, bApplyAsSnapshot);
-			NumApplied = 1;
-			break;
-		}
 
-		case EGameplayModEffect::LinkedGameplayEffect:
-		{
-			TargetEffectSpecs.Add(InMod.TargetEffectSpec.ToSharedRef());
-			NumApplied = 1;
-			break;
-		}
+	
+
+	// Capture source tags	
+	InEffectContext.GetOwnedGameplayTags(CapturedSourceTags);
+
+	// Capture source Attributes
+	// Is this the right place to do it? Do we ever need to create spec and capture attributes at a later time? If so, this will need to move.
+	CapturedRelevantAttributes.CaptureAttributes(InEffectContext.GetInstigatorAbilitySystemComponent(), EGameplayEffectAttributeCaptureSource::Source);
+
+	// Make TargetEffectSpecs too
+	for (UGameplayEffect* TargetDef : InDef->TargetEffects)
+	{
+		TargetEffectSpecs.Add(TSharedRef<FGameplayEffectSpec>(new FGameplayEffectSpec(TargetDef, EffectContext, Level)));
 	}
+}
 
+void FGameplayEffectSpec::SetLevel(float InLevel)
+{
+	Level = InLevel;
+	Duration = Def->Duration.GetValueAtLevel(InLevel);
+	Period = Def->Period.GetValueAtLevel(InLevel);
+	ChanceToApplyToTarget = Def->ChanceToApplyToTarget.GetValueAtLevel(InLevel);
+	ChanceToExecuteOnGameplayEffect = Def->ChanceToExecuteOnGameplayEffect.GetValueAtLevel(InLevel);
+}
 
-	return NumApplied;
+float FGameplayEffectSpec::GetLevel() const
+{
+	return Level;
 }
 
 float FGameplayEffectSpec::GetDuration() const
 {
-	return Duration.Get()->Evaluate().Magnitude;
+	return Duration;
 }
 
 float FGameplayEffectSpec::GetPeriod() const
 {
-	return Period.Get()->Evaluate().Magnitude;
+	return Period;
 }
 
 float FGameplayEffectSpec::GetChanceToApplyToTarget() const
 {
-	return ChanceToApplyToTarget.Get()->Evaluate().Magnitude;
+	return ChanceToApplyToTarget;
 }
 
 float FGameplayEffectSpec::GetChanceToExecuteOnGameplayEffect() const
 {
-	return ChanceToExecuteOnGameplayEffect.Get()->Evaluate().Magnitude;
+	return ChanceToExecuteOnGameplayEffect;
+}
+
+float FGameplayEffectSpec::GetModifierMagnitude(const FModifierSpec& ModSpec) const
+{
+	return ModSpec.EvaluatedMagnitude;
+
+	//return ModSpec.Info.Magnitude.GetValueAtLevel(Level);
+}
+
+void FGameplayEffectSpec::CalculateModifierMagnitudes()
+{
+	for (FModifierSpec& Mod : Modifiers)
+	{
+		Mod.CalculateMagnitude(*this);
+	}
 }
 
 float FGameplayEffectSpec::GetMagnitude(const FGameplayAttribute &Attribute) const
 {
 	float CurrentMagnitude = 0.f;
-
 	for (const FModifierSpec &Mod : Modifiers)
-	{
-		if (Mod.Info.ModifierType == EGameplayMod::Attribute)
+	{		
+		if (Mod.Info.Attribute != Attribute)
 		{
-			if (Mod.Info.Attribute != Attribute)
-			{
-				continue;
-			}
-
-			// Todo: Tags/application checks here - make sure we can still apply
-
-			// First, evaluate all of our data 
-
-			FGameplayModifierEvaluatedData EvaluatedData = Mod.Aggregator.Get()->Evaluate();
-
-			FAggregator	Aggregator(CurrentMagnitude, SKILL_AGG_DEBUG(TEXT("Magnitude of Attribute %s"), *Mod.Info.Attribute.GetName()));
-
-			Aggregator.ExecuteMod(Mod.Info.ModifierOp, EvaluatedData);
-
-			CurrentMagnitude = Aggregator.Evaluate().Magnitude;
+			continue;
 		}
+
+		CurrentMagnitude = GetModifierMagnitude(Mod);
+		break;
 	}
 
 	return CurrentMagnitude;
@@ -422,27 +344,6 @@ void FGameplayEffectSpec::PruneModifiedAttributes()
 	}
 }
 
-bool FGameplayEffectSpec::ShouldApplyAsSnapshot(const FModifierQualifier &QualifierContext) const
-{
-	bool ShouldSnapshot;
-	switch(Def->CopyPolicy)
-	{
-		case EGameplayEffectCopyPolicy::AlwaysSnapshot:
-			ShouldSnapshot = true;
-			break;
-			
-		case EGameplayEffectCopyPolicy::AlwaysLink:
-			ShouldSnapshot = false;
-			break;
-			
-		default:
-			ShouldSnapshot = (QualifierContext.Type() == EGameplayMod::OutgoingGE);
-			break;
-	}
-	
-	return ShouldSnapshot;
-}
-
 EGameplayEffectStackingPolicy::Type FGameplayEffectSpec::GetStackingType() const
 {
 	return Def->StackingPolicy;
@@ -454,674 +355,94 @@ EGameplayEffectStackingPolicy::Type FGameplayEffectSpec::GetStackingType() const
 //
 // --------------------------------------------------------------------------------------------------------------------------------------------------------
 
-bool FModifierSpec::CanModifyInContext(const FModifierQualifier &QualifierContext) const
-{
-	// Can only modify if this spec is valid within this Qualifier Context
-	// (E.g, if I am an OutgoingGE mod, I cannot modify during a IncomingGE context)
-	if (Info.ModifierType != QualifierContext.Type())
-	{
-		return false;
-	}	
-
-	return true;
-}
-
-bool FModifierSpec::CanModifyModifier(FModifierSpec &Other, const FModifierQualifier &QualifierContext) const
-{
-	// Attribute is essentially a key. This isn't 100% necessary - we could just rely on tags
-	if (Info.Attribute != Other.Info.Attribute)
-	{
-		return false;
-	}
-
-	// Tag checking is done at the FAggregator level. So all we do here is the attribute check.
-
-	return true;
-}
-
-FModifierSpec::FModifierSpec(const FGameplayModifierInfo& InInfo, TSharedPtr<FGameplayEffectLevelSpec> InLevel, const FGlobalCurveDataOverride* CurveData, const FGameplayEffectContextHandle& InEffectContext, float Level)
+FModifierSpec::FModifierSpec(const FGameplayModifierInfo& InInfo)
 : Info(InInfo)
-, Aggregator(new FAggregator(FGameplayModifierData(InInfo, CurveData), InLevel, SKILL_AGG_DEBUG(TEXT("FModifierSpec: %s "), *InInfo.ToSimpleString())))
+, EvaluatedMagnitude(0.f)
 {
-	Aggregator.Get()->RegisterLevelDependancies();
-	if (InInfo.TargetEffect)
-	{
-		TargetEffectSpec = TSharedPtr<FGameplayEffectSpec>(new FGameplayEffectSpec(InInfo.TargetEffect, InEffectContext, Level, CurveData));
-	}
-}
-
-void FModifierSpec::ApplyModTo(FModifierSpec &Other, bool TakeSnapshot) const
-{
-	Other.Aggregator.Get()->ApplyMod(this->Info.ModifierOp, this->Aggregator, TakeSnapshot);
-}
 
-void FModifierSpec::ExecuteModOn(FModifierSpec &Other) const
-{
-	ABILITY_LOG_SCOPE(TEXT("Executing %s on %s"), *ToSimpleString(), *Other.ToSimpleString() );
-	Other.Aggregator.Get()->ExecuteModAggr(this->Info.ModifierOp, this->Aggregator);
 }
 
-bool FModifierSpec::AreTagRequirementsSatisfied(const FModifierSpec &ModifierToBeModified) const
+float FModifierSpec::GetEvaluatedMagnitude() const
 {
-	const FGameplayModifierEvaluatedData& ToBeModifiedData = ModifierToBeModified.Aggregator.Get()->Evaluate();
-
-	bool HasRequired = ToBeModifiedData.Tags.MatchesAll(this->Info.RequiredTags, true);
-	bool HasIgnored = ToBeModifiedData.Tags.MatchesAny(this->Info.IgnoreTags, false);
-
-	return HasRequired && !HasIgnored;
+	return EvaluatedMagnitude;
 }
 
-// --------------------------------------------------------------------------------------------------------------------------------------------------------
-//
-//	FAggregatorRef
-//
-// --------------------------------------------------------------------------------------------------------------------------------------------------------
+void FModifierSpec::CalculateMagnitude(OUT FGameplayEffectSpec& OwnerSpec)
+{	
+	EvaluatedMagnitude = Info.Magnitude.GetValueAtLevel(OwnerSpec.GetLevel());
 
-void FAggregatorRef::MakeUnique()
-{
-	ABILITY_LOG(Log, TEXT("MakeUnique %s"), *ToString());
-
-	// Make a hard ref copy of our FAggregator
-	MakeHardRef();
-	SharedPtr = TSharedPtr<FAggregator>(new FAggregator(*SharedPtr.Get()));
-	WeakPtr = SharedPtr;
-
-	// Update dependancy chain so that the copy we just made is updated if any of the applied modifiers change
-	SharedPtr->RefreshDependencies();
-}
-
-void FAggregatorRef::MakeUniqueDeep()
-{
-	ABILITY_LOG_SCOPE(TEXT("MakeUniqueDeep %s"), *ToString());
-
-	// Make a hard ref copy of our FAggregator
-	MakeHardRef();
-	SharedPtr = TSharedPtr<FAggregator>(new FAggregator(*SharedPtr.Get()));
-	WeakPtr = SharedPtr;
-
-	// Make all of our mods UniqueDeep
-	Get()->MakeUniqueDeep();
-
-	// Update dependancy chain so that the copy we just made is updated if any of the applied modifiers change
-	SharedPtr->RefreshDependencies();
-}
-
-FString FAggregatorRef::ToString() const
-{
-	if (SharedPtr.IsValid())
-	{
-		return FString::Printf(TEXT("[HardRef to: %s]"), *SharedPtr->ToSimpleString());
-	}
-	if (WeakPtr.IsValid())
+	// @todo: this should have an enumeration over which to use
+	/*
+	if (Info.MagnitudeV2.CalculationClassMagnitude)
 	{
-		return FString::Printf(TEXT("[SoftRef to: %s]"), *WeakPtr.Pin()->ToSimpleString());
-	}
-	return FString(TEXT("Invalid"));
-}
-
-bool FAggregatorRef::NetSerialize(FArchive& Ar, class UPackageMap* Map, bool& bOutSuccess)
-{
-	float EvaluatedMagnitude=0.f;
+		const UGameplayModCalculation* ModCalcCDO = Info.MagnitudeV2.CalculationClassMagnitude->GetDefaultObject<UGameplayModCalculation>();
+		check(ModCalcCDO);
 
-	if (Ar.IsSaving())
-	{
-		if (IsValid())
-		{
-			EvaluatedMagnitude = Get()->Evaluate().Magnitude;
-		}
-	}
-
-	Ar << EvaluatedMagnitude;
-
-	if (Ar.IsLoading())
-	{
-		if (!IsValid())
-		{
-			SharedPtr = TSharedPtr<struct FAggregator>(new FAggregator);
-			WeakPtr = SharedPtr;
-		}
-
-		Get()->SetFromNetSerialize(EvaluatedMagnitude);
-	}
-	return true;
-}
-
-void FAggregator::SetFromNetSerialize(float SerializedValue)
-{
-	BaseData = FGameplayModifierData(SerializedValue, NULL);
-	CachedData = FGameplayModifierEvaluatedData(SerializedValue);
-}
-
-// --------------------------------------------------------------------------------------------------------------------------------------------------------
-//
-//	FAggregator
-//
-// --------------------------------------------------------------------------------------------------------------------------------------------------------
-
-FAggregator::FAggregator()
-{
-#if SKILL_SYSTEM_AGGREGATOR_DEBUG
-	CopiesMade = 0;
-	FAggregator::AllocationStats.DefaultCStor++;
-#endif
-}
-
-FAggregator::FAggregator(const FGameplayModifierData &InBaseData, TSharedPtr<FGameplayEffectLevelSpec> InLevel, const TCHAR *InDebugStr)
-: Level(InLevel)
-, BaseData(InBaseData)
-{
-#if SKILL_SYSTEM_AGGREGATOR_DEBUG
-	if (InDebugStr)
-	{
-		DebugString = FString(InDebugStr);
-	}
-	CopiesMade = 0;
-	FAggregator::AllocationStats.ModifierCStor++;
-#endif
-}
-
-FAggregator::FAggregator(const FScalableFloat &InBaseMagnitude, TSharedPtr<FGameplayEffectLevelSpec> LevelInfo, const TCHAR *InDebugStr)
-: Level(LevelInfo)
-, BaseData(InBaseMagnitude)
-{
-#if SKILL_SYSTEM_AGGREGATOR_DEBUG
-	if (InDebugStr)
-	{
-		DebugString = FString(InDebugStr);
-	}
-	CopiesMade = 0;
-	FAggregator::AllocationStats.ScalableFloatCstor++;
-#endif
-}
-
-FAggregator::FAggregator(const FGameplayModifierEvaluatedData &InEvalData, const TCHAR *InDebugStr)
-: Level(TSharedPtr<FGameplayEffectLevelSpec>(new FGameplayEffectLevelSpec()))
-, BaseData(InEvalData.Magnitude, InEvalData.Callbacks)
-{
-#if SKILL_SYSTEM_AGGREGATOR_DEBUG
-	if (InDebugStr)
-	{
-		DebugString = FString(InDebugStr);
-	}
-	CopiesMade = 0;
-	FAggregator::AllocationStats.FloatCstor++;
-#endif
-}
-
-FAggregator::FAggregator(const FAggregator &In)
-{
-	*this = In;
-#if SKILL_SYSTEM_AGGREGATOR_DEBUG
-	FAggregator::AllocationStats.CopyCstor++;
-	CopiesMade = 0;
-	DebugString = FString::Printf(TEXT("Copy %d of [%s]"), ++In.CopiesMade, *In.DebugString);
-#endif
-} 
-
-FAggregator::~FAggregator()
-{
-	for (TWeakPtr<FAggregator> WeakPtr : Dependants)
-	{
-		if (WeakPtr.IsValid())
-		{
-			ABILITY_LOG(Log, TEXT("%s Marking Dependant %s Dirty on Destroy"), *ToSimpleString(), *WeakPtr.Pin()->ToSimpleString());
-
-			WeakPtr.Pin()->MarkDirty();
-#if SKILL_SYSTEM_AGGREGATOR_DEBUG
-			FAggregator::AllocationStats.DependantsUpdated++;
-#endif
-		}
-	}
-}
-
-void FAggregator::RefreshDependencies()
-{
-	TWeakPtr<FAggregator> LocalWeakPtr(SharedThis(this));
-	for (int32 i = 0; i < EGameplayModOp::Max; ++i)
-	{
-		for (FAggregatorRef &Ref : Mods[i])
-		{
-			if (Ref.IsValid())
-			{
-				Ref.Get()->AddDependantAggregator(LocalWeakPtr);
-			}
-		}
-	}
-	
-	RegisterLevelDependancies();
-}
-
-void FAggregator::ApplyMod(EGameplayModOp::Type ModType, FAggregatorRef Ref, bool TakeSnapshot)
-{
-	if (TakeSnapshot)
-	{
-		Ref.MakeUniqueDeep();
+		EvaluatedMagnitude = ModCalcCDO->CalculateMagnitude(OwnerSpec);
 	}
 	else
 	{
-		Ref.MakeSoftRef();
+		EvaluatedMagnitude = Info.MagnitudeV2.ScalableFloatMagnitude.GetValueAtLevel(OwnerSpec.GetLevel());
 	}
-	Mods[ModType].Push(Ref);
-
-	// Make the ref tell us if it changes
-	if (HasBeenAlreadyMadeSharable())
-	{
-		Ref.Get()->AddDependantAggregator(TWeakPtr< FAggregator >(SharedThis(this)));
-	}
-	MarkDirty();
-}
-
-// Execute is intended to directly modify the base value of an aggregator.
-// However if the aggregator's base value is not static (scales with some outside influence)
-// we will treat this execute as an apply, but make our own copy of the passed in aggregator.
-// (so that it will essentially be permanent).
-void FAggregator::ExecuteModAggr(EGameplayModOp::Type ModType, FAggregatorRef Ref)
-{
-	const FGameplayModifierEvaluatedData& EvaluatedData = Ref.Get()->Evaluate();
-	ExecuteMod(ModType, EvaluatedData);	
+	*/
 }
 
-void FAggregator::ExecuteMod(EGameplayModOp::Type ModType, const FGameplayModifierEvaluatedData& EvaluatedData)
-{
-	// Am I allowed to even be executed on? If my base data scales then all I can do is apply this to myself
-	if (!BaseData.Magnitude.IsStatic())
-	{
-		ABILITY_LOG(Log, TEXT("Treating Execute as an Apply since our Level is not static!"));
-		FAggregatorRef NewAggregator(new FAggregator(EvaluatedData, SKILL_AGG_DEBUG(TEXT("ExecutedMod"))));
-		ApplyMod(ModType, NewAggregator, false);
-		return;
-	}
-
-	switch (ModType)
-	{
-		case EGameplayModOp::Override:
-		{
-			BaseData.Magnitude.SetValue(EvaluatedData.Magnitude);
-			BaseData.Tags = EvaluatedData.Tags;
-			break;
-		}
-		case EGameplayModOp::Additive:
-		{
-			BaseData.Magnitude.Value += EvaluatedData.Magnitude;
-			BaseData.Tags.AppendTags(EvaluatedData.Tags);
-			break;
-		}
-		case EGameplayModOp::Multiplicitive:
-		{
-			BaseData.Magnitude.Value *= EvaluatedData.Magnitude;
-			BaseData.Tags.AppendTags(EvaluatedData.Tags);
-			break;
-		}
-		case EGameplayModOp::Division:
-		{
-			BaseData.Magnitude.Value /= EvaluatedData.Magnitude;
-			BaseData.Tags.AppendTags(EvaluatedData.Tags);
-			break;
-		}
-		case EGameplayModOp::Callback:
-		{
-			// Executing a callback mod on another aggregator, what is expected here?
-			check(false);
-			break;
-		}
-	}
-
-	ABILITY_LOG(Log, TEXT("ExecuteMod: %s new BaseData.Magnitude: %s"), *ToSimpleString(), *BaseData.Magnitude.ToSimpleString());
-	MarkDirty();
-}
-
-void FAggregator::AddDependantAggregator(TWeakPtr<FAggregator> InDependant)
-{
-	check(InDependant.IsValid());
-	check(!Dependants.Contains(InDependant));
-
-	ABILITY_LOG(Log, TEXT("AddDependantAggregator: %s is a dependant of %s"), *InDependant.Pin()->ToSimpleString(), *ToSimpleString());
-	Dependants.Add(InDependant);
-}
-
-void FAggregator::TakeSnapshotOfLevel()
-{
-	check(Level.IsValid());
-	Level = TSharedPtr<FGameplayEffectLevelSpec>( new FGameplayEffectLevelSpec( *Level.Get()) );
-}
-
-void FAggregator::RegisterLevelDependancies()
-{
-	if (!BaseData.Magnitude.IsStatic())
-	{
-		TWeakPtr<FAggregator> LocalWeakPtr(SharedThis(this));
-		Level->RegisterLevelDependancy(LocalWeakPtr);
-	}
-}
-
-// Please try really hard to never add a "force full re-evaluate" flag to this function!
-// We want to strive to make this system dirty the cached data when its actually dirtied, 
-// and never do full catch all rebuilds.
-
-const FGameplayModifierEvaluatedData& FAggregator::Evaluate() const
-{
-	SCOPE_CYCLE_COUNTER(STAT_AggregatorEvaluate);
-
-	ABILITY_LOG_SCOPE(TEXT("Aggregator Evaluate %s"), *ToSimpleString());
-
-	if (!CachedData.IsValid)
-	{
-		// ------------------------------------------------------------------------
-		// If there are any overrides, then just take the first valid one.
-		// ------------------------------------------------------------------------
-		for (const FAggregatorRef &Agg : Mods[EGameplayModOp::Override])
-		{
-			ABILITY_LOG_SCOPE(TEXT("EGameplayModOp::Override"));
-			if (Agg.IsValid())
-			{
-				CachedData = Agg.Get()->Evaluate();
-				return CachedData;
-			}
-		}
-
-		// If we are going to do math, we need to lock our base value in.
-		// Calculate our magnitude at our level
-		// (We may not have a level, in that case, ensure we also don't have a leveling table)
-
-		float EvaluatedMagnitude = 0.f;
-		if (Level->IsValid())
-		{
-			EvaluatedMagnitude = BaseData.Magnitude.GetValueAtLevel(Level->GetLevel());
-		}
-		else
-		{
-			EvaluatedMagnitude = BaseData.Magnitude.GetValueChecked();
-		}
-
-		CachedData = FGameplayModifierEvaluatedData(EvaluatedMagnitude, BaseData.Callbacks, ActiveHandle, &BaseData.Tags);
-
-		int32 TotalModCount = Mods[EGameplayModOp::Additive].Num() + Mods[EGameplayModOp::Multiplicitive].Num() + Mods[EGameplayModOp::Division].Num();
-		
-		// Early out if no mods.
-		// We need to calculate num of mods anyways to do ModList.Reserve.
-		if (TotalModCount <= 0)
-		{
-			ABILITY_LOG(Log, TEXT("Final Magnitude: %.2f"), CachedData.Magnitude);
-			return CachedData;
-		}
-
-		// ------------------------------------------------------------------------
-		//	Apply Numeric Modifiers
-		//		This is convoluted due to tagging. We have modifiers that require we have or don't have certain tags.
-		//		These mods can also give us new tags. Its possible the 1st mod will require a tag that the 2nd mod gives us.
-		//
-		//	The basic approach here is create an ordered, linear list of all modifiers:
-		//		[Additive][Multiplicative][Division] mods
-		//
-		//	We make a pass through the list, aggregating as we go. During a pass we keep track of what what tags we've added
-		//	and if there were any mods that we rejected due to not having tags. When the pass is over, we check if we added
-		//	any that we needed. If so, we make another pass. (Once a modifier aggregated, we remove it from the list).
-		//
-		//	Paradoxes are still possible. ModX gives tag A, requires we don't have tag B. ModY gives tag B, requires we don't have tag A.
-		//	We detect this in a single pass. In the above example we would aggregate ModX, but warn loudly when we found ModY.
-		//	(we expect content to solve this via stacking rules, or just not making these type of requirements).
-		//
-		//
-		//
-		//	
-		// ------------------------------------------------------------------------
-		
-		// We have to do tag aggregation to figure out what we can and can't apply.
-		FGameplayTagContainer CommittedIgnoreTags;
-
-		TArray<const FAggregator*> ModList;
-		ModList.Reserve(TotalModCount + (EGameplayModOp::Override - EGameplayModOp::Additive));
-
-		// Build linear list of what we will aggregate
-		for (int32 OpIdx = (EGameplayModOp::Additive); OpIdx < EGameplayModOp::Override; ++OpIdx)
-		{
-			for (const FAggregatorRef &Agg : Mods[OpIdx])
-			{	
-				ModList.Add(Agg.Get());
-			}
-			ModList.Add(this);	// Sentinel value to signify "NextOp"
-		}
-
-		static const float OpBias[EGameplayModOp::Override] = {
-			0.f,	// EGameplayModOp::Additive
-			1.f,	// EGameplayModOp::Multiplicitive
-			1.f,	// EGameplayModOp::Division
-		};
-
-		float OpAggregation[EGameplayModOp::Override] = {
-			0.f,	// EGameplayModOp::Additive
-			1.f,	// EGameplayModOp::Multiplicitive
-			1.f,	// EGameplayModOp::Division
-		};
-
-		// Make multiple passes to do tagging
-		while (true)
-		{
-			FGameplayTagContainer	NewGiveTags;
-			FGameplayTagContainer	NewIgnoreTags;
-			FGameplayTagContainer	MissingTags;
-
-			EGameplayModOp::Type	ModOp = (EGameplayModOp::Additive);
-			
-			for (int32 ModIdx = 0; ModIdx < ModList.Num()-1; ++ModIdx)
-			{
-				const FAggregator* Agg = ModList[ModIdx];
-				if (Agg == NULL)
-				{
-					continue;
-				}
-				if (Agg == this)
-				{
-					ModOp = static_cast<EGameplayModOp::Type>(static_cast<int32>(ModOp)+1);
-					check(ModOp < EGameplayModOp::Override);
-					continue;
-				}
-				
-				const FGameplayModifierData &ModBaseData = ModList[ModIdx]->BaseData;
-				const FGameplayTagContainer &ModIgnoreTags = ModBaseData.IgnoreTags;
-				const FGameplayTagContainer &ModRequireTags = ModBaseData.RequireTags;
-
-				// This mod requires we have certain tags
-				if (ModRequireTags.Num() > 0 && !CachedData.Tags.MatchesAll(ModRequireTags, false) && !NewGiveTags.MatchesAll(ModRequireTags, false))
-				{
-					// But something else could give us this tag! So keep track of it and don't remove this Mod from the ModList.
-					MissingTags.AppendTags(ModRequireTags);
-					continue;
-				}
-
-				// This mod is now either accepted or rejected. It will not be needed for subsequent passes in this Evaluate, so NULL It out now.
-				ModList[ModIdx] = NULL;
-				
-				// This mod requires we don't have certain tags
-				if (CachedData.Tags.MatchesAny(ModIgnoreTags, false))
-				{
-					continue;
-				}
-
-				// Check for conflicts within this pass
-				if (ModIgnoreTags.Num() > 0 && NewGiveTags.Num() > 0 && ModIgnoreTags.MatchesAny(NewGiveTags, false))
-				{
-					// Pass problem!
-					ABILITY_LOG(Warning, TEXT("Tagging conflicts during Aggregate! Use Stacking rules to avoid this"));
-					ABILITY_LOG(Warning, TEXT("   While Evaluating: %s"), *ToSimpleString());
-					ABILITY_LOG(Warning, TEXT("   Applying Mod: %s"), *Agg->ToSimpleString());
-					ABILITY_LOG(Warning, TEXT("   ModIgnoreTags: %s"), *ModIgnoreTags.ToString() );
-					ABILITY_LOG(Warning, TEXT("   NewGiveTags: %s"), *NewGiveTags.ToString());
-					continue;
-				}
-
-				// Our requirements on the mod
-				const FGameplayModifierEvaluatedData& ModEvaluatedData = Agg->Evaluate();
-
-				// We have already committed during this aggregation to not have certain tags
-				if (CommittedIgnoreTags.Num() > 0 && CommittedIgnoreTags.MatchesAny(ModEvaluatedData.Tags, false))
-				{
-					continue;
-				}
-
-				if (NewIgnoreTags.Num() > 0 && ModEvaluatedData.Tags.Num() > 0 && NewIgnoreTags.MatchesAny(ModEvaluatedData.Tags, false))
-				{
-					// Pass problem!
-					ABILITY_LOG(Warning, TEXT("Tagging conflicts during Aggregate! Use Stacking rules to avoid this"));
-					ABILITY_LOG(Warning, TEXT("   While Evaluating: %s"), *ToSimpleString());
-					ABILITY_LOG(Warning, TEXT("   Applying Mod: %s"), *Agg->ToSimpleString());
-					ABILITY_LOG(Warning, TEXT("   Mod Tags: %s"), *ModEvaluatedData.Tags.ToString());
-					ABILITY_LOG(Warning, TEXT("   NewIgnoreTags: %s"), *NewIgnoreTags.ToString());
-					continue;
-				}
-
-				// Commit this Mod
-				NewIgnoreTags.AppendTags(ModIgnoreTags);
-				ModEvaluatedData.Aggregate(NewGiveTags, OpAggregation[ModOp], OpBias[ModOp]);
-			}
-
-			// Commit this pass's tags
-			CachedData.Tags.AppendTags(NewGiveTags);
-
-			// Keep doing passes until we don't add any new tags or don't have any missing tags
-			if (MissingTags.Num() <= 0 || !MissingTags.MatchesAny( NewGiveTags, false ))
-			{
-				break;
-			}
-		}
-
-		float Division = OpAggregation[EGameplayModOp::Division] > 0.f ? OpAggregation[EGameplayModOp::Division] : 1.f;
-
-		CachedData.Magnitude = ((CachedData.Magnitude + OpAggregation[EGameplayModOp::Additive]) * OpAggregation[EGameplayModOp::Multiplicitive]) / Division;
-
-		ABILITY_LOG(Log, TEXT("Final Magnitude: %.2f"), CachedData.Magnitude);
-		
-	}
-	else
-	{
-		ABILITY_LOG(Log, TEXT("CachedData was valid. Magnitude: %.2f"), CachedData.Magnitude);
-	}
-	return CachedData;
-}
-
-void FAggregator::PreEvaluate(FGameplayEffectModCallbackData &Data) const
-{
-	check(Data.ModifierSpec.Aggregator.Get() == this);
-	for (const FAggregatorRef &Agg : Mods[EGameplayModOp::Callback])
-	{
-		if (Agg.IsValid())
-		{
-			Agg.Get()->Evaluate().InvokePreExecute(Data);
-		}
-	}
-}
-
-void FAggregator::PostEvaluate(const struct FGameplayEffectModCallbackData &Data) const
-{
-	check(Data.ModifierSpec.Aggregator.Get() == this);
-	for (const FAggregatorRef &Agg : Mods[EGameplayModOp::Callback])
-	{
-		if (Agg.IsValid())
-		{
-			Agg.Get()->Evaluate().InvokePostExecute(Data);
-		}
-	}
-}
-
-FAggregator & FAggregator::MarkDirty()
-{
-	CachedData.IsValid = false;
-
-	// Execute OnDirty callbacks first. This may do things like update the actual uproperty value of an attribute.
-	OnDirty.ExecuteIfBound(this);
-
-	// Now tell people who depend on my value that I have changed. Important to do this after the OnDirty callback has been called.
-	for (int32 i=0; i < Dependants.Num(); ++i)
-	{
-		TWeakPtr<FAggregator> WeakPtr = Dependants[i];
-		if (WeakPtr.IsValid())
-		{
-			ABILITY_LOG(Log, TEXT("%s Marking Dependant %s Dirty (from ::MarkDirty())"), *ToSimpleString(), *WeakPtr.Pin()->ToSimpleString());
-#if SKILL_SYSTEM_AGGREGATOR_DEBUG
-			FAggregator::AllocationStats.DependantsUpdated++;
-#endif
-			WeakPtr.Pin()->MarkDirty();
-		}
-		else
-		{
-			Dependants.RemoveAtSwap(i);
-			--i;
-		}
-	}
-
-	return *this;
-}
-
-void FAggregator::MakeUniqueDeep()
-{
-	for (int32 i = 0; i < EGameplayModOp::Max; ++i)
-	{
-		for (FAggregatorRef &Ref : Mods[i])
-		{
-			if (Ref.IsValid())
-			{
-				Ref.MakeUniqueDeep();
-			}
-		}
-	}
-}
-
-void FAggregator::ClearAllDependancies()
-{
-	Dependants.SetNum(0);
-	OnDirty.Unbind();
-}
 
 // --------------------------------------------------------------------------------------------------------------------------------------------------------
 //
-//	FGameplayModifierEvaluatedData
+//	FActiveGameplayEffect
 //
 // --------------------------------------------------------------------------------------------------------------------------------------------------------
 
-void FGameplayModifierEvaluatedData::Aggregate(FGameplayTagContainer &OutTags, float &OutMagnitude, const float Bias) const
-{
-	OutMagnitude += (Magnitude - Bias);
-	OutTags.AppendTags(Tags);
-}
 
-void FGameplayModifierEvaluatedData::InvokePreExecute(FGameplayEffectModCallbackData &Data) const
+void FActiveGameplayEffect::CheckOngoingTagRequirements(const FGameplayTagContainer& OwnerTags, FActiveGameplayEffectsContainer& OwningContainer)
 {
-	if (Callbacks)
+	bool ShouldBeInhibited = !Spec.Def->OngoingTagRequirements.RequirementsMet(OwnerTags);
+
+	if (IsInhibited != ShouldBeInhibited)
 	{
-		for ( const FGameplayModifierCallback& Callback : *Callbacks )
+		if (ShouldBeInhibited)
 		{
-			TSubclassOf<UGameplayEffectExtension> ExtClass = Callback.ExtensionClass;
-			if (ExtClass)
+			// Register this ActiveGameplayEffects modifiers with our Attribute Aggregators
+			if (Spec.GetPeriod() == UGameplayEffect::NO_PERIOD)
 			{
-				UGameplayEffectExtension * Ext = ExtClass->GetDefaultObject<UGameplayEffectExtension>();
-				Ext->PreGameplayEffectExecute(*this, Data);
+				for (int32 ModIdx = 0; ModIdx < Spec.Modifiers.Num(); ++ModIdx)
+				{
+					const FModifierSpec &Mod = Spec.Modifiers[ModIdx];
+					
+					ABILITY_LOG_SCOPE(TEXT("Applying Attribute Mod %s to property"), *Mod.ToSimpleString());
+
+					FAggregator* Aggregator = OwningContainer.FindOrCreateAttributeAggregator(Mod.Info.Attribute).Get();
+
+					Aggregator->RemoveMod(Handle);
+				}
 			}
 		}
-	}
-}
-
-void FGameplayModifierEvaluatedData::InvokePostExecute(const FGameplayEffectModCallbackData &Data) const
-{
-	if (Callbacks)
-	{
-		for ( const FGameplayModifierCallback& Callback : *Callbacks )
+		else
 		{
-			TSubclassOf<class UGameplayEffectExtension> ExtClass = Callback.ExtensionClass;
-			if (ExtClass)
+			// Register this ActiveGameplayEffects modifiers with our Attribute Aggregators
+			if (Spec.GetPeriod() == UGameplayEffect::NO_PERIOD)
 			{
-				UGameplayEffectExtension * Ext = ExtClass->GetDefaultObject<UGameplayEffectExtension>();
-				Ext->PostGameplayEffectExecute(*this, Data);
+				for (int32 ModIdx = 0; ModIdx < Spec.Modifiers.Num(); ++ModIdx)
+				{
+					const FModifierSpec &Mod = Spec.Modifiers[ModIdx];
+					
+					ABILITY_LOG_SCOPE(TEXT("Applying Attribute Mod %s to property"), *Mod.ToSimpleString());
+
+					FAggregator* Aggregator = OwningContainer.FindOrCreateAttributeAggregator(Mod.Info.Attribute).Get();
+
+					// GE_FIXME: Figure this out - when does EvaluatedData get calculated (here? Before and stored off?)
+					float EvaluatedData = Spec.GetModifierMagnitude(Mod);
+
+					Aggregator->AddMod(EvaluatedData, Mod.Info.ModifierOp, &Mod.Info.SourceTags, &Mod.Info.TargetTags, Handle);
+				}
 			}
 		}
+
+		IsInhibited = ShouldBeInhibited;
 	}
 }
 
-// --------------------------------------------------------------------------------------------------------------------------------------------------------
-//
-//	FActiveGameplayEffect
-//
-// --------------------------------------------------------------------------------------------------------------------------------------------------------
-
 bool FActiveGameplayEffect::CanBeStacked(const FActiveGameplayEffect& Other) const
 {
 	return (Handle != Other.Handle) && Spec.GetStackingType() == Other.Spec.GetStackingType() && Spec.Def->Modifiers[0].Attribute == Other.Spec.Def->Modifiers[0].Attribute;
@@ -1198,80 +519,20 @@ void FActiveGameplayEffect::PostReplicatedChange(const struct FActiveGameplayEff
 //
 // --------------------------------------------------------------------------------------------------------------------------------------------------------
 
-bool FActiveGameplayEffectsContainer::ApplyActiveEffectsTo(OUT FGameplayEffectSpec &Spec, const FModifierQualifier &QualifierContext) const
-{
-	ABILITY_LOG_SCOPE(TEXT("ApplyActiveEffectsTo: %s %s"), *Spec.ToSimpleString(), *QualifierContext.ToString());
-	for (const FActiveGameplayEffect& ActiveEffect : GameplayEffects)
-	{
-		// We dont want to use FModifierQualifier::TestTarget here, since we aren't the 'target'. We are applying stuff to Spec which will be applied to a target.
-		if (QualifierContext.IgnoreHandle().IsValid() && QualifierContext.IgnoreHandle() == ActiveEffect.Handle)
-		{
-			continue;
-		}
-		// should we try to apply ActiveEffect to Spec?
-		float ChanceToExecute = ActiveEffect.Spec.GetChanceToExecuteOnGameplayEffect();
-		if ((ChanceToExecute < 1.f - SMALL_NUMBER) && (FMath::FRand() > ChanceToExecute))
-		{
-			continue;
-		}
-		if (Spec.ApplyModifiersFrom(ActiveEffect.Spec, QualifierContext) < 0)
-		{
-			// ActiveEffect provides immunity to Spec. No need to look at other active effects
-			return false;
-		}
-	}
-
-	return true;
-}
-
-/** This is the main function that applies/attaches a GameplayEffect on Attributes and ActiveGameplayEffects */
-void FActiveGameplayEffectsContainer::ApplySpecToActiveEffectsAndAttributes(FGameplayEffectSpec &Spec, const FModifierQualifier &QualifierContext)
+void FActiveGameplayEffectsContainer::RegisterWithOwner(UAbilitySystemComponent* InOwner)
 {
-	for (const FModifierSpec &Mod : Spec.Modifiers)
+	if (Owner != InOwner)
 	{
-		if (Mod.Info.ModifierType == EGameplayMod::Attribute)
-		{
-			// Todo: Tag/application checks here
-
-			ABILITY_LOG_SCOPE(TEXT("Applying Attribute Mod %s to property"), *Mod.ToSimpleString());
-
-			FAggregator * Aggregator = FindOrCreateAttributeAggregator(Mod.Info.Attribute).Get();
-
-			// Add the modifier to the property value aggregator
-			// Note that this will immediately invoke the callback to update the attribute's current value, so we don't have to explicitly do it here.
-			Aggregator->ApplyMod(Mod.Info.ModifierOp, Mod.Aggregator, Spec.ShouldApplyAsSnapshot(QualifierContext));
-		}
+		Owner = InOwner;
 
-		if (Mod.Info.ModifierType == EGameplayMod::ActiveGE)
-		{
-			ABILITY_LOG_SCOPE(TEXT("Applying Mod %s to ActiveEffects"),  *Mod.ToSimpleString());
-
-			bool bUpdateStackedEffects = false;
-
-			// TODO: Tag checks here
-
-			// This modifies GEs that are currently active, so apply this to them...
-			for (FActiveGameplayEffect & ActiveEffect : GameplayEffects)
-			{
-				if (!QualifierContext.TestTarget(ActiveEffect.Handle))
-				{
-					continue;
-				}
-
-				ActiveEffect.Spec.ApplyModifiersFrom( Spec, FModifierQualifier().Type(EGameplayMod::ActiveGE) );
-				bUpdateStackedEffects = true;
-			}
-
-			if (bUpdateStackedEffects)
-			{
-				StacksNeedToRecalculate();
-			}
-		}
+		// Binding raw is ok here, since the owner is literally the UObject that owns us. If we are destroyed, its because that uobject is destroyed,
+		// and if that is destroyed, the delegate wont be able to fire.
+		Owner->RegisterGenericGameplayTagEvent().AddRaw(this, &FActiveGameplayEffectsContainer::OnOwnerTagChange);
 	}
 }
 
 /** This is the main function that executes a GameplayEffect on Attributes and ActiveGameplayEffects */
-void FActiveGameplayEffectsContainer::ExecuteActiveEffectsFrom(FGameplayEffectSpec &Spec, const FModifierQualifier &QualifierContext)
+void FActiveGameplayEffectsContainer::ExecuteActiveEffectsFrom(FGameplayEffectSpec &Spec, FPredictionKey PredictionKey)
 {
 	// If there are no modifiers, we always want to apply the GameplayCue. If there are modifiers, we only want to invoke the GameplayCue if one of them went through (coudl be blocked by immunity or % chance roll)
 	bool InvokeGameplayCueExecute = (Spec.Modifiers.Num() == 0);
@@ -1286,139 +547,141 @@ void FActiveGameplayEffectsContainer::ExecuteActiveEffectsFrom(FGameplayEffectSp
 		}
 	}
 
-	float ChanceToExecute = Spec.GetChanceToExecuteOnGameplayEffect();
-
-	// check if the new effect removes or modifies existing effects
-	if (Spec.Def->AppliesImmunityTo == EGameplayImmunity::ActiveGE)
-	{
-		for (int32 ii = 0; ii < GameplayEffects.Num(); ++ii)
-		{
-			// should we try to apply this to the current effect?
-			if ((ChanceToExecute < 1.f - SMALL_NUMBER) && (FMath::FRand() > ChanceToExecute))
-			{
-				continue;
-			}
+	float ChanceToExecute = Spec.GetChanceToExecuteOnGameplayEffect();		// Not implemented? Should we just remove ChanceToExecute?
 
-			if (Spec.Def->AreGameplayEffectTagRequirementsSatisfied(GameplayEffects[ii].Spec.Def))
-			{
-				GameplayEffects.RemoveAtSwap(ii);
-				--ii;
-				continue;
-			}
+	// Capture our own tags. Should this ever be a setting or are Target tags always up to date? Seems simplier to keep them always up to date.
+	Spec.CapturedSourceTags.RemoveAllTags();
+	Owner->GetOwnedGameplayTags(Spec.CapturedSourceTags);
 
-			// we kept the effect, now check its mods to see if we should remove any of them
-			for (int32 jj = 0; jj < GameplayEffects[ii].Spec.Modifiers.Num(); ++jj)
-			{
-				const FGameplayModifierEvaluatedData& Data = GameplayEffects[ii].Spec.Modifiers[jj].Aggregator.Get()->Evaluate();
-				if (GameplayEffects[ii].Spec.Def->AreGameplayEffectTagRequirementsSatisfied(Data.Tags))
-				{
-					GameplayEffects[ii].Spec.Modifiers.RemoveAtSwap(jj);
-					--jj;
-				}
-			}
-		}
-	}
 
-	bool bModifiesActiveGEs = false;
+	// ------------------------------------------------------
+	//	Modifiers
+	//		These will modify the base value of attributes
+	// ------------------------------------------------------
 
 	for (const FModifierSpec &Mod : Spec.Modifiers)
 	{
-		if (Mod.Info.ModifierType == EGameplayMod::Attribute)
+		UAttributeSet* AttributeSet = const_cast<UAttributeSet*>(Owner->GetAttributeSubobject(Mod.Info.Attribute.GetAttributeSetClass()));
+		if (AttributeSet == NULL)
 		{
-			UAttributeSet* AttributeSet = const_cast<UAttributeSet*>(Owner->GetAttributeSubobject(Mod.Info.Attribute.GetAttributeSetClass()));
-			if (AttributeSet == NULL)
-			{
-				// Our owner doesn't have this attribute, so we can't do anything
-				ABILITY_LOG(Log, TEXT("%s does not have attribute %s. Skipping modifier"), *Owner->GetPathName(), *Mod.Info.Attribute.GetName());
-				continue;
-			}
-
-			ABILITY_LOG_SCOPE(TEXT("Executing Attribute Mod %s"), *Mod.ToSimpleString());
-
-			// First, evaluate all of our data 
-
-			FGameplayModifierEvaluatedData EvaluatedData = Mod.Aggregator.Get()->Evaluate();
+			// Our owner doesn't have this attribute, so we can't do anything
+			ABILITY_LOG(Log, TEXT("%s does not have attribute %s. Skipping modifier"), *Owner->GetPathName(), *Mod.Info.Attribute.GetName());
+			continue;
+		}
 
-			FGameplayEffectModCallbackData ExecuteData(Spec, Mod, EvaluatedData, *Owner);
-						
-			/** This should apply 'gameplay effect specific' rules, such as life steal, shields, etc */
-			Mod.Aggregator.Get()->PreEvaluate(ExecuteData);
+		ABILITY_LOG_SCOPE(TEXT("Executing Attribute Mod %s"), *Mod.ToSimpleString());
 
-			/** This should apply 'gamewide' rules. Such as clamping Health to MaxHealth or granting +3 health for every point of strength, etc 
-			 *	PreAttributeModify can return false to 'throw out' this modification.
-			 */
-			if (AttributeSet->PreGameplayEffectExecute(ExecuteData) == false)
-			{
-				continue;
-			}
+		// First, evaluate all of our data
+			
+		// GE_Remove: still fixing this up:
 
-			// Todo: Tags/application checks here - make sure we can still apply
-			InvokeGameplayCueExecute = true;
+		float EvaluatedMagnitude = Mod.Info.Magnitude.GetValueAtLevel(Spec.GetLevel());
 
-			// Do we have active GE's that are already modifying this?
-			FAggregatorRef *RefPtr = OngoingPropertyEffects.Find(Mod.Info.Attribute);
-			if (RefPtr)
-			{
-				ABILITY_LOG(Log, TEXT("Property %s has active mods. Adding to Aggregator."), *Mod.Info.Attribute.GetName());
-				RefPtr->Get()->ExecuteMod(Mod.Info.ModifierOp, EvaluatedData);
-			}
-			else
-			{
-				// Modify the property inplace, without putting it in the OngoingPropertyEffects map
-				float		CurrentValueOfProperty = Owner->GetNumericAttribute(Mod.Info.Attribute);
-				FAggregator	Aggregator(CurrentValueOfProperty, SKILL_AGG_DEBUG(TEXT("Inplace Attribute %s"), *Mod.Info.Attribute.GetName()));
+		FGameplayModifierEvaluatedData EvaluatedData(EvaluatedMagnitude);
 
-				Aggregator.ExecuteMod(Mod.Info.ModifierOp, EvaluatedData);
-				
-				const float NewPropertyValue = Aggregator.Evaluate().Magnitude;
+		FGameplayEffectModCallbackData ExecuteData(Spec, Mod, EvaluatedData, *Owner);
 
-				InternalUpdateNumericalAttribute(Mod.Info.Attribute, NewPropertyValue, &ExecuteData);
-			}
+		/** 
+		 *  This should apply 'gamewide' rules. Such as clamping Health to MaxHealth or granting +3 health for every point of strength, etc 
+		 *	PreAttributeModify can return false to 'throw out' this modification.
+		 */
+		if (AttributeSet->PreGameplayEffectExecute(ExecuteData) == false)
+		{
+			continue;
+		}
 
-			FGameplayEffectModifiedAttribute* ModifiedAttribute = Spec.GetModifiedAttribute(Mod.Info.Attribute);
-			if (!ModifiedAttribute)
-			{
-				// If we haven't already created a modified attribute holder, create it
-				ModifiedAttribute = Spec.AddModifiedAttribute(Mod.Info.Attribute);
-			}
-			ModifiedAttribute->TotalMagnitude += EvaluatedData.Magnitude;
+		// Todo: Tags/application checks here - make sure we can still apply (GE_Remove ??? is this still valid?)
+		InvokeGameplayCueExecute = true;
 
-			/** This should apply 'gameplay effect specific' rules, such as life steal, shields, etc */
-			Mod.Aggregator.Get()->PostEvaluate(ExecuteData);
+		// Do we have active GE's that are already modifying this?
+		FAggregatorRef* RefPtr = AttributeAggregatorMap.Find(Mod.Info.Attribute);
+		if (RefPtr)
+		{
+			ABILITY_LOG(Log, TEXT("Property %s has active mods. Adding to Aggregator."), *Mod.Info.Attribute.GetName());
+			RefPtr->Get()->ExecModOnBaseValue(Mod.Info.ModifierOp, Mod.GetEvaluatedMagnitude() );
+		}
+		else
+		{
+			// Modify the property inplace, without putting it in the AttributeAggregatorMap map
+			float		CurrentValueOfProperty = Owner->GetNumericAttribute(Mod.Info.Attribute);
+			const float NewPropertyValue = FAggregator::StaticExecModOnBaseValue(CurrentValueOfProperty, Mod.Info.ModifierOp, Mod.GetEvaluatedMagnitude());
 
-			/** This should apply 'gamewide' rules. Such as clamping Health to MaxHealth or granting +3 health for every point of strength, etc */
-			AttributeSet->PostGameplayEffectExecute(ExecuteData);
+			InternalUpdateNumericalAttribute(Mod.Info.Attribute, NewPropertyValue, &ExecuteData);
 		}
-		else if(Mod.Info.ModifierType == EGameplayMod::ActiveGE)
+
+		FGameplayEffectModifiedAttribute* ModifiedAttribute = Spec.GetModifiedAttribute(Mod.Info.Attribute);
+		if (!ModifiedAttribute)
 		{
-			bModifiesActiveGEs = true;
+			// If we haven't already created a modified attribute holder, create it
+			ModifiedAttribute = Spec.AddModifiedAttribute(Mod.Info.Attribute);
 		}
+		ModifiedAttribute->TotalMagnitude += EvaluatedData.Magnitude;
+
+		/** This should apply 'gamewide' rules. Such as clamping Health to MaxHealth or granting +3 health for every point of strength, etc */
+		AttributeSet->PostGameplayEffectExecute(ExecuteData);
 	}
 
-	// If any of the mods apply to active gameplay effects try to apply them now.
-	// We need to do this here because all of the modifiers either need to apply or not apply to each active gameplay effect based on their chance to execute on gameplay effects.
-	if (bModifiesActiveGEs)
+	// ------------------------------------------------------
+	//	Executions
+	//		This will run custom code to 'do stuff'
+	// ------------------------------------------------------
+	 
+	for (int32 ExecIdx=0; ExecIdx < Spec.Def->Executions.Num(); ++ExecIdx)
 	{
-		for (FActiveGameplayEffect & ActiveEffect : GameplayEffects)
+		const FGameplayEffectExecutionDefinition& Exec = Spec.Def->Executions[ExecIdx];
+		if (Exec.CalculationClass == nullptr)
+			continue;
+
+		const UGameplayEffectCalculation* CDO = Cast<UGameplayEffectCalculation>(Exec.CalculationClass->ClassDefaultObject);
+		check(CDO);
+
+
+		// Apply custom inline mods
+		FActiveGameplayEffectHandle ModifierHandle = FActiveGameplayEffectHandle::GenerateNewHandle(Owner);
+		TSet<FAggregator*>	ModifiedAggregators;
+		for (const FExtensionAttributeModifierInfo& Mod : Exec.CalculationModifiers)
 		{
-			// Don't apply spec to itself
-			if (!QualifierContext.TestTarget(ActiveEffect.Handle))
-			{
-				continue;
-			}
+			// Evaluate this mods current value. Right now this is always a scalable float, but we'd like to support
+			// pulling arbitrary attribute values here too - these would be pulled from the Spec's captured attribute only?
+			// But bottom line - we evaluate them right here and that is their value for this scope. Right?
+			
+			// We do a const cast here becase we are being really careful! We will add Mods below but we will also remove them once the 
+			// custom calculation is finished. Custom Exec functions will only get const references to the aggregator.
+			
+			// If the custom exec functions want to add mods inplace, they can do it either like:
+			//	A) like we are doing it here (careful const cast + remove)
+			//  B) We give them a way to clone the aggregators on the stack and allow them to do whatever they want to those copies
+			//  C) We provide another optional parameter into ::Evaluate that takes additional mods.
+			 
+			
+			// @todo: Just always assuming snapshot for now, which is wrong. This should use a different struct.
+			FAggregator& Aggregator = const_cast<FAggregator&>(Spec.CapturedRelevantAttributes.GetAggregator(FGameplayEffectAttributeCaptureDefinition(Mod.Attribute, Mod.Source, true)));
 
-			// should we try to apply this to the current effect?
-			if ((ChanceToExecute < 1.f - SMALL_NUMBER) && (FMath::FRand() > ChanceToExecute))
-			{
-				continue;
-			}
 
-			if (ActiveEffect.Spec.ExecuteModifiersFrom(Spec, FModifierQualifier().Type(EGameplayMod::ActiveGE)) > 0)
-			{
-				InvokeGameplayCueExecute = true;
-			}
+			float EvaluatedMagnitude = Mod.Magnitude.GetValueAtLevel(Spec.GetLevel());
+
+			Aggregator.DisableCallbacks();
+			Aggregator.AddMod(EvaluatedMagnitude, Mod.ModifierOp, &Mod.SourceTags, &Mod.TargetTags, ModifierHandle);
+
+			ModifiedAggregators.Add(&Aggregator);
+		}
+
+		CDO->Execute(Spec, ExecIdx, Owner);
+
+		// Remove those custom inline mods
+		for (FAggregator* Aggregator : ModifiedAggregators)
+		{
+			Aggregator->RemoveMod(ModifierHandle);
+			Aggregator->EnabledCallbacks();
 		}
+
 	}
+
+
+	// ------------------------------------------------------
+	//	Invoke GameplayCue events
+	// ------------------------------------------------------
+
 	
 	// Prune the modified attributes before we replicate
 	Spec.PruneModifiedAttributes();
@@ -1428,7 +691,7 @@ void FActiveGameplayEffectsContainer::ExecuteActiveEffectsFrom(FGameplayEffectSp
 		// TODO: check replication policy. Right now we will replicate every execute via a multicast RPC
 
 		ABILITY_LOG(Log, TEXT("Invoking Execute GameplayCue for %s"), *Spec.ToSimpleString() );
-		Owner->NetMulticast_InvokeGameplayCueExecuted_FromSpec(Spec, QualifierContext.PredictionKey());
+		Owner->NetMulticast_InvokeGameplayCueExecuted_FromSpec(Spec, PredictionKey);
 	}
 }
 
@@ -1438,7 +701,7 @@ void FActiveGameplayEffectsContainer::ExecutePeriodicGameplayEffect(FActiveGamep
 	if (ActiveEffect)
 	{
 		// Execute
-		ExecuteActiveEffectsFrom(ActiveEffect->Spec, FModifierQualifier().IgnoreHandle(Handle));
+		ExecuteActiveEffectsFrom(ActiveEffect->Spec);
 	}
 }
 
@@ -1455,15 +718,9 @@ FActiveGameplayEffect* FActiveGameplayEffectsContainer::GetActiveGameplayEffect(
 	return nullptr;
 }
 
-void FActiveGameplayEffectsContainer::AddDependancyToAttribute(FGameplayAttribute Attribute, const TWeakPtr<FAggregator> InDependant)
-{
-	FAggregator * Aggregator = FindOrCreateAttributeAggregator(Attribute).Get();
-	Aggregator->AddDependantAggregator(InDependant);
-}
-
-FAggregatorRef & FActiveGameplayEffectsContainer::FindOrCreateAttributeAggregator(FGameplayAttribute Attribute)
+FAggregatorRef& FActiveGameplayEffectsContainer::FindOrCreateAttributeAggregator(FGameplayAttribute Attribute)
 {
-	FAggregatorRef *RefPtr = OngoingPropertyEffects.Find(Attribute);
+	FAggregatorRef *RefPtr = AttributeAggregatorMap.Find(Attribute);
 	if (RefPtr)
 	{
 		return *RefPtr;
@@ -1471,24 +728,46 @@ FAggregatorRef & FActiveGameplayEffectsContainer::FindOrCreateAttributeAggregato
 
 	// Create a new aggregator for this attribute.
 	float CurrentValueOfProperty = Owner->GetNumericAttribute(Attribute);
-	ABILITY_LOG(Log, TEXT("Creating new entry in OngoingPropertyEffect map for AddDependancyToAttribute. CurrentValue: %.2f"), CurrentValueOfProperty);
+	ABILITY_LOG(Log, TEXT("Creating new entry in AttributeAggregatorMap for %s. CurrentValue: %.2f"), *Attribute.GetName(), CurrentValueOfProperty);
 
-	FAggregator *NewPropertyAggregator = new FAggregator(FGameplayModifierEvaluatedData(CurrentValueOfProperty), SKILL_AGG_DEBUG(TEXT("Attribute %s Aggregator"), *Attribute.GetName()));
+	FAggregator* NewAttributeAggregator = new FAggregator(CurrentValueOfProperty);
 
-	NewPropertyAggregator->OnDirty = FAggregator::FOnDirty::CreateRaw(this, &FActiveGameplayEffectsContainer::OnPropertyAggregatorDirty, Attribute);
+	// FIXME: AddRaw is bad
+	NewAttributeAggregator->OnDirty.AddUObject(Owner, &UAbilitySystemComponent::OnAttributeAggregatorDirty, Attribute);
 
-	return OngoingPropertyEffects.Add(Attribute, FAggregatorRef(NewPropertyAggregator));
+	return AttributeAggregatorMap.Add(Attribute, FAggregatorRef(NewAttributeAggregator));
+}
+
+// Move this to UAttributeComponent?
+void FActiveGameplayEffectsContainer::OnAttributeAggregatorDirty(FAggregator* Aggregator, FGameplayAttribute Attribute)
+{
+	ABILITY_LOG_SCOPE(TEXT("FActiveGameplayEffectsContainer::OnAttributeAggregatorDirty"));
+	check(AttributeAggregatorMap.FindChecked(Attribute).Get() == Aggregator);
+
+	// Our Aggregator has changed, we need to reevaluate this aggregator and update the current value of the attribute.
+	// Note that this is not an execution, so there are no 'source' and 'target' tags to fill out in the FAggregatorEvaluateParameters.
+	// ActiveGameplayEffects that have required owned tags will be turned on/off via delegates, and will add/remove themselves from attribute
+	// aggregators when that happens.
+	
+	FAggregatorEvaluateParameters EvaluationParameters;
+	
+	float NewValue = Aggregator->Evaluate(EvaluationParameters);
+	InternalUpdateNumericalAttribute(Attribute, NewValue, nullptr);
 }
 
 void FActiveGameplayEffectsContainer::SetBaseAttributeValueFromReplication(FGameplayAttribute Attribute, float BaseValue)
 {
-	FAggregatorRef* RefPtr = OngoingPropertyEffects.Find(Attribute);
-	if (RefPtr && RefPtr->IsValid())
+/*
+	FIXME: the approach for client side attribute prediction will need to be slightly rethought. Rather than touching base values we should be able
+	to just apply Mods... this function may just go away?
+
+	FAggregatorRef* RefPtr = AttributeAggregatorMap.Find(Attribute);
+	if (RefPtr && RefPtr->Get())
 	{
-		check(RefPtr->Get()->BaseData.Magnitude.IsStatic());
-		RefPtr->Get()->BaseData.Magnitude.SetValue(BaseValue);
+		RefPtr->Get()->SetBaseValue(Attribute);
 		RefPtr->Get()->MarkDirty();
 	}
+*/
 }
 
 float FActiveGameplayEffectsContainer::GetGameplayEffectDuration(FActiveGameplayEffectHandle Handle) const
@@ -1506,23 +785,6 @@ float FActiveGameplayEffectsContainer::GetGameplayEffectDuration(FActiveGameplay
 	return UGameplayEffect::INFINITE_DURATION;
 }
 
-void FActiveGameplayEffectsContainer::OnDurationAggregatorDirty(FAggregator* Aggregator, UAbilitySystemComponent* InOwner, FActiveGameplayEffectHandle Handle)
-{
-	for (FActiveGameplayEffect& Effect : GameplayEffects)
-	{
-		if (Effect.Handle == Handle)
-		{
-			FTimerManager& TimerManager = InOwner->GetWorld()->GetTimerManager();
-			float CurrDuration = Effect.GetDuration();
-			float TimeRemaining = CurrDuration - TimerManager.GetTimerElapsed(Effect.DurationHandle);
-			FTimerDelegate Delegate = FTimerDelegate::CreateUObject(InOwner, &UAbilitySystemComponent::CheckDurationExpired, Handle);
-			TimerManager.SetTimer(Effect.DurationHandle, Delegate, CurrDuration, false, FMath::Max(TimeRemaining, 0.f));
-
-			return;
-		}
-	}
-}
-
 float FActiveGameplayEffectsContainer::GetGameplayEffectMagnitude(FActiveGameplayEffectHandle Handle, FGameplayAttribute Attribute) const
 {
 	// Could make this a map for quicker lookup
@@ -1534,7 +796,7 @@ float FActiveGameplayEffectsContainer::GetGameplayEffectMagnitude(FActiveGamepla
 			{
 				if (Mod.Info.Attribute == Attribute)
 				{
-					return Mod.Aggregator.Get()->Evaluate().Magnitude;
+					return Effect.Spec.GetModifierMagnitude(Mod);
 				}
 			}
 		}
@@ -1564,39 +826,18 @@ bool FActiveGameplayEffectsContainer::IsGameplayEffectActive(FActiveGameplayEffe
 	return false;
 }
 
-float FActiveGameplayEffectsContainer::GetGameplayEffectMagnitudeByTag(FActiveGameplayEffectHandle Handle, const FGameplayTag& InTag) const
+void FActiveGameplayEffectsContainer::CaptureAttributeForGameplayEffect(OUT FGameplayEffectAttributeCaptureSpec& OutCaptureSpec)
 {
-	// Could make this a map for quicker lookup
-	for (FActiveGameplayEffect Effect : GameplayEffects)
+	FAggregatorRef& AttributeAggregator = FindOrCreateAttributeAggregator(OutCaptureSpec.BackingDefinition.AttributeToCapture);
+	
+	if (OutCaptureSpec.BackingDefinition.bSnapshot)
 	{
-		if (Effect.Handle == Handle)
-		{
-			for (const FModifierSpec &Mod : Effect.Spec.Modifiers)
-			{
-				if (Mod.Info.OwnedTags.HasTag(InTag, EGameplayTagMatchType::IncludeParentTags, EGameplayTagMatchType::Explicit))
-				{
-					return Mod.Aggregator.Get()->Evaluate().Magnitude;
-				}
-			}
-		}
+		OutCaptureSpec.AttributeAggregator.TakeSnapshotOf(AttributeAggregator);
+	}
+	else
+	{
+		OutCaptureSpec.AttributeAggregator = AttributeAggregator;
 	}
-
-	ABILITY_LOG(Warning, TEXT("GetGameplayEffectMagnitude called with invalid Handle: %s"), *Handle.ToString());
-	return -1.f;
-}
-
-void FActiveGameplayEffectsContainer::OnPropertyAggregatorDirty(FAggregator* Aggregator, FGameplayAttribute Attribute)
-{
-	check(OngoingPropertyEffects.FindChecked(Attribute).Get() == Aggregator);
-
-	ABILITY_LOG_SCOPE(TEXT("FActiveGameplayEffectsContainer::OnPropertyAggregatorDirty"));
-
-	// Immediately calculate the newest value of the property			
-	float NewValue = Aggregator->Evaluate().Magnitude;
-
-	InternalUpdateNumericalAttribute(Attribute, NewValue, nullptr);
-
-	
 }
 
 void FActiveGameplayEffectsContainer::InternalUpdateNumericalAttribute(FGameplayAttribute Attribute, float NewValue, const FGameplayEffectModCallbackData* ModData)
@@ -1628,7 +869,7 @@ FActiveGameplayEffect& FActiveGameplayEffectsContainer::CreateNewActiveGameplayE
 	SCOPE_CYCLE_COUNTER(STAT_CreateNewActiveGameplayEffect);
 
 	FActiveGameplayEffectHandle NewHandle = FActiveGameplayEffectHandle::GenerateNewHandle(Owner);
-	FActiveGameplayEffect & NewEffect = *new (GameplayEffects)FActiveGameplayEffect(NewHandle, Spec, GetWorldTime(), GetGameStateTime(), InPredictionKey);
+	FActiveGameplayEffect& NewEffect = *new (GameplayEffects)FActiveGameplayEffect(NewHandle, Spec, GetWorldTime(), GetGameStateTime(), InPredictionKey);
 
 	// register callbacks with the timer manager
 	if (Owner)
@@ -1668,8 +909,6 @@ FActiveGameplayEffect& FActiveGameplayEffectsContainer::CreateNewActiveGameplayE
 
 			TimerManager.SetTimer(NewEffect.PeriodHandle, Delegate, Spec.GetPeriod(), true, FirstDelay); // this is going to be off by a frame for stacking because of the pending list
 		}
-
-		NewEffect.Spec.Duration.Get()->OnDirty = FAggregator::FOnDirty::CreateRaw(this, &FActiveGameplayEffectsContainer::OnDurationAggregatorDirty, Owner, NewHandle);
 	}
 	
 	if (InPredictionKey.IsValidKey() == false || IsNetAuthority())	// Clients predicting a GameplayEffect must not call MarkItemDirty
@@ -1686,24 +925,43 @@ FActiveGameplayEffect& FActiveGameplayEffectsContainer::CreateNewActiveGameplayE
 		
 	}
 
-	InternalOnActiveGameplayEffectAdded(NewEffect);	
+	InternalOnActiveGameplayEffectAdded(NewEffect);
 
 	return NewEffect;
 }
 
-void FActiveGameplayEffectsContainer::InternalOnActiveGameplayEffectAdded(const FActiveGameplayEffect& Effect)
+void FActiveGameplayEffectsContainer::InternalOnActiveGameplayEffectAdded(FActiveGameplayEffect& Effect)
 {
-	// Update gameplaytag count and broadcast delegate if we just added this tag (count=0, prior to increment)
-	Owner->UpdateTagMap(Effect.Spec.Def->OwnedTagsContainer, 1);
+	// Add our ongoing tag requirements to the dependency map. We will actually check for these tags below.
+	for (const FGameplayTag& Tag : Effect.Spec.Def->OngoingTagRequirements.IgnoreTags)
+	{
+		ActiveEffectTagDependencies.FindOrAdd(Tag).Add(Effect.Handle);
+	}
 
-	//Owner.ActiveGameplayCues.AddCue(
+	for (const FGameplayTag& Tag : Effect.Spec.Def->OngoingTagRequirements.RequireTags)
+	{
+		ActiveEffectTagDependencies.FindOrAdd(Tag).Add(Effect.Handle);
+	}
 
+	// Update our owner with the tags this GameplayEffect grants them
+	Owner->UpdateTagMap(Effect.Spec.Def->OwnedTagsContainer, 1);
 	for (const FGameplayEffectCue& Cue : Effect.Spec.Def->GameplayCues)
 	{
 		Owner->UpdateTagMap(Cue.GameplayCueTags, 1);
-	}	
-}
+	}
+
+	// Check if we should actually be turned on or not (this will turn us on for the first time)
+	FGameplayTagContainer OwnerTags;
+	Owner->GetOwnedGameplayTags(OwnerTags);
 
+	// Calc all of our modifier magnitudes now. Some may need to update later based on attributes changing, etc, but those should
+	// be done through delegate callbacks.
+	Effect.Spec.CalculateModifierMagnitudes();
+
+	// @todo: Probably do the initial mod spec magnitude calculations right here, but only on the server
+	// @todo: Shouldn't the rest of the stuff about dependencies only run on the server?
+	Effect.CheckOngoingTagRequirements(OwnerTags, *this);
+}
 
 bool FActiveGameplayEffectsContainer::RemoveActiveGameplayEffect(FActiveGameplayEffectHandle Handle)
 {
@@ -1754,6 +1012,49 @@ bool FActiveGameplayEffectsContainer::InternalRemoveActiveGameplayEffect(int32 I
 			Owner->GetWorld()->GetTimerManager().ClearTimer(GameplayEffects[Idx].PeriodHandle);
 		}
 
+		// Remove our tag requirements from the dependancy map
+		for (const FGameplayTag& Tag : Effect.Spec.Def->OngoingTagRequirements.IgnoreTags)
+		{
+			auto Ptr = ActiveEffectTagDependencies.Find(Tag);
+			if (Ptr)
+			{
+				Ptr->Remove(Effect.Handle);
+				if (Ptr->Num() <= 0)
+				{
+					ActiveEffectTagDependencies.Remove(Tag);
+				}
+			}
+		}
+
+		for (const FGameplayTag& Tag : Effect.Spec.Def->OngoingTagRequirements.RequireTags)
+		{
+			auto Ptr = ActiveEffectTagDependencies.Find(Tag);
+			if (Ptr)
+			{
+				Ptr->Remove(Effect.Handle);
+				if (Ptr->Num() <= 0)
+				{
+					ActiveEffectTagDependencies.Remove(Tag);
+				}
+			}
+		}
+
+
+		// Fixme: Yuck?
+		for (const FGameplayModifierInfo& Mod : Effect.Spec.Def->Modifiers)
+		{
+			if (Mod.Attribute.IsValid())
+			{
+				FAggregatorRef* RefPtr = AttributeAggregatorMap.Find(Mod.Attribute);
+				if (RefPtr)
+				{
+					RefPtr->Get()->RemoveMod(Effect.Handle);
+				}
+			}
+		}
+
+
+
 		GameplayEffects.RemoveAtSwap(Idx);
 		MarkArrayDirty();
 
@@ -1793,6 +1094,28 @@ void FActiveGameplayEffectsContainer::InternalOnActiveGameplayEffectRemoved(cons
 	}
 }
 
+void FActiveGameplayEffectsContainer::OnOwnerTagChange(FGameplayTag TagChange, int32 NewCount)
+{
+	// It may be beneficial to do a scoped lock on attribute reeevalluation during this dunction
+
+	auto Ptr = ActiveEffectTagDependencies.Find(TagChange);
+	if (Ptr)
+	{
+		FGameplayTagContainer OwnerTags;
+		Owner->GetOwnedGameplayTags(OwnerTags);
+
+		TSet<FActiveGameplayEffectHandle>& Handles = *Ptr;
+		for (const FActiveGameplayEffectHandle& Handle : Handles)
+		{
+			FActiveGameplayEffect* ActiveEffect = GetActiveGameplayEffect(Handle);
+			if (ActiveEffect)
+			{
+				ActiveEffect->CheckOngoingTagRequirements(OwnerTags, *this);
+			}
+		}
+	}
+}
+
 bool FActiveGameplayEffectsContainer::IsNetAuthority() const
 {
 	check(Owner);
@@ -1801,14 +1124,7 @@ bool FActiveGameplayEffectsContainer::IsNetAuthority() const
 
 void FActiveGameplayEffectsContainer::PreDestroy()
 {
-	// Prior to destruction we need to clear all dependancies
-	for (auto It : OngoingPropertyEffects)
-	{
-		if (It.Value.IsValid())
-		{
-			It.Value.Get()->ClearAllDependancies();
-		}
-	}
+	
 }
 
 int32 FActiveGameplayEffectsContainer::GetGameStateTime() const
@@ -1850,7 +1166,7 @@ void FActiveGameplayEffectsContainer::CheckDuration(FActiveGameplayEffectHandle
 					float PeriodTimeRemaining = TimerManager.GetTimerRemaining(Effect.PeriodHandle);
 					if (FMath::IsNearlyEqual(PeriodTimeRemaining, 0.f))
 					{
-						ExecuteActiveEffectsFrom(Effect.Spec, FModifierQualifier().IgnoreHandle(Effect.Handle));
+						ExecuteActiveEffectsFrom(Effect.Spec);
 					}
 				}
 
@@ -1983,9 +1299,7 @@ bool FActiveGameplayEffectsContainer::CanApplyAttributeModifiers(const UGameplay
 {
 	SCOPE_CYCLE_COUNTER(STAT_GameplayEffectsCanApplyAttributeModifiers);
 
-	FGameplayEffectSpec	Spec(GameplayEffect, EffectContext, Level, Owner->GetCurveDataOverride());
-
-	ApplyActiveEffectsTo(Spec, FModifierQualifier().Type(EGameplayMod::IncomingGE));
+	FGameplayEffectSpec	Spec(GameplayEffect, EffectContext, Level);
 	
 	for (const FModifierSpec& Mod : Spec.Modifiers)
 	{
@@ -1994,7 +1308,7 @@ bool FActiveGameplayEffectsContainer::CanApplyAttributeModifiers(const UGameplay
 		{
 			const UAttributeSet* Set = Owner->GetAttributeSubobject(Mod.Info.Attribute.GetAttributeSetClass());
 			float CurrentValue = Mod.Info.Attribute.GetNumericValueChecked(Set);
-			float CostValue = Mod.Aggregator.Get()->Evaluate().Magnitude;
+			float CostValue = Spec.GetModifierMagnitude(Mod);
 
 			if (CurrentValue + CostValue < 0.f)
 			{
@@ -2094,55 +1408,6 @@ bool FActiveGameplayEffectsContainer::HasPredictedEffectWithPredictedKey(FPredic
 	return false;
 }
 
-// --------------------------------------------------------------------------------------------------------------------------------------------------------
-//
-//	FGameplayEffectLevelSpec
-//
-// --------------------------------------------------------------------------------------------------------------------------------------------------------
-
-float FGameplayEffectLevelSpec::GetLevel() const
-{
-	if (ConstantLevel != INVALID_LEVEL)
-	{
-		return ConstantLevel;
-	}
-
-	if (Source.IsValid() && Attribute.GetUProperty() != NULL)
-	{
-		UClass * SetClass = Attribute.GetAttributeSetClass();
-		check(SetClass);
-
-		UAbilitySystemComponent* Component = UAbilitySystemGlobals::GetAbilitySystemComponentFromActor(Source.Get());
-		if (Component)
-		{
-			for (const UAttributeSet* Set : Component->SpawnedAttributes)
-			{
-				if (Set->GetClass()->IsChildOf(SetClass))
-				{
-					CachedLevel = Attribute.GetNumericValueChecked(CastChecked<UAttributeSet>(Set));
-					return CachedLevel;
-				}
-			}
-		}
-	}
-
-	// Our source is invalid. 
-	ConstantLevel = CachedLevel;
-	return CachedLevel;
-}
-
-void FGameplayEffectLevelSpec::RegisterLevelDependancy(TWeakPtr<FAggregator> OwningAggregator)
-{
-	if (Source.IsValid() && Attribute.GetUProperty() != NULL)
-	{
-		UAbilitySystemComponent* SourceComponent = UAbilitySystemGlobals::GetAbilitySystemComponentFromActor(Source.Get());
-		if (SourceComponent)
-		{
-			SourceComponent->AddDependancyToAttribute(Attribute, OwningAggregator);
-		}
-	}
-}
-
 // --------------------------------------------------------------------------------------------------------------------------------------------------------
 //
 //	Misc
@@ -2193,7 +1458,7 @@ bool FActiveGameplayEffectQuery::Matches(const FActiveGameplayEffect& Effect) co
 		bool FailedModifyingAttributeCheck = true;
 		for (const FModifierSpec& Modifier : Effect.Spec.Modifiers)
 		{
-			if (Modifier.Info.Attribute == ModifyingAttribute && Modifier.Info.ModifierType == EGameplayMod::Attribute)
+			if (Modifier.Info.Attribute == ModifyingAttribute)
 			{
 				FailedModifyingAttributeCheck = false;
 				break;
diff --git a/Engine/Source/Runtime/GameplayAbilities/Private/GameplayEffectAggregator.cpp b/Engine/Source/Runtime/GameplayAbilities/Private/GameplayEffectAggregator.cpp
--- a/Engine/Source/Runtime/GameplayAbilities/Private/GameplayEffectAggregator.cpp
+++ b/Engine/Source/Runtime/GameplayAbilities/Private/GameplayEffectAggregator.cpp
@@ -3,9 +3,6 @@
 #include "AbilitySystemPrivatePCH.h"
 #include "GameplayEffectAggregator.h"
 
-#if GE_REFACTOR
-#pragma optimize( "", off )
-
 bool FAggregatorMod::Qualifies(const FAggregatorEvaluateParameters& Parameters) const
 {
 	bool SourceMet = (!SourceTagReqs || SourceTagReqs->RequirementsMet(Parameters.SourceTags));
@@ -15,6 +12,11 @@ bool FAggregatorMod::Qualifies(const FAggregatorEvaluateParameters& Parameters)
 }
 
 float FAggregator::Evaluate(const FAggregatorEvaluateParameters& Parameters) const
+{
+	return EvaluateWithBase(BaseValue, Parameters);
+}
+
+float FAggregator::EvaluateWithBase(float InlineBaseValue, const FAggregatorEvaluateParameters& Parameters) const
 {
 	for (const FAggregatorMod& Mod : Mods[EGameplayModOp::Override])
 	{
@@ -34,20 +36,18 @@ float FAggregator::Evaluate(const FAggregatorEvaluateParameters& Parameters) con
 		Division = 1.f;
 	}
 
-	return ((BaseValue + Additive) * Multiplicitive) / Division;
+	return ((InlineBaseValue + Additive) * Multiplicitive) / Division;
 }
 
 void FAggregator::SetBaseValue(float NewBaseValue)
 {
 	BaseValue = NewBaseValue;
-	OnDirty.Broadcast(this);
+	BroadcastOnDirty();
 }
 
-float FAggregator::StaticExecModOnBaseValue(float BaseValue, const struct FModifierSpec& Mod)
+float FAggregator::StaticExecModOnBaseValue(float BaseValue, TEnumAsByte<EGameplayModOp::Type> ModifierOp, float EvaluatedMagnitude)
 {
-	float EvaluatedMagnitude = Mod.Info.Magnitude.GetValueChecked(); // fixme ?
-
-	switch (Mod.Info.ModifierOp)
+	switch (ModifierOp)
 	{
 	case EGameplayModOp::Override:
 	{
@@ -77,10 +77,10 @@ float FAggregator::StaticExecModOnBaseValue(float BaseValue, const struct FModif
 	return BaseValue;
 }
 
-void FAggregator::ExecModOnBaseValue(const struct FModifierSpec& Mod)
+void FAggregator::ExecModOnBaseValue(TEnumAsByte<EGameplayModOp::Type> ModifierOp, float EvaluatedMagnitude)
 {
-	BaseValue = StaticExecModOnBaseValue(BaseValue, Mod);
-	OnDirty.Broadcast(this);
+	BaseValue = StaticExecModOnBaseValue(BaseValue, ModifierOp, EvaluatedMagnitude);
+	BroadcastOnDirty();
 }
 
 float FAggregator::SumMods(const TArray<FAggregatorMod> &Mods, float Bias, const FAggregatorEvaluateParameters& Parameters) const
@@ -98,22 +98,19 @@ float FAggregator::SumMods(const TArray<FAggregatorMod> &Mods, float Bias, const
 	return Sum;
 }
 
-void FAggregator::AddMod(const FGameplayTagRequirements* SourceTagReqs, const FGameplayTagRequirements* TargetTagReqs, float EvaluatedMagnitude, FActiveGameplayEffectHandle ActiveHandle)
+void FAggregator::AddMod(float EvaluatedMagnitude, TEnumAsByte<EGameplayModOp::Type> ModifierOp, const FGameplayTagRequirements* SourceTagReqs, const FGameplayTagRequirements* TargetTagReqs, FActiveGameplayEffectHandle ActiveHandle)
 {
-	check(Data);
-	check(Data->Modifiers.IsValidIndex(ModIdx));
-
-	TArray<FAggregatorMod> &ModList = Mods[Data->Modifiers[ModIdx].ModifierOp];
+	TArray<FAggregatorMod> &ModList = Mods[ModifierOp];
 
 	int32 NewIdx = ModList.AddUninitialized();
 	FAggregatorMod& NewMod = ModList[NewIdx];
 
-	NewMod.Data = Data;
-	NewMod.ModIdx = ModIdx;
+	NewMod.SourceTagReqs = SourceTagReqs;
+	NewMod.TargetTagReqs = TargetTagReqs;
 	NewMod.EvaluatedMagnitude = EvaluatedMagnitude;
 	NewMod.ActiveHandle = ActiveHandle;
 
-	OnDirty.Broadcast(this );
+	BroadcastOnDirty();
 }
 
 void FAggregator::RemoveMod(FActiveGameplayEffectHandle ActiveHandle)
@@ -126,7 +123,7 @@ void FAggregator::RemoveMod(FActiveGameplayEffectHandle ActiveHandle)
 		RemoveModsWithActiveHandle(Mods[EGameplayModOp::Override], ActiveHandle);
 	}
 
-	OnDirty.Broadcast(this);
+	BroadcastOnDirty();
 }
 
 void FAggregator::RemoveModsWithActiveHandle(TArray<FAggregatorMod>& Mods, FActiveGameplayEffectHandle ActiveHandle)
@@ -151,6 +148,25 @@ void FAggregator::TakeSnapshotOf(const FAggregator& AggToSnapshot)
 	}
 }
 
+void FAggregator::BroadcastOnDirty()
+{
+	if (!CallbacksDisabled)
+	{
+		OnDirty.Broadcast(this);
+	}
+
+}
+
+void FAggregator::DisableCallbacks()
+{
+	CallbacksDisabled= true;
+}
+
+void FAggregator::EnabledCallbacks()
+{
+	CallbacksDisabled = false;
+}
+
 void FAggregatorRef::TakeSnapshotOf(const FAggregatorRef& RefToSnapshot)
 {
 	if (RefToSnapshot.Data.IsValid())
@@ -165,4 +181,3 @@ void FAggregatorRef::TakeSnapshotOf(const FAggregatorRef& RefToSnapshot)
 		Data.Reset();
 	}
 }
-#endif
\ No newline at end of file
diff --git a/Engine/Source/Runtime/GameplayAbilities/Private/GameplayEffectExtension_LifestealTest.cpp b/Engine/Source/Runtime/GameplayAbilities/Private/GameplayEffectExtension_LifestealTest.cpp
--- a/Engine/Source/Runtime/GameplayAbilities/Private/GameplayEffectExtension_LifestealTest.cpp
+++ b/Engine/Source/Runtime/GameplayAbilities/Private/GameplayEffectExtension_LifestealTest.cpp
@@ -39,10 +39,8 @@ void UGameplayEffectExtension_LifestealTest::PostGameplayEffectExecute(const FGa
 			LocalHealthRestore = Cast<UGameplayEffect>(StaticConstructObject(UGameplayEffect::StaticClass(), GetTransientPackage(), FName(TEXT("LifestealHealthRestore"))));
 			LocalHealthRestore->Modifiers.SetNum(1);
 			LocalHealthRestore->Modifiers[0].Magnitude.SetValue(HealthToRestore);
-			LocalHealthRestore->Modifiers[0].ModifierType = EGameplayMod::Attribute;
 			LocalHealthRestore->Modifiers[0].ModifierOp = EGameplayModOp::Additive;
 			LocalHealthRestore->Modifiers[0].Attribute.SetUProperty(HealthProperty);
-			LocalHealthRestore->Modifiers[0].OwnedTags.AddTag(GameplayTagsModule.GetGameplayTagsManager().RequestGameplayTag(FName(TEXT("Lifesteal"))));
 			LocalHealthRestore->Duration.Value = UGameplayEffect::INSTANT_APPLICATION;
 			LocalHealthRestore->Period.Value = UGameplayEffect::NO_PERIOD;
 		}
diff --git a/Engine/Source/Runtime/GameplayAbilities/Private/GameplayEffectExtension_ShieldTest.cpp b/Engine/Source/Runtime/GameplayAbilities/Private/GameplayEffectExtension_ShieldTest.cpp
--- a/Engine/Source/Runtime/GameplayAbilities/Private/GameplayEffectExtension_ShieldTest.cpp
+++ b/Engine/Source/Runtime/GameplayAbilities/Private/GameplayEffectExtension_ShieldTest.cpp
@@ -15,57 +15,7 @@ UGameplayEffectExtension_ShieldTest::UGameplayEffectExtension_ShieldTest(const F
 
 void UGameplayEffectExtension_ShieldTest::PreGameplayEffectExecute(const FGameplayModifierEvaluatedData &SelfData, FGameplayEffectModCallbackData &Data) const
 {
-	IGameplayTagsModule& GameplayTagsModule = IGameplayTagsModule::Get();
-	UAbilitySystemComponent *Source = Data.EffectSpec.EffectContext.GetOriginalInstigatorAbilitySystemComponent();
-
-	// FIXME: some annoyances here: Damage about to be applied = Data.EvaluatedData.Magnitude = negative. Do some sign flipping here that would make more sense if we were dealing with
-	// 'damage' (positive) instead of 'health' (negative)
-
-	// How much damage we will absorb
-	float ShieldedDamage = FMath::Min(-Data.EvaluatedData.Magnitude, SelfData.Magnitude);
-
-	// Decrease the magnitude of damage done
-	Data.EvaluatedData.Magnitude += ShieldedDamage;
-
-	if (ShieldedDamage >= SelfData.Magnitude)
-	{
-		check(SelfData.Handle.IsValid());
-
-		// The shield is now depleted, so remove the gameplay effect altogether
-		bool RemovedSomething = Data.Target.RemoveActiveGameplayEffect(SelfData.Handle);
-		check(RemovedSomething);
-
-		// Note we could get in trouble with the above checks, depending on how the GameplayEffect that called us was setup. 
-		// For example if we were applied to a periodic damage effect, and were copied as a snapshot, our shield could expire
-		// but our aggregator on the DOT gameplayeffect would still exist with a now invalid handle. This is ok, and possibly powerful
-		// but just means that the way an effect like this is written would have to take it into account. This one doesn't.
-		//
-		// It even opens up interesting possibilities. You could have a shield specifically designed for dots: something that would absorb
-		// the first X damage from each DOT applied to you. Practicaly? Maybe not, but demonstrates power of this setup.
-	}
-	else if (ShieldedDamage > 0.f)
-	{
-		// Apply a GameplayEffect to remove some of our shield
-		UGameplayEffect * LocalShieldRemoval = ShieldRemoveGameplayEffect;
-
-		if (!LocalShieldRemoval)
-		{
-			UProperty *HealthProperty = FindFieldChecked<UProperty>(UAbilitySystemTestAttributeSet::StaticClass(), GET_MEMBER_NAME_CHECKED(UAbilitySystemTestAttributeSet, Health));
-
-			// Since this is a test class and we don't want to tie it any actual content assets, just construct a GameplayEffect here.
-			LocalShieldRemoval = Cast<UGameplayEffect>(StaticConstructObject(UGameplayEffect::StaticClass(), GetTransientPackage(), FName(TEXT("ShieldAbsorbRemoval"))));
-			LocalShieldRemoval->Modifiers.SetNum(1);
-			LocalShieldRemoval->Modifiers[0].Magnitude.SetValue(-ShieldedDamage);
-			LocalShieldRemoval->Modifiers[0].ModifierType = EGameplayMod::ActiveGE;
-			LocalShieldRemoval->Modifiers[0].ModifierOp = EGameplayModOp::Additive;
-			LocalShieldRemoval->Modifiers[0].Attribute.SetUProperty(HealthProperty);
-			LocalShieldRemoval->Modifiers[0].OwnedTags.AddTag(GameplayTagsModule.GetGameplayTagsManager().RequestGameplayTag(FName(TEXT("ShieldAbsorb"))));
-			LocalShieldRemoval->Duration.Value = UGameplayEffect::INSTANT_APPLICATION;
-			LocalShieldRemoval->Period.Value = UGameplayEffect::NO_PERIOD;
-		}
-
-		Data.Target.ApplyGameplayEffectToSelf(LocalShieldRemoval, ShieldedDamage, Source->GetEffectContext(), FModifierQualifier().ExclusiveTarget(SelfData.Handle));
-	}
+	
 }
 
 void UGameplayEffectExtension_ShieldTest::PostGameplayEffectExecute(const FGameplayModifierEvaluatedData &SelfData, const FGameplayEffectModCallbackData &Data) const
diff --git a/Engine/Source/Runtime/GameplayAbilities/Private/GameplayEffectStackingExtension_CappedNumberTest.cpp b/Engine/Source/Runtime/GameplayAbilities/Private/GameplayEffectStackingExtension_CappedNumberTest.cpp
--- a/Engine/Source/Runtime/GameplayAbilities/Private/GameplayEffectStackingExtension_CappedNumberTest.cpp
+++ b/Engine/Source/Runtime/GameplayAbilities/Private/GameplayEffectStackingExtension_CappedNumberTest.cpp
@@ -15,6 +15,9 @@ UGameplayEffectStackingExtension_CappedNumberTest::UGameplayEffectStackingExtens
 
 void UGameplayEffectStackingExtension_CappedNumberTest::CalculateStack(TArray<FActiveGameplayEffect*>& CustomGameplayEffects, FActiveGameplayEffectsContainer& Container, FActiveGameplayEffect& CurrentEffect)
 {
+#if 0
+	// GE_REMOVE - this stuff should just go?
+
 	// this effect shouldn't be in the array so be sure to count it as well
 	int32 EffectiveCount = FMath::Min(CustomGameplayEffects.Num() + 1, MaxNumInStack);
 	
@@ -51,7 +54,7 @@ void UGameplayEffectStackingExtension_CappedNumberTest::CalculateStack(TArray<FA
 			// remove any stacking information that was already applied to the current modifier
 			for (int32 Idx = 0; Idx < Mod.Aggregator.Get()->Mods[EGameplayModOp::Multiplicitive].Num(); ++Idx)
 			{
-				FAggregatorRef& Agg = Mod.Aggregator.Get()->Mods[EGameplayModOp::Multiplicitive][Idx];
+				FAggregatorRefOld& Agg = Mod.Aggregator.Get()->Mods[EGameplayModOp::Multiplicitive][Idx];
 				if (Agg.Get()->BaseData.Tags.HasTag(IGameplayTagsModule::RequestGameplayTag("Stack.CappedNumber"), EGameplayTagMatchType::IncludeParentTags, EGameplayTagMatchType::Explicit))
 				{
 					Mod.Aggregator.Get()->Mods[EGameplayModOp::Multiplicitive].RemoveAtSwap(Idx);
@@ -72,4 +75,5 @@ void UGameplayEffectStackingExtension_CappedNumberTest::CalculateStack(TArray<FA
 			ModSpec.ApplyModTo(Mod, true);
 		}
 	}
+#endif
 }
diff --git a/Engine/Source/Runtime/GameplayAbilities/Private/GameplayEffectStackingExtension_DiminishingReturnsTest.cpp b/Engine/Source/Runtime/GameplayAbilities/Private/GameplayEffectStackingExtension_DiminishingReturnsTest.cpp
--- a/Engine/Source/Runtime/GameplayAbilities/Private/GameplayEffectStackingExtension_DiminishingReturnsTest.cpp
+++ b/Engine/Source/Runtime/GameplayAbilities/Private/GameplayEffectStackingExtension_DiminishingReturnsTest.cpp
@@ -22,6 +22,9 @@ UGameplayEffectStackingExtension_DiminishingReturnsTest::UGameplayEffectStacking
 
 void UGameplayEffectStackingExtension_DiminishingReturnsTest::CalculateStack(TArray<FActiveGameplayEffect*>& CustomGameplayEffects, FActiveGameplayEffectsContainer& Container, FActiveGameplayEffect& CurrentEffect)
 {
+#if 0
+	// GE_REMOVE this stuff should just go?
+
 	// look up the correct value in the curve table
 	UCurveTable* CurveTable = GetCurveTable();
 	float Multiplier = 1.f;
@@ -41,7 +44,7 @@ void UGameplayEffectStackingExtension_DiminishingReturnsTest::CalculateStack(TAr
 			// remove any stacking information that was already applied to the current modifier
 			for (int32 Idx = 0; Idx < Mod.Aggregator.Get()->Mods[EGameplayModOp::Multiplicitive].Num(); ++Idx)
 			{
-				FAggregatorRef& Agg = Mod.Aggregator.Get()->Mods[EGameplayModOp::Multiplicitive][Idx];
+				FAggregatorRefOld& Agg = Mod.Aggregator.Get()->Mods[EGameplayModOp::Multiplicitive][Idx];
 				if (Agg.Get()->BaseData.Tags.HasTag(IGameplayTagsModule::RequestGameplayTag("Stack.DiminishingReturns"), EGameplayTagMatchType::IncludeParentTags, EGameplayTagMatchType::Explicit))
 				{
 					Mod.Aggregator.Get()->Mods[EGameplayModOp::Multiplicitive].RemoveAtSwap(Idx);
@@ -62,4 +65,5 @@ void UGameplayEffectStackingExtension_DiminishingReturnsTest::CalculateStack(TAr
 			ModSpec.ApplyModTo(Mod, true);
 		}
 	}
+#endif
 }
