Gameplay Effect inital refactor

diff --git a/Engine/Source/Runtime/GameplayAbilities/Private/GameplayEffectTests.cpp b/Engine/Source/Runtime/GameplayAbilities/Private/GameplayEffectTests.cpp
--- a/Engine/Source/Runtime/GameplayAbilities/Private/GameplayEffectTests.cpp
+++ b/Engine/Source/Runtime/GameplayAbilities/Private/GameplayEffectTests.cpp
@@ -32,6023 +32,75 @@ void GameplayTest_TickWorld(UWorld *World, float Time)
 	}
 }
 
-void MySharedPointerTest()
-{
-	// Test that outside will be invalid after inside goes out of scope
-	{
-		FAggregatorRef	Outside;
-		check(!Outside.IsValid());
-		{
-			FAggregatorRef Inside(new FAggregator());
-			check(Inside.IsValid());
-
-			Outside.SetSoftRef(&Inside);
-			check(Outside.IsValid());
-		}
-
-		check(!Outside.IsValid());
-	}
-
-	// Test that outside will valid since it calls MakeHardRef
-	{
-		FAggregatorRef	Outside;
-		check(!Outside.IsValid());
-		{
-			FAggregatorRef Inside(new FAggregator());
-			check(Inside.IsValid());
-
-			Outside.SetSoftRef(&Inside);
-			check(Outside.IsValid());
-
-			// The difference
-			Outside.MakeHardRef();
-		}
-	
-		check(Outside.IsValid());
-	}
-
-	// TArray test
-	TArray< FAggregatorRef >	Array;
-	TArray< FAggregatorRef >	Array2;
-	Array2.SetNum(640);
-
-	{
-		FAggregatorRef Inside(new FAggregator());
-		Array.SetNum(1);
-		Array[0].SetSoftRef(&Inside);
-		Array[0].MakeHardRef();
-	}
-
-	check(Array[0].IsValid());
-	Array.SetNum(1024);
-	check(Array[0].IsValid());
-}
-
-void MySharedPointerTest_Array()
-{
-	TArray< TSharedPtr<FAggregator> >	Array1;
-	TArray< TSharedPtr<FAggregator> >	Array2;
-	
-	Array1.Emplace( TSharedPtr<FAggregator>( new FAggregator() ) );
-
-	Array2 = Array1;
-
-
-	Array1[0].Reset();
-
-	Array2[0] = TSharedPtr<FAggregator>( new FAggregator(*Array2[0].Get()) );
-
-
-	check( Array2[0].IsValid() );
-}
-
-bool GameplayEffectsTest_InstantDamage(UWorld *World, FAutomationTestBase * Test)
-{
-	const float StartHealth = 100.f;
-	const float DamageValue = -5.f;
-
-	AAbilitySystemTestPawn *SourceActor = World->SpawnActor<AAbilitySystemTestPawn>();
-	AAbilitySystemTestPawn *DestActor = World->SpawnActor<AAbilitySystemTestPawn>();
-
-	UProperty *HealthProperty = FindFieldChecked<UProperty>(UAbilitySystemTestAttributeSet::StaticClass(), GET_MEMBER_NAME_CHECKED(UAbilitySystemTestAttributeSet, Health));
-
-	UAbilitySystemComponent * SourceComponent = SourceActor->GetAbilitySystemComponent();
-	UAbilitySystemComponent * DestComponent = DestActor->GetAbilitySystemComponent();
-	SourceComponent->GetSet<UAbilitySystemTestAttributeSet>()->Health = StartHealth;
-	DestComponent->GetSet<UAbilitySystemTestAttributeSet>()->Health = StartHealth;
-
-	{
-		ABILITY_LOG_SCOPE(TEXT("Apply InstantDamage"))
-
-		UGameplayEffect * BaseDmgEffect = Cast<UGameplayEffect>(StaticConstructObject(UGameplayEffect::StaticClass(), GetTransientPackage(), FName(TEXT("BaseDmgEffect"))));
-		BaseDmgEffect->Modifiers.SetNum(1);
-		BaseDmgEffect->Modifiers[0].Magnitude.SetValue(DamageValue);
-		BaseDmgEffect->Modifiers[0].ModifierType = EGameplayMod::Attribute;
-		BaseDmgEffect->Modifiers[0].ModifierOp = EGameplayModOp::Additive;
-		BaseDmgEffect->Modifiers[0].Attribute.SetUProperty(HealthProperty);
-		BaseDmgEffect->Modifiers[0].OwnedTags.AddTag(IGameplayTagsModule::RequestGameplayTag(FName(TEXT("Damage.Basic"))));
-		BaseDmgEffect->Duration.Value = UGameplayEffect::INSTANT_APPLICATION;
-
-		SourceComponent->ApplyGameplayEffectToTarget(BaseDmgEffect, DestComponent, 1.f);
-		
-		Test->TestTrue(SKILL_TEST_TEXT("Basic Instant Damage Applied"), (DestComponent->GetSet<UAbilitySystemTestAttributeSet>()->Health == (StartHealth + DamageValue)));
-		ABILITY_LOG(Log, TEXT("Final Health: %.2f"), DestComponent->GetSet<UAbilitySystemTestAttributeSet>()->Health);
-	}
-
-	World->EditorDestroyActor(SourceActor, false);
-	World->EditorDestroyActor(DestActor, false);
-
-	return true;
-}
-
-bool GameplayEffectsTest_InstantDamageRemap(UWorld *World, FAutomationTestBase * Test)
-{
-	// This is the same as GameplayEffectsTest_InstantDamage but modifies the Damage attribute and confirms it is remapped to -Health by UAbilitySystemTestAttributeSet::PostAttributeModify
-
-	const float StartHealth = 100.f;
-	const float DamageValue = 5.f;		// Note: Damage is positive, mapped to -Health in UAbilitySystemTestAttributeSet::PostAttributeModify
-
-	AAbilitySystemTestPawn *SourceActor = World->SpawnActor<AAbilitySystemTestPawn>();
-	AAbilitySystemTestPawn *DestActor = World->SpawnActor<AAbilitySystemTestPawn>();
-
-	UProperty *HealthProperty = FindFieldChecked<UProperty>(UAbilitySystemTestAttributeSet::StaticClass(), GET_MEMBER_NAME_CHECKED(UAbilitySystemTestAttributeSet, Health));
-	UProperty *DamageProperty = FindFieldChecked<UProperty>(UAbilitySystemTestAttributeSet::StaticClass(), GET_MEMBER_NAME_CHECKED(UAbilitySystemTestAttributeSet, Damage));
-
-	UAbilitySystemComponent * SourceComponent = SourceActor->GetAbilitySystemComponent();
-	UAbilitySystemComponent * DestComponent = DestActor->GetAbilitySystemComponent();
-	SourceComponent->GetSet<UAbilitySystemTestAttributeSet>()->Health = StartHealth;
-	DestComponent->GetSet<UAbilitySystemTestAttributeSet>()->Health = StartHealth;
-
-	{
-		ABILITY_LOG_SCOPE(TEXT("Apply InstantDamage"));
-
-		UGameplayEffect * BaseDmgEffect = Cast<UGameplayEffect>(StaticConstructObject(UGameplayEffect::StaticClass(), GetTransientPackage(), FName(TEXT("BaseDmgEffect"))));
-		BaseDmgEffect->Modifiers.SetNum(1);
-		BaseDmgEffect->Modifiers[0].Magnitude.SetValue(DamageValue);
-		BaseDmgEffect->Modifiers[0].ModifierType = EGameplayMod::Attribute;
-		BaseDmgEffect->Modifiers[0].ModifierOp = EGameplayModOp::Additive;
-		BaseDmgEffect->Modifiers[0].Attribute.SetUProperty(DamageProperty);
-		BaseDmgEffect->Modifiers[0].OwnedTags.AddTag(IGameplayTagsModule::RequestGameplayTag(FName(TEXT("Damage.Basic"))));
-		BaseDmgEffect->Duration.Value = UGameplayEffect::INSTANT_APPLICATION;
-
-		SourceComponent->ApplyGameplayEffectToTarget(BaseDmgEffect, DestComponent, 1.f);
-
-		// Now we should have lost some health
-		{
-			float ActualValue = DestComponent->GetSet<UAbilitySystemTestAttributeSet>()->Health;
-			float ExpectedValue = StartHealth + -DamageValue;
-			Test->TestTrue(SKILL_TEST_TEXT("Damage Applied. Actual: %.2f == Exected: %.2f", ActualValue, ExpectedValue), ActualValue == ExpectedValue);
-		}
-
-		// Confirm the damage attribute itself was reset to 0 when it was applied to health
-		{
-			float ActualValue = DestComponent->GetSet<UAbilitySystemTestAttributeSet>()->Damage;
-			float ExpectedValue = 0.f;
-			Test->TestTrue(SKILL_TEST_TEXT("Damage Applied. Actual: %.2f == Exected: %.2f", ActualValue, ExpectedValue), ActualValue == ExpectedValue);
-		}
-	}
-
-	World->EditorDestroyActor(SourceActor, false);
-	World->EditorDestroyActor(DestActor, false);
-
-	return true;
-}
-
-bool GameplayEffectsTest_InstantDamage_Buffed(UWorld *World, FAutomationTestBase * Test)
-{
-	const float StartHealth = 100.f;
-	const float DamageValue = -5.f;
-	const float BonusDamageMultiplier = 2.f;
-
-	AAbilitySystemTestPawn *SourceActor = World->SpawnActor<AAbilitySystemTestPawn>();
-	AAbilitySystemTestPawn *DestActor = World->SpawnActor<AAbilitySystemTestPawn>();
-
-	UProperty *HealthProperty = FindFieldChecked<UProperty>(UAbilitySystemTestAttributeSet::StaticClass(), GET_MEMBER_NAME_CHECKED(UAbilitySystemTestAttributeSet, Health));
-
-	UAbilitySystemComponent * SourceComponent = SourceActor->GetAbilitySystemComponent();
-	UAbilitySystemComponent * DestComponent = DestActor->GetAbilitySystemComponent();
-	SourceComponent->GetSet<UAbilitySystemTestAttributeSet>()->Health = StartHealth;
-	DestComponent->GetSet<UAbilitySystemTestAttributeSet>()->Health = StartHealth;
-
-	// Setup a GE to modify OutgoingGEs
-	{
-		ABILITY_LOG_SCOPE(TEXT("Apply DamageBuff"))
-		
-		UGameplayEffect * BaseDmgEffect = Cast<UGameplayEffect>(StaticConstructObject(UGameplayEffect::StaticClass(), GetTransientPackage(), FName(TEXT("DamageBuff"))));
-		BaseDmgEffect->Modifiers.SetNum(1);
-		BaseDmgEffect->Modifiers[0].Magnitude.SetValue(BonusDamageMultiplier);
-		BaseDmgEffect->Modifiers[0].ModifierType = EGameplayMod::OutgoingGE;
-		BaseDmgEffect->Modifiers[0].ModifierOp = EGameplayModOp::Multiplicitive;
-		BaseDmgEffect->Modifiers[0].Attribute.SetUProperty(HealthProperty);
-		BaseDmgEffect->Modifiers[0].OwnedTags.AddTag(IGameplayTagsModule::RequestGameplayTag(FName(TEXT("Damage.Basic"))));
-		BaseDmgEffect->Duration.SetValue(UGameplayEffect::INFINITE_DURATION);
-
-		// Apply to self
-		SourceComponent->ApplyGameplayEffectToTarget(BaseDmgEffect, SourceComponent, 1.f);
-	}
-
-	// Apply Damage
-	{
-		ABILITY_LOG_SCOPE(TEXT("Apply InstantDamage"))
-
-		UGameplayEffect * BaseDmgEffect = Cast<UGameplayEffect>(StaticConstructObject(UGameplayEffect::StaticClass(), GetTransientPackage(), FName(TEXT("Damage"))));
-		BaseDmgEffect->Modifiers.SetNum(1);
-		BaseDmgEffect->Modifiers[0].Magnitude.SetValue(DamageValue);
-		BaseDmgEffect->Modifiers[0].ModifierType = EGameplayMod::Attribute;
-		BaseDmgEffect->Modifiers[0].ModifierOp = EGameplayModOp::Additive;
-		BaseDmgEffect->Modifiers[0].Attribute.SetUProperty(HealthProperty);
-		BaseDmgEffect->Modifiers[0].OwnedTags.AddTag(IGameplayTagsModule::RequestGameplayTag(FName(TEXT("Damage.Basic"))));
-		BaseDmgEffect->Duration.Value = UGameplayEffect::INSTANT_APPLICATION;
-
-		// Apply to target
-		SourceComponent->ApplyGameplayEffectToTarget(BaseDmgEffect, DestComponent, 1.f);
-
-		float ExpectedValue = (StartHealth + (DamageValue * BonusDamageMultiplier));
-
-		Test->TestTrue(SKILL_TEST_TEXT("Buff Instant Damage Applied"), (DestComponent->GetSet<UAbilitySystemTestAttributeSet>()->Health == ExpectedValue));
-		ABILITY_LOG(Log, TEXT("Final Health: %.2f"), DestComponent->GetSet<UAbilitySystemTestAttributeSet>()->Health);
-	}
-
-	World->EditorDestroyActor(SourceActor, false);
-	World->EditorDestroyActor(DestActor, false);
-
-	return true;
-}
-
-bool GameplayEffectsTest_TemporaryDamage(UWorld *World, FAutomationTestBase * Test)
-{
-	/**
-	 *	This test applies a temporary -10 Health GE then removes it to show Health goes back to start
-	 */
-
-	const float StartHealth = 100.f;
-	const float DamageValue = -5.f;
-
-	AAbilitySystemTestPawn *SourceActor = World->SpawnActor<AAbilitySystemTestPawn>();
-	AAbilitySystemTestPawn *DestActor = World->SpawnActor<AAbilitySystemTestPawn>();
-
-	UProperty *HealthProperty = FindFieldChecked<UProperty>(UAbilitySystemTestAttributeSet::StaticClass(), GET_MEMBER_NAME_CHECKED(UAbilitySystemTestAttributeSet, Health));
-
-	UAbilitySystemComponent * SourceComponent = SourceActor->GetAbilitySystemComponent();
-	UAbilitySystemComponent * DestComponent = DestActor->GetAbilitySystemComponent();
-	SourceComponent->GetSet<UAbilitySystemTestAttributeSet>()->Health = StartHealth;
-	DestComponent->GetSet<UAbilitySystemTestAttributeSet>()->Health = StartHealth;
-
-	// Apply "Damage" but set to INFINITE_DURATION
-	// (An odd example for damage, but would make sense for something like run speed, etc)
-	FActiveGameplayEffectHandle AppliedHandle;
-	{
-		ABILITY_LOG_SCOPE(TEXT("Apply Permanent Damage"))
-
-		UGameplayEffect * BaseDmgEffect = Cast<UGameplayEffect>(StaticConstructObject(UGameplayEffect::StaticClass(), GetTransientPackage(), FName(TEXT("Damage"))));
-		BaseDmgEffect->Modifiers.SetNum(1);
-		BaseDmgEffect->Modifiers[0].Magnitude.SetValue(DamageValue);
-		BaseDmgEffect->Modifiers[0].ModifierType = EGameplayMod::Attribute;
-		BaseDmgEffect->Modifiers[0].ModifierOp = EGameplayModOp::Additive;
-		BaseDmgEffect->Modifiers[0].Attribute.SetUProperty(HealthProperty);
-		BaseDmgEffect->Modifiers[0].OwnedTags.AddTag(IGameplayTagsModule::RequestGameplayTag(FName(TEXT("Damage.Basic"))));
-		BaseDmgEffect->Duration.Value = UGameplayEffect::INFINITE_DURATION;
-
-		// Apply to target
-		AppliedHandle = SourceComponent->ApplyGameplayEffectToTarget(BaseDmgEffect, DestComponent, 1.f);
-
-		float ExpectedValue = (StartHealth + DamageValue);
-
-		Test->TestTrue(SKILL_TEST_TEXT("INFINITE_DURATION Damage Applied"), (DestComponent->GetSet<UAbilitySystemTestAttributeSet>()->Health == ExpectedValue));
-		ABILITY_LOG(Log, TEXT("After Damage Applied: Health: %.2f"), DestComponent->GetSet<UAbilitySystemTestAttributeSet>()->Health);
-	}
-
-	// Now remove the GameplayEffect we just added and confirm Health goes back to starting value
-	{
-		bool RemovedEffect = DestComponent->RemoveActiveGameplayEffect(AppliedHandle);
-		float ExpectedValue = StartHealth;
-
-		Test->TestTrue(SKILL_TEST_TEXT("INFINITE_DURATION Damage Applied"), (DestComponent->GetSet<UAbilitySystemTestAttributeSet>()->Health == ExpectedValue));
-		ABILITY_LOG(Log, TEXT("After Removal. Health: %.2f. RemovedEffecte: %d"), DestComponent->GetSet<UAbilitySystemTestAttributeSet>()->Health, RemovedEffect);
-	}
-
-	World->EditorDestroyActor(SourceActor, false);
-	World->EditorDestroyActor(DestActor, false);
-
-	return true;
-}
-
-bool GameplayEffectsTest_TemporaryDamageBuffed(UWorld *World, FAutomationTestBase * Test)
-{
-	/**
-	 *	This test applies a temporary -10 Health GE then buffs it with an executed (ActiveGE) GE, then removes it to show Health goes back to initial value
-	 */
-
-	const float StartHealth = 100.f;
-	const float DamageValue = -5.f;
-	const float DamageBuffMultiplier = 2.f;
-
-	AAbilitySystemTestPawn *SourceActor = World->SpawnActor<AAbilitySystemTestPawn>();
-	AAbilitySystemTestPawn *DestActor = World->SpawnActor<AAbilitySystemTestPawn>();
-
-	UProperty *HealthProperty = FindFieldChecked<UProperty>(UAbilitySystemTestAttributeSet::StaticClass(), GET_MEMBER_NAME_CHECKED(UAbilitySystemTestAttributeSet, Health));
-
-	UAbilitySystemComponent * SourceComponent = SourceActor->GetAbilitySystemComponent();
-	UAbilitySystemComponent * DestComponent = DestActor->GetAbilitySystemComponent();
-	SourceComponent->GetSet<UAbilitySystemTestAttributeSet>()->Health = StartHealth;
-	DestComponent->GetSet<UAbilitySystemTestAttributeSet>()->Health = StartHealth;
-
-	// Apply "Damage" but set to INFINITE_DURATION
-	// (An odd example for damage, but would make sense for something like run speed, etc)
-	FActiveGameplayEffectHandle AppliedHandle;
-	{
-		ABILITY_LOG_SCOPE(TEXT("Apply Permanent Damage"))
-
-		UGameplayEffect * BaseDmgEffect = Cast<UGameplayEffect>(StaticConstructObject(UGameplayEffect::StaticClass(), GetTransientPackage(), FName(TEXT("Damage"))));
-		BaseDmgEffect->Modifiers.SetNum(1);
-		BaseDmgEffect->Modifiers[0].Magnitude.SetValue(DamageValue);
-		BaseDmgEffect->Modifiers[0].ModifierType = EGameplayMod::Attribute;
-		BaseDmgEffect->Modifiers[0].ModifierOp = EGameplayModOp::Additive;
-		BaseDmgEffect->Modifiers[0].Attribute.SetUProperty(HealthProperty);
-		BaseDmgEffect->Modifiers[0].OwnedTags.AddTag(IGameplayTagsModule::RequestGameplayTag(FName(TEXT("Damage.Basic"))));
-		BaseDmgEffect->Duration.Value = UGameplayEffect::INFINITE_DURATION;
-
-		// Apply to target
-		AppliedHandle = SourceComponent->ApplyGameplayEffectToTarget(BaseDmgEffect, DestComponent, 1.f);
-
-		float ExpectedValue = (StartHealth + DamageValue);
-
-		Test->TestTrue(SKILL_TEST_TEXT("INFINITE_DURATION Damage Applied"), (DestComponent->GetSet<UAbilitySystemTestAttributeSet>()->Health == ExpectedValue));
-		ABILITY_LOG(Log, TEXT("After Damage Applied: Health: %.2f"), DestComponent->GetSet<UAbilitySystemTestAttributeSet>()->Health);
-
-		Test->TestTrue(SKILL_TEST_TEXT("Number of GameplayEffects=1"), DestComponent->GetNumActiveGameplayEffect() == 1);
-	}
-
-	// Now Buff the GameplayEffect we just added and confirm the health removal is increased 2x
-	{
-		ABILITY_LOG_SCOPE(TEXT("Buff Permanent Damage"))
-
-		UGameplayEffect * BuffDmgEffect = Cast<UGameplayEffect>(StaticConstructObject(UGameplayEffect::StaticClass(), GetTransientPackage(), FName(TEXT("BuffDamage"))));
-		BuffDmgEffect->Modifiers.SetNum(1);
-		BuffDmgEffect->Modifiers[0].Magnitude.SetValue(DamageBuffMultiplier);
-		BuffDmgEffect->Modifiers[0].ModifierType = EGameplayMod::ActiveGE;
-		BuffDmgEffect->Modifiers[0].ModifierOp = EGameplayModOp::Multiplicitive;
-		BuffDmgEffect->Modifiers[0].Attribute.SetUProperty(HealthProperty);
-		BuffDmgEffect->Modifiers[0].OwnedTags.AddTag(IGameplayTagsModule::RequestGameplayTag(FName(TEXT("Damage.Basic"))));
-		BuffDmgEffect->Duration.Value = UGameplayEffect::INSTANT_APPLICATION;
-
-		// Apply to target
-		SourceComponent->ApplyGameplayEffectToTarget(BuffDmgEffect, DestComponent, 1.f);
-
-		float ExpectedValue = (StartHealth + (DamageValue * DamageBuffMultiplier));
-
-		Test->TestTrue(SKILL_TEST_TEXT("INFINITE_DURATION Buffed Damage Applied"), (DestComponent->GetSet<UAbilitySystemTestAttributeSet>()->Health == ExpectedValue));
-		ABILITY_LOG(Log, TEXT("After Damage Applied: Health: %.2f"), DestComponent->GetSet<UAbilitySystemTestAttributeSet>()->Health);
-
-		// Confirm still only 1 active GE (since this was instant application
-		Test->TestTrue(SKILL_TEST_TEXT("Number of GameplayEffects=1"), DestComponent->GetNumActiveGameplayEffect() == 1);
-	}
-
-
-	// Now remove the GameplayEffect we just added and confirm Health goes back to starting value
-	{
-		bool RemovedEffect = DestComponent->RemoveActiveGameplayEffect(AppliedHandle);
-		float ExpectedValue = StartHealth;
-
-		Test->TestTrue(SKILL_TEST_TEXT("INFINITE_DURATION Damage Applied"), (DestComponent->GetSet<UAbilitySystemTestAttributeSet>()->Health == ExpectedValue));
-		ABILITY_LOG(Log, TEXT("After Removal. Health: %.2f. RemovedEffecte: %d"), DestComponent->GetSet<UAbilitySystemTestAttributeSet>()->Health, RemovedEffect);
-
-		// Confirm no more GEs
-		Test->TestTrue(SKILL_TEST_TEXT("Number of GameplayEffects=0"), DestComponent->GetNumActiveGameplayEffect() == 0);
-	}
-
-	World->EditorDestroyActor(SourceActor, false);
-	World->EditorDestroyActor(DestActor, false);
-
-	return true;
-}
-
-bool GameplayEffectsTest_TemporaryDamageTemporaryBuff(UWorld *World, FAutomationTestBase * Test)
-{
-	/**
-	 *	This test applies a temporary -10 Health GE then applies a temporary buff to the health modification. It then removes the buff, then the damage, and checks the health values return to what is expected
-	 */
-
-	const float StartHealth = 100.f;
-	const float DamageValue = -5.f;
-	const float DamageBuffMultiplier = 2.f;
-
-	AAbilitySystemTestPawn *SourceActor = World->SpawnActor<AAbilitySystemTestPawn>();
-	AAbilitySystemTestPawn *DestActor = World->SpawnActor<AAbilitySystemTestPawn>();
-
-	UProperty *HealthProperty = FindFieldChecked<UProperty>(UAbilitySystemTestAttributeSet::StaticClass(), GET_MEMBER_NAME_CHECKED(UAbilitySystemTestAttributeSet, Health));
-
-	UAbilitySystemComponent * SourceComponent = SourceActor->GetAbilitySystemComponent();
-	UAbilitySystemComponent * DestComponent = DestActor->GetAbilitySystemComponent();
-	SourceComponent->GetSet<UAbilitySystemTestAttributeSet>()->Health = StartHealth;
-	DestComponent->GetSet<UAbilitySystemTestAttributeSet>()->Health = StartHealth;
-
-	// Apply "Damage" but set to INFINITE_DURATION
-	// (An odd example for damage, but would make sense for something like run speed, etc)
-	FActiveGameplayEffectHandle AppliedHandle;
-	FActiveGameplayEffectHandle AppliedHandleBuff;
-
-	{
-		ABILITY_LOG_SCOPE(TEXT("Apply Permanent Damage"))
-
-		UGameplayEffect * BaseDmgEffect = Cast<UGameplayEffect>(StaticConstructObject(UGameplayEffect::StaticClass(), GetTransientPackage(), FName(TEXT("Damage"))));
-		BaseDmgEffect->Modifiers.SetNum(1);
-		BaseDmgEffect->Modifiers[0].Magnitude.SetValue(DamageValue);
-		BaseDmgEffect->Modifiers[0].ModifierType = EGameplayMod::Attribute;
-		BaseDmgEffect->Modifiers[0].ModifierOp = EGameplayModOp::Additive;
-		BaseDmgEffect->Modifiers[0].Attribute.SetUProperty(HealthProperty);
-		BaseDmgEffect->Modifiers[0].OwnedTags.AddTag(IGameplayTagsModule::RequestGameplayTag(FName(TEXT("Damage.Basic"))));
-		BaseDmgEffect->Duration.Value = UGameplayEffect::INFINITE_DURATION;
-
-		// Apply to target
-		AppliedHandle = SourceComponent->ApplyGameplayEffectToTarget(BaseDmgEffect, DestComponent, 1.f);
-
-		float ExpectedValue = (StartHealth + DamageValue);
-
-		Test->TestTrue(SKILL_TEST_TEXT("GameplayEffectsTest_TemporaryDamageTemporaryBuff INFINITE_DURATION Damage Applied"), (DestComponent->GetSet<UAbilitySystemTestAttributeSet>()->Health == ExpectedValue));
-		ABILITY_LOG(Log, TEXT("After Damage Applied: Health: %.2f"), DestComponent->GetSet<UAbilitySystemTestAttributeSet>()->Health);
-
-		Test->TestTrue(SKILL_TEST_TEXT("GameplayEffectsTest_TemporaryDamageTemporaryBuff Number of GameplayEffects=1"), DestComponent->GetNumActiveGameplayEffect() == 1);
-	}
-
-	// Now Buff the GameplayEffect we just added and confirm the health removal is increased 2x
-	{
-		ABILITY_LOG_SCOPE(TEXT("Buff Permanent Damage"))
-
-		UGameplayEffect * BuffDmgEffect = Cast<UGameplayEffect>(StaticConstructObject(UGameplayEffect::StaticClass(), GetTransientPackage(), FName(TEXT("BuffDamage"))));
-		BuffDmgEffect->Modifiers.SetNum(1);
-		BuffDmgEffect->Modifiers[0].Magnitude.SetValue(DamageBuffMultiplier);
-		BuffDmgEffect->Modifiers[0].ModifierType = EGameplayMod::ActiveGE;
-		BuffDmgEffect->Modifiers[0].ModifierOp = EGameplayModOp::Multiplicitive;
-		BuffDmgEffect->Modifiers[0].Attribute.SetUProperty(HealthProperty);
-		BuffDmgEffect->Modifiers[0].OwnedTags.AddTag(IGameplayTagsModule::RequestGameplayTag(FName(TEXT("Damage.Basic"))));
-		BuffDmgEffect->Duration.Value = UGameplayEffect::INFINITE_DURATION;
-		BuffDmgEffect->CopyPolicy = EGameplayEffectCopyPolicy::AlwaysLink;		// Force this to link, so that when we remove it it will go away to any modifier it was applied to
-
-		// Apply to target
-		AppliedHandleBuff = SourceComponent->ApplyGameplayEffectToTarget(BuffDmgEffect, DestComponent, 1.f);
-
-		float ExpectedValue = (StartHealth + (DamageValue * DamageBuffMultiplier));
-
-		Test->TestTrue(SKILL_TEST_TEXT("GameplayEffectsTest_TemporaryDamageTemporaryBuff INFINITE_DURATION Buffed Damage Applied"), (DestComponent->GetSet<UAbilitySystemTestAttributeSet>()->Health == ExpectedValue));
-		ABILITY_LOG(Log, TEXT("After Damage Applied: Health: %.2f"), DestComponent->GetSet<UAbilitySystemTestAttributeSet>()->Health);
-
-		// Confirm there are 2 GEs
-		Test->TestTrue(SKILL_TEST_TEXT("GameplayEffectsTest_TemporaryDamageTemporaryBuff Number of GameplayEffects=1"), DestComponent->GetNumActiveGameplayEffect() == 2);
-	}
-
-	// Print out the whole enchillada
-	{
-		// DestComponent->PrintAllGameplayEffects();
-	}
-
-	// Remove the buff GE
-	{
-		bool RemovedEffect = DestComponent->RemoveActiveGameplayEffect(AppliedHandleBuff);
-		float ExpectedValue = StartHealth + DamageValue;
-
-		Test->TestTrue(SKILL_TEST_TEXT("GameplayEffectsTest_TemporaryDamageTemporaryBuff INFINITE_DURATION Damage Buff Removed"), (DestComponent->GetSet<UAbilitySystemTestAttributeSet>()->Health == ExpectedValue));
-		ABILITY_LOG(Log, TEXT("After Buff Removal. Health: %.2f. RemovedEffecte: %d"), DestComponent->GetSet<UAbilitySystemTestAttributeSet>()->Health, RemovedEffect);
-
-		// Confirm 1 more GEs
-		Test->TestTrue(SKILL_TEST_TEXT("Number of GameplayEffects=1"), DestComponent->GetNumActiveGameplayEffect() == 1);
-	}
-
-	// Remove the damage GE
-	{
-		bool RemovedEffect = DestComponent->RemoveActiveGameplayEffect(AppliedHandle);
-		float ExpectedValue = StartHealth;
-
-		Test->TestTrue(SKILL_TEST_TEXT("GameplayEffectsTest_TemporaryDamageTemporaryBuff INFINITE_DURATION Damage Removed"), (DestComponent->GetSet<UAbilitySystemTestAttributeSet>()->Health == ExpectedValue));
-		ABILITY_LOG(Log, TEXT("After Removal. Health: %.2f. RemovedEffecte: %d"), DestComponent->GetSet<UAbilitySystemTestAttributeSet>()->Health, RemovedEffect);
-
-		// Confirm no more GEs
-		Test->TestTrue(SKILL_TEST_TEXT("GameplayEffectsTest_TemporaryDamageTemporaryBuff Number of GameplayEffects=0"), DestComponent->GetNumActiveGameplayEffect() == 0);
-	}
-
-	World->EditorDestroyActor(SourceActor, false);
-	World->EditorDestroyActor(DestActor, false);
-
-	return true;
-}
-
-bool GameplayEffectsTest_LinkedBuffDestroy(UWorld *World, FAutomationTestBase * Test)
-{
-	/**
-	 *	Apply a perm health reduction that is buffed by an outgoing GE buff. Then destroy the buff and see what happens to the perm applied Ge.
-	 */
-
-	const float StartHealth = 100.f;
-	const float DamageValue = -5.f;
-	const float DamageBuffMultiplier = 2.f;
-
-	AAbilitySystemTestPawn *SourceActor = World->SpawnActor<AAbilitySystemTestPawn>();
-	AAbilitySystemTestPawn *DestActor = World->SpawnActor<AAbilitySystemTestPawn>();
-
-	UProperty *HealthProperty = FindFieldChecked<UProperty>(UAbilitySystemTestAttributeSet::StaticClass(), GET_MEMBER_NAME_CHECKED(UAbilitySystemTestAttributeSet, Health));
-
-	UAbilitySystemComponent * SourceComponent = SourceActor->GetAbilitySystemComponent();
-	UAbilitySystemComponent * DestComponent = DestActor->GetAbilitySystemComponent();
-	SourceComponent->GetSet<UAbilitySystemTestAttributeSet>()->Health = StartHealth;
-	DestComponent->GetSet<UAbilitySystemTestAttributeSet>()->Health = StartHealth;
-
-	// Apply "Damage" but set to INFINITE_DURATION
-	// (An odd example for damage, but would make sense for something like run speed, etc)
-	FActiveGameplayEffectHandle AppliedHandle;
-	FActiveGameplayEffectHandle AppliedHandleBuff;
-
-	{
-		ABILITY_LOG_SCOPE(TEXT("Apply Damage Buff"))
-
-		UGameplayEffect * BuffEffect = Cast<UGameplayEffect>(StaticConstructObject(UGameplayEffect::StaticClass(), GetTransientPackage(), FName(TEXT("BuffOutgoingDamage"))));
-		BuffEffect->Modifiers.SetNum(1);
-		BuffEffect->Modifiers[0].Magnitude.SetValue(DamageBuffMultiplier);
-		BuffEffect->Modifiers[0].ModifierType = EGameplayMod::OutgoingGE;
-		BuffEffect->Modifiers[0].ModifierOp = EGameplayModOp::Multiplicitive;
-		BuffEffect->Modifiers[0].Attribute.SetUProperty(HealthProperty);
-		BuffEffect->Duration.Value = UGameplayEffect::INFINITE_DURATION;
-		BuffEffect->CopyPolicy = EGameplayEffectCopyPolicy::AlwaysLink;		// Always link so that when this is destroyed, health return to normal
-
-		// Apply to target
-		AppliedHandleBuff = SourceComponent->ApplyGameplayEffectToTarget(BuffEffect, SourceComponent, 1.f);
-
-		Test->TestTrue(SKILL_TEST_TEXT("Number of Source GameplayEffect: %d == 1", SourceComponent->GetNumActiveGameplayEffect()), SourceComponent->GetNumActiveGameplayEffect() == 1);
-	}
-
-	{
-		ABILITY_LOG_SCOPE(TEXT("Apply Permanent (infinite duration) Damage"))
-
-		UGameplayEffect * BaseDmgEffect = Cast<UGameplayEffect>(StaticConstructObject(UGameplayEffect::StaticClass(), GetTransientPackage(), FName(TEXT("Damage"))));
-		BaseDmgEffect->Modifiers.SetNum(1);
-		BaseDmgEffect->Modifiers[0].Magnitude.SetValue(DamageValue);
-		BaseDmgEffect->Modifiers[0].ModifierType = EGameplayMod::Attribute;
-		BaseDmgEffect->Modifiers[0].ModifierOp = EGameplayModOp::Additive;
-		BaseDmgEffect->Modifiers[0].Attribute.SetUProperty(HealthProperty);
-		BaseDmgEffect->Modifiers[0].OwnedTags.AddTag(IGameplayTagsModule::RequestGameplayTag(FName(TEXT("Damage.Basic"))));
-		BaseDmgEffect->Duration.Value = UGameplayEffect::INFINITE_DURATION;
-
-		// Apply to target
-		AppliedHandle = SourceComponent->ApplyGameplayEffectToTarget(BaseDmgEffect, DestComponent, 1.f);
-
-		float ExpectedValue = (StartHealth + (DamageValue * DamageBuffMultiplier));
-		float ActualValue = DestComponent->GetSet<UAbilitySystemTestAttributeSet>()->Health;
-
-		Test->TestTrue(SKILL_TEST_TEXT("Damaged Applied. ActualValue: %.2f. ExpectedValue: %.2f.", ActualValue, ExpectedValue), (ActualValue == ExpectedValue));
-		Test->TestTrue(SKILL_TEST_TEXT("NumberOfActive GameplayEffects GameplayEffects %d == 1", DestComponent->GetNumActiveGameplayEffect() ), DestComponent->GetNumActiveGameplayEffect() == 1);
-	}
-
-	// Remove the buff GE
-	{
-		bool RemovedEffect = SourceComponent->RemoveActiveGameplayEffect(AppliedHandleBuff);
-		float ExpectedValue = StartHealth + DamageValue;
-		float ActualValue = DestComponent->GetSet<UAbilitySystemTestAttributeSet>()->Health;
-
-		// Confirm we regained healtht
-		Test->TestTrue(SKILL_TEST_TEXT("After Buff Removal. ActualValue: %.2f. ExpectedValue: %.2f. RemovedEffecte: %d", ActualValue, ExpectedValue, RemovedEffect), (ActualValue  == ExpectedValue));
-
-		// Confirm number of GEs
-		Test->TestTrue(SKILL_TEST_TEXT("Dest Number of GameplayEffects=%d", SourceComponent->GetNumActiveGameplayEffect()), SourceComponent->GetNumActiveGameplayEffect() == 0);
-		Test->TestTrue(SKILL_TEST_TEXT("Src Number of GameplayEffects=%d", DestComponent->GetNumActiveGameplayEffect()), DestComponent->GetNumActiveGameplayEffect() == 1);
-	}
-
-	// DestComponent->PrintAllGameplayEffects();
-
-	World->EditorDestroyActor(SourceActor, false);
-	World->EditorDestroyActor(DestActor, false);
-
-	return true;
-}
-
-bool GameplayEffectsTest_SnapshotBuffDestroy(UWorld *World, FAutomationTestBase * Test)
-{
-	/**
-	*	Apply a perm health reduction that is buffed by an outgoing GE buff. Then destroy the buff and see what happens to the perm applied Ge.
-	*/
-
-	const float StartHealth = 100.f;
-	const float DamageValue = -5.f;
-	const float DamageBuffMultiplier = 2.f;
-
-	AAbilitySystemTestPawn *SourceActor = World->SpawnActor<AAbilitySystemTestPawn>();
-	AAbilitySystemTestPawn *DestActor = World->SpawnActor<AAbilitySystemTestPawn>();
-
-	UProperty *HealthProperty = FindFieldChecked<UProperty>(UAbilitySystemTestAttributeSet::StaticClass(), GET_MEMBER_NAME_CHECKED(UAbilitySystemTestAttributeSet, Health));
-
-	UAbilitySystemComponent * SourceComponent = SourceActor->GetAbilitySystemComponent();
-	UAbilitySystemComponent * DestComponent = DestActor->GetAbilitySystemComponent();
-	
-	SourceComponent->GetSet<UAbilitySystemTestAttributeSet>()->Health = StartHealth;
-	DestComponent->GetSet<UAbilitySystemTestAttributeSet>()->Health = StartHealth;
-
-	// Apply "Damage" but set to INFINITE_DURATION
-	// (An odd example for damage, but would make sense for something like run speed, etc)
-	FActiveGameplayEffectHandle AppliedHandle;
-	FActiveGameplayEffectHandle AppliedHandleBuff;
-
-	{
-		ABILITY_LOG_SCOPE(TEXT("Apply Damage Buff"))
-
-		UGameplayEffect * BuffEffect = Cast<UGameplayEffect>(StaticConstructObject(UGameplayEffect::StaticClass(), GetTransientPackage(), FName(TEXT("BuffOutgoingDamage"))));
-		BuffEffect->Modifiers.SetNum(1);
-		BuffEffect->Modifiers[0].Magnitude.SetValue(DamageBuffMultiplier);
-		BuffEffect->Modifiers[0].ModifierType = EGameplayMod::OutgoingGE;
-		BuffEffect->Modifiers[0].ModifierOp = EGameplayModOp::Multiplicitive;
-		BuffEffect->Modifiers[0].Attribute.SetUProperty(HealthProperty);
-		BuffEffect->Duration.Value = UGameplayEffect::INFINITE_DURATION;
-		BuffEffect->CopyPolicy = EGameplayEffectCopyPolicy::AlwaysSnapshot;		// Always snapshot (though the default for outgoing should already be snapshot - but this could change per project at some point)
-
-		// Apply to target
-		AppliedHandleBuff = SourceComponent->ApplyGameplayEffectToTarget(BuffEffect, SourceComponent, 1.f);
-
-		Test->TestTrue(SKILL_TEST_TEXT("Number of Source GameplayEffect: %d == 1", SourceComponent->GetNumActiveGameplayEffect()), SourceComponent->GetNumActiveGameplayEffect() == 1);
-	}
-
-	{
-		ABILITY_LOG_SCOPE(TEXT("Apply Permanent (infinite duration) Damage"))
-
-		UGameplayEffect * BaseDmgEffect = Cast<UGameplayEffect>(StaticConstructObject(UGameplayEffect::StaticClass(), GetTransientPackage(), FName(TEXT("Damage"))));
-		BaseDmgEffect->Modifiers.SetNum(1);
-		BaseDmgEffect->Modifiers[0].Magnitude.SetValue(DamageValue);
-		BaseDmgEffect->Modifiers[0].ModifierType = EGameplayMod::Attribute;
-		BaseDmgEffect->Modifiers[0].ModifierOp = EGameplayModOp::Additive;
-		BaseDmgEffect->Modifiers[0].Attribute.SetUProperty(HealthProperty);
-		BaseDmgEffect->Modifiers[0].OwnedTags.AddTag(IGameplayTagsModule::RequestGameplayTag(FName(TEXT("Damage.Basic"))));
-		BaseDmgEffect->Duration.Value = UGameplayEffect::INFINITE_DURATION;
-
-		// Apply to target
-		AppliedHandle = SourceComponent->ApplyGameplayEffectToTarget(BaseDmgEffect, DestComponent, 1.f);
-
-		float ExpectedValue = (StartHealth + (DamageValue * DamageBuffMultiplier));
-		float ActualValue = DestComponent->GetSet<UAbilitySystemTestAttributeSet>()->Health;
-
-		Test->TestTrue(SKILL_TEST_TEXT("Damaged Applied. ActualValue: %.2f. ExpectedValue: %.2f.", ActualValue, ExpectedValue), (ActualValue == ExpectedValue));
-		Test->TestTrue(SKILL_TEST_TEXT("NumberOfActive GameplayEffects GameplayEffects %d == 1", DestComponent->GetNumActiveGameplayEffect()), DestComponent->GetNumActiveGameplayEffect() == 1);
-	}
-
-	// Remove the buff GE
-	{
-		bool RemovedEffect = SourceComponent->RemoveActiveGameplayEffect(AppliedHandleBuff);
-		
-		// Check health again - it shouldn't have changed!
-		float ExpectedValue = (StartHealth + (DamageValue * DamageBuffMultiplier));
-		float ActualValue = DestComponent->GetSet<UAbilitySystemTestAttributeSet>()->Health;
-
-		// Confirm we regained healtht
-		Test->TestTrue(SKILL_TEST_TEXT("After Buff Removal. ActualValue: %.2f. ExpectedValue: %.2f. RemovedEffecte: %d", ActualValue, ExpectedValue, RemovedEffect), (ActualValue == ExpectedValue));
-
-		// Confirm number of GEs
-		Test->TestTrue(SKILL_TEST_TEXT("Dest Number of GameplayEffects=%d", SourceComponent->GetNumActiveGameplayEffect()), SourceComponent->GetNumActiveGameplayEffect() == 0);
-		Test->TestTrue(SKILL_TEST_TEXT("Src Number of GameplayEffects=%d", DestComponent->GetNumActiveGameplayEffect()), DestComponent->GetNumActiveGameplayEffect() == 1);
-	}
-
-	// DestComponent->PrintAllGameplayEffects();
-
-	World->EditorDestroyActor(SourceActor, false);
-	World->EditorDestroyActor(DestActor, false);
-
-	return true;
-}
-
-bool GameplayEffectsTest_DurationBuff(UWorld *World, FAutomationTestBase * Test)
-{
-	/**
-	*	Tests duration buff and debuffs. Also tests canceling duration buffs.
-	*/
-
-	const float StartHealth = 100.f;
-	const float DamageValue = -5.f;
-	const float BaseDuration = 2.f;
-	const float DurationBuff = 1.f;
-	const float DurationDebuff = -1.f;
-
-	AAbilitySystemTestPawn *SourceActor = World->SpawnActor<AAbilitySystemTestPawn>();
-	AAbilitySystemTestPawn *DestActor = World->SpawnActor<AAbilitySystemTestPawn>();
-
-	UProperty *HealthProperty = FindFieldChecked<UProperty>(UAbilitySystemTestAttributeSet::StaticClass(), GET_MEMBER_NAME_CHECKED(UAbilitySystemTestAttributeSet, Health));
-
-	UAbilitySystemComponent * SourceComponent = SourceActor->GetAbilitySystemComponent();
-	UAbilitySystemComponent * DestComponent = DestActor->GetAbilitySystemComponent();
-
-	SourceComponent->GetSet<UAbilitySystemTestAttributeSet>()->Health = StartHealth;
-	DestComponent->GetSet<UAbilitySystemTestAttributeSet>()->Health = StartHealth;
-
-	// Apply Damage with 2 duration
-
-	FActiveGameplayEffectHandle AppliedDamageHandle;
-	FActiveGameplayEffectHandle AppliedDurationHandle;
-	//
-	// Duration Debuff
-	//
-	{
-		ABILITY_LOG_SCOPE(TEXT("Apply Damage mod that lasts 2 seconds"))
-
-		UGameplayEffect * BaseDmgEffect = Cast<UGameplayEffect>(StaticConstructObject(UGameplayEffect::StaticClass(), GetTransientPackage(), FName(TEXT("Damage"))));
-		BaseDmgEffect->Modifiers.SetNum(1);
-		BaseDmgEffect->Modifiers[0].Magnitude.SetValue(DamageValue);
-		BaseDmgEffect->Modifiers[0].ModifierType = EGameplayMod::Attribute;
-		BaseDmgEffect->Modifiers[0].ModifierOp = EGameplayModOp::Additive;
-		BaseDmgEffect->Modifiers[0].Attribute.SetUProperty(HealthProperty);
-		BaseDmgEffect->Modifiers[0].OwnedTags.AddTag(IGameplayTagsModule::RequestGameplayTag(FName(TEXT("Damage.Basic"))));
-		BaseDmgEffect->Duration.SetValue(BaseDuration);
-
-		// Apply to target
-		AppliedDamageHandle = SourceComponent->ApplyGameplayEffectToTarget(BaseDmgEffect, DestComponent, 1.f);
-
-		float ExpectedHealthValue = (StartHealth + (DamageValue));
-		float ActualHealthValue = DestComponent->GetSet<UAbilitySystemTestAttributeSet>()->Health;
-		float ExpectedDuration = BaseDuration;
-		float ActualDuration = DestComponent->GetGameplayEffectDuration(AppliedDamageHandle);
-		
-		Test->TestTrue(SKILL_TEST_TEXT("Damage Applied. %.2f == %.2f", ActualHealthValue, ExpectedHealthValue), ActualHealthValue == ExpectedHealthValue);
-		
-		Test->TestTrue(SKILL_TEST_TEXT("Duration of GameplayEffect. %.2f == %.2f", ActualDuration, ExpectedDuration), ActualDuration == ExpectedDuration);
-	}
-
-	GameplayTest_TickWorld(World, SMALL_NUMBER); // start the effect ticking
-	GameplayTest_TickWorld(World, 0.5f);
-
-	// Debuff the duration of the effect
-	{
-		ABILITY_LOG_SCOPE(TEXT("Reduce damage mod during by 1 second"))
-
-		UGameplayEffect * DurationEffect = Cast<UGameplayEffect>(StaticConstructObject(UGameplayEffect::StaticClass(), GetTransientPackage(), FName(TEXT("Duration Debuff"))));
-		DurationEffect->Modifiers.SetNum(1);
-		DurationEffect->Modifiers[0].Magnitude.SetValue(DurationDebuff);
-		DurationEffect->Modifiers[0].ModifierType = EGameplayMod::ActiveGE;
-		DurationEffect->Modifiers[0].ModifierOp = EGameplayModOp::Additive;
-		DurationEffect->Modifiers[0].EffectType = EGameplayModEffect::Duration;
-		DurationEffect->Modifiers[0].Attribute.SetUProperty(HealthProperty);
-		DurationEffect->Modifiers[0].OwnedTags.AddTag(IGameplayTagsModule::RequestGameplayTag(FName(TEXT("Damage.Basic"))));
-		DurationEffect->Duration.SetValue(UGameplayEffect::INFINITE_DURATION);
-
-		// Apply to target
-		AppliedDurationHandle = SourceComponent->ApplyGameplayEffectToTarget(DurationEffect, DestComponent, 1.f);
-
-		float ExpectedDuration = BaseDuration + DurationDebuff;
-		float ActualDuration = DestComponent->GetGameplayEffectDuration(AppliedDamageHandle);
-
-		// Confirm that our duration changed
-		Test->TestTrue(SKILL_TEST_TEXT("Duration of GameplayEffect PostMod. %.2f == %.2f", ActualDuration, ExpectedDuration), ActualDuration == ExpectedDuration);
-	}
-
-	// tick beyond the new duration but not past the old duration
-	GameplayTest_TickWorld(World, 1.f);
-
-	{
-		float ExpectedHealthValue = StartHealth;
-		float ActualHealthValue = DestComponent->GetSet<UAbilitySystemTestAttributeSet>()->Health;
-
-		Test->TestTrue(SKILL_TEST_TEXT("Damage Applied. %.2f == %.2f", ActualHealthValue, ExpectedHealthValue), ActualHealthValue == ExpectedHealthValue);
-	}
-
-	//
-	// Duration Buff
-	//
-	{
-		ABILITY_LOG_SCOPE(TEXT("Apply Damage mod that lasts 2 seconds"))
-
-		UGameplayEffect * BaseDmgEffect = Cast<UGameplayEffect>(StaticConstructObject(UGameplayEffect::StaticClass(), GetTransientPackage(), FName(TEXT("Damage"))));
-		BaseDmgEffect->Modifiers.SetNum(1);
-		BaseDmgEffect->Modifiers[0].Magnitude.SetValue(DamageValue);
-		BaseDmgEffect->Modifiers[0].ModifierType = EGameplayMod::Attribute;
-		BaseDmgEffect->Modifiers[0].ModifierOp = EGameplayModOp::Additive;
-		BaseDmgEffect->Modifiers[0].Attribute.SetUProperty(HealthProperty);
-		BaseDmgEffect->Modifiers[0].OwnedTags.AddTag(IGameplayTagsModule::RequestGameplayTag(FName(TEXT("Damage.Basic"))));
-		BaseDmgEffect->Duration.SetValue(BaseDuration);
-
-		// Apply to target
-		AppliedDamageHandle = SourceComponent->ApplyGameplayEffectToTarget(BaseDmgEffect, DestComponent, 1.f);
-
-		float ExpectedHealthValue = (StartHealth + (DamageValue));
-		float ActualHealthValue = DestComponent->GetSet<UAbilitySystemTestAttributeSet>()->Health;
-		float ExpectedDuration = BaseDuration;
-		float ActualDuration = DestComponent->GetGameplayEffectDuration(AppliedDamageHandle);
-
-		Test->TestTrue(SKILL_TEST_TEXT("Damage Applied. %.2f == %.2f", ActualHealthValue, ExpectedHealthValue), ActualHealthValue == ExpectedHealthValue);
-
-		Test->TestTrue(SKILL_TEST_TEXT("Duration of GameplayEffect. %.2f == %.2f", ActualDuration, ExpectedDuration), ActualDuration == ExpectedDuration);
-	}
-
-	GameplayTest_TickWorld(World, SMALL_NUMBER); // start the effect ticking
-	GameplayTest_TickWorld(World, 0.5f);
-
-	// Buff the duration of the effect
-	{
-		ABILITY_LOG_SCOPE(TEXT("Increase damage mod during by 1 second"))
-
-		UGameplayEffect * DurationEffect = Cast<UGameplayEffect>(StaticConstructObject(UGameplayEffect::StaticClass(), GetTransientPackage(), FName(TEXT("Duration Buff"))));
-		DurationEffect->Modifiers.SetNum(1);
-		DurationEffect->Modifiers[0].Magnitude.SetValue(DurationBuff);
-		DurationEffect->Modifiers[0].ModifierType = EGameplayMod::ActiveGE;
-		DurationEffect->Modifiers[0].ModifierOp = EGameplayModOp::Additive;
-		DurationEffect->Modifiers[0].EffectType = EGameplayModEffect::Duration;
-		DurationEffect->Modifiers[0].Attribute.SetUProperty(HealthProperty);
-		DurationEffect->Modifiers[0].OwnedTags.AddTag(IGameplayTagsModule::RequestGameplayTag(FName(TEXT("Damage.Basic"))));
-		DurationEffect->Duration.SetValue(UGameplayEffect::INFINITE_DURATION);
-
-		// Apply to target
-		AppliedDurationHandle = SourceComponent->ApplyGameplayEffectToTarget(DurationEffect, DestComponent, 1.f);
-
-		float ExpectedDuration = BaseDuration + DurationBuff;
-		float ActualDuration = DestComponent->GetGameplayEffectDuration(AppliedDamageHandle);
-
-		// Confirm that our duration changed
-		Test->TestTrue(SKILL_TEST_TEXT("Duration of GameplayEffect PostMod. %.2f == %.2f", ActualDuration, ExpectedDuration), ActualDuration == ExpectedDuration);
-	}
-
-	// tick beyond the old duration but not past the new duration
-	GameplayTest_TickWorld(World, 2.f);
-
-	{
-		float ExpectedHealthValue = (StartHealth + (DamageValue));
-		float ActualHealthValue = DestComponent->GetSet<UAbilitySystemTestAttributeSet>()->Health;
-
-		Test->TestTrue(SKILL_TEST_TEXT("Damage Applied. %.2f == %.2f", ActualHealthValue, ExpectedHealthValue), ActualHealthValue == ExpectedHealthValue);
-	}
-
-	// tick past the new duration
-	GameplayTest_TickWorld(World, 1.f);
-
-	{
-		float ExpectedHealthValue = StartHealth;
-		float ActualHealthValue = DestComponent->GetSet<UAbilitySystemTestAttributeSet>()->Health;
-
-		Test->TestTrue(SKILL_TEST_TEXT("Damage Applied. %.2f == %.2f", ActualHealthValue, ExpectedHealthValue), ActualHealthValue == ExpectedHealthValue);
-	}
-
-
-	//
-	// Removing Duration buff
-	//
-	{
-		ABILITY_LOG_SCOPE(TEXT("Apply Damage mod that lasts 2 seconds"))
-
-		UGameplayEffect * BaseDmgEffect = Cast<UGameplayEffect>(StaticConstructObject(UGameplayEffect::StaticClass(), GetTransientPackage(), FName(TEXT("Damage"))));
-		BaseDmgEffect->Modifiers.SetNum(1);
-		BaseDmgEffect->Modifiers[0].Magnitude.SetValue(DamageValue);
-		BaseDmgEffect->Modifiers[0].ModifierType = EGameplayMod::Attribute;
-		BaseDmgEffect->Modifiers[0].ModifierOp = EGameplayModOp::Additive;
-		BaseDmgEffect->Modifiers[0].Attribute.SetUProperty(HealthProperty);
-		BaseDmgEffect->Modifiers[0].OwnedTags.AddTag(IGameplayTagsModule::RequestGameplayTag(FName(TEXT("Damage.Basic"))));
-		BaseDmgEffect->Duration.SetValue(BaseDuration);
-
-		// Apply to target
-		AppliedDamageHandle = SourceComponent->ApplyGameplayEffectToTarget(BaseDmgEffect, DestComponent, 1.f);
-
-		float ExpectedHealthValue = (StartHealth + (DamageValue));
-		float ActualHealthValue = DestComponent->GetSet<UAbilitySystemTestAttributeSet>()->Health;
-		float ExpectedDuration = BaseDuration;
-		float ActualDuration = DestComponent->GetGameplayEffectDuration(AppliedDamageHandle);
-
-		Test->TestTrue(SKILL_TEST_TEXT("Damage Applied. %.2f == %.2f", ActualHealthValue, ExpectedHealthValue), ActualHealthValue == ExpectedHealthValue);
-
-		Test->TestTrue(SKILL_TEST_TEXT("Duration of GameplayEffect. %.2f == %.2f", ActualDuration, ExpectedDuration), ActualDuration == ExpectedDuration);
-	}
-
-	GameplayTest_TickWorld(World, SMALL_NUMBER); // start the effect ticking
-	GameplayTest_TickWorld(World, 0.5f);
-
-	// Buff the duration of the effect
-	{
-		ABILITY_LOG_SCOPE(TEXT("Increase damage mod during by 1 second"))
-
-		UGameplayEffect * DurationEffect = Cast<UGameplayEffect>(StaticConstructObject(UGameplayEffect::StaticClass(), GetTransientPackage(), FName(TEXT("Duration Buff"))));
-		DurationEffect->Modifiers.SetNum(1);
-		DurationEffect->Modifiers[0].Magnitude.SetValue(DurationBuff);
-		DurationEffect->Modifiers[0].ModifierType = EGameplayMod::ActiveGE;
-		DurationEffect->Modifiers[0].ModifierOp = EGameplayModOp::Additive;
-		DurationEffect->Modifiers[0].EffectType = EGameplayModEffect::Duration;
-		DurationEffect->Modifiers[0].Attribute.SetUProperty(HealthProperty);
-		DurationEffect->Modifiers[0].OwnedTags.AddTag(IGameplayTagsModule::RequestGameplayTag(FName(TEXT("Damage.Basic"))));
-		DurationEffect->Duration.SetValue(UGameplayEffect::INFINITE_DURATION);
-
-		// Apply to target
-		AppliedDurationHandle = SourceComponent->ApplyGameplayEffectToTarget(DurationEffect, DestComponent, 1.f);
-
-		float ExpectedDuration = BaseDuration + DurationBuff;
-		float ActualDuration = DestComponent->GetGameplayEffectDuration(AppliedDamageHandle);
-
-		// Confirm that our duration changed
-		Test->TestTrue(SKILL_TEST_TEXT("Duration of GameplayEffect PostMod. %.2f == %.2f", ActualDuration, ExpectedDuration), ActualDuration == ExpectedDuration);
-	}
-
-	// tick beyond the old duration but not past the new duration
-	GameplayTest_TickWorld(World, 2.f);
-
-	// Remove the duration effect and see if the duration goes back to the original duration
-	{
-		bool RemovedEffect = DestComponent->RemoveActiveGameplayEffect(AppliedDurationHandle);
-
-		float ExpectedDuration = BaseDuration;
-		float ActualDuration = DestComponent->GetGameplayEffectDuration(AppliedDamageHandle);
-		float ExpectedHealthValue = StartHealth + DamageValue;
-		float ActualHealthValue = DestComponent->GetSet<UAbilitySystemTestAttributeSet>()->Health;
-
-		Test->TestTrue(SKILL_TEST_TEXT("Duration of GameplayEffect Post Mod Remove. %.2f == %.2f", ActualDuration, ExpectedDuration), ActualDuration == ExpectedDuration);
-		Test->TestTrue(SKILL_TEST_TEXT("Damage Applied. %.2f == %.2f", ActualHealthValue, ExpectedHealthValue), ActualHealthValue == ExpectedHealthValue);
-	}
-
-	// small tick so that we remain between the unmodified and modified duration
-	GameplayTest_TickWorld(World, KINDA_SMALL_NUMBER); // moves the new timer to the active list
-	GameplayTest_TickWorld(World, KINDA_SMALL_NUMBER);
-
-	{
-		float ExpectedHealthValue = StartHealth;
-		float ActualHealthValue = DestComponent->GetSet<UAbilitySystemTestAttributeSet>()->Health;
-
-		Test->TestTrue(SKILL_TEST_TEXT("Damage Applied. %.2f == %.2f", ActualHealthValue, ExpectedHealthValue), ActualHealthValue == ExpectedHealthValue);
-	}
-
-	World->EditorDestroyActor(SourceActor, false);
-	World->EditorDestroyActor(DestActor, false);
-
-	return true;
-}
-
-bool GameplayEffectsTest_DamageBuffBuff_Basic(UWorld *World, FAutomationTestBase * Test)
-{
-	/**
-	*	Buff a Damage Buff, then apply damage
-	*/
-
-	const float StartHealth = 100.f;
-	const float DamageValue = -5.f;
-	const float DamageBuffMultiplier = 2.f;			// Damage is buff and multiplied by 2
-	const float DamageBuffMultiplierBonus = 1.f;	// The above multiplier receives a +1 bonus (we expect a final multiplier of 3)
-
-	AAbilitySystemTestPawn *SourceActor = World->SpawnActor<AAbilitySystemTestPawn>();
-	AAbilitySystemTestPawn *DestActor = World->SpawnActor<AAbilitySystemTestPawn>();
-
-	UProperty *HealthProperty = FindFieldChecked<UProperty>(UAbilitySystemTestAttributeSet::StaticClass(), GET_MEMBER_NAME_CHECKED(UAbilitySystemTestAttributeSet, Health));
-
-	UAbilitySystemComponent * SourceComponent = SourceActor->GetAbilitySystemComponent();
-	UAbilitySystemComponent * DestComponent = DestActor->GetAbilitySystemComponent();
-	SourceComponent->GetSet<UAbilitySystemTestAttributeSet>()->Health = StartHealth;
-	DestComponent->GetSet<UAbilitySystemTestAttributeSet>()->Health = StartHealth;
-
-	// Apply "Damage" but set to INFINITE_DURATION
-	// (An odd example for damage, but would make sense for something like run speed, etc)
-	FActiveGameplayEffectHandle BuffBuffHandle;
-	FActiveGameplayEffectHandle BuffHandle;
-	FActiveGameplayEffectHandle DamageHandle;
-
-	{
-		ABILITY_LOG_SCOPE(TEXT("Apply Buff Buff"))
-
-		// Here we are choosing to do this by adding an perm IncomingGE buff first. There are other ways to do this.
-
-		UGameplayEffect * BuffEffect = Cast<UGameplayEffect>(StaticConstructObject(UGameplayEffect::StaticClass(), GetTransientPackage(), FName(TEXT("BuffBuff"))));
-		BuffEffect->Modifiers.SetNum(1);
-		BuffEffect->Modifiers[0].Magnitude.SetValue(DamageBuffMultiplierBonus);
-		BuffEffect->Modifiers[0].ModifierType = EGameplayMod::IncomingGE;
-		BuffEffect->Modifiers[0].ModifierOp = EGameplayModOp::Additive;
-		BuffEffect->Modifiers[0].Attribute.SetUProperty(HealthProperty);
-		BuffEffect->Duration.Value = UGameplayEffect::INFINITE_DURATION;
-		BuffEffect->CopyPolicy = EGameplayEffectCopyPolicy::AlwaysLink;
-
-		// Apply to target
-		BuffBuffHandle = SourceComponent->ApplyGameplayEffectToTarget(BuffEffect, SourceComponent, 1.f);
-
-		Test->TestTrue(SKILL_TEST_TEXT("Number of Source GameplayEffect: %d == 1", SourceComponent->GetNumActiveGameplayEffect()), SourceComponent->GetNumActiveGameplayEffect() == 1);
-	}
-
-	{
-		ABILITY_LOG_SCOPE(TEXT("Apply Damage Buff"))
-
-		UGameplayEffect * BuffEffect = Cast<UGameplayEffect>(StaticConstructObject(UGameplayEffect::StaticClass(), GetTransientPackage(), FName(TEXT("Buff"))));
-		BuffEffect->Modifiers.SetNum(1);
-		BuffEffect->Modifiers[0].Magnitude.SetValue(DamageBuffMultiplier);
-		BuffEffect->Modifiers[0].ModifierType = EGameplayMod::OutgoingGE;
-		BuffEffect->Modifiers[0].ModifierOp = EGameplayModOp::Multiplicitive;
-		BuffEffect->Modifiers[0].Attribute.SetUProperty(HealthProperty);
-		BuffEffect->Duration.Value = UGameplayEffect::INFINITE_DURATION;
-		BuffEffect->CopyPolicy = EGameplayEffectCopyPolicy::AlwaysSnapshot;
-
-		// Apply to target
-		BuffHandle = SourceComponent->ApplyGameplayEffectToTarget(BuffEffect, SourceComponent, 1.f);
-		Test->TestTrue(SKILL_TEST_TEXT("Number of Source GameplayEffect: %d == 2", SourceComponent->GetNumActiveGameplayEffect()), SourceComponent->GetNumActiveGameplayEffect() == 2);
-
-		// Check that the buff was buffed
-		float ExpectedBuffMagnitude = DamageBuffMultiplier + DamageBuffMultiplierBonus;
-		float ActualBuffMagnitude = SourceComponent->GetGameplayEffectMagnitude(BuffHandle, FGameplayAttribute(HealthProperty));
-
-		Test->TestTrue(SKILL_TEST_TEXT("Buff Applied. Check Magnitude. ActualValue: %.2f. ExpectedValue: %.2f.", ActualBuffMagnitude, ExpectedBuffMagnitude), (ActualBuffMagnitude == ExpectedBuffMagnitude));
-	}
-
-	{
-		ABILITY_LOG_SCOPE(TEXT("Apply Permanent (infinite duration) Damage"))
-
-		UGameplayEffect * BaseDmgEffect = Cast<UGameplayEffect>(StaticConstructObject(UGameplayEffect::StaticClass(), GetTransientPackage(), FName(TEXT("Damage"))));
-		BaseDmgEffect->Modifiers.SetNum(1);
-		BaseDmgEffect->Modifiers[0].Magnitude.SetValue(DamageValue);
-		BaseDmgEffect->Modifiers[0].ModifierType = EGameplayMod::Attribute;
-		BaseDmgEffect->Modifiers[0].ModifierOp = EGameplayModOp::Additive;
-		BaseDmgEffect->Modifiers[0].Attribute.SetUProperty(HealthProperty);
-		BaseDmgEffect->Modifiers[0].OwnedTags.AddTag(IGameplayTagsModule::RequestGameplayTag(FName(TEXT("Damage.Basic"))));
-		BaseDmgEffect->Duration.Value = UGameplayEffect::INFINITE_DURATION;
-
-		// Apply to target
-		DamageHandle = SourceComponent->ApplyGameplayEffectToTarget(BaseDmgEffect, DestComponent, 1.f);
-
-		float ExpectedValue = (StartHealth + (DamageValue * (DamageBuffMultiplier + DamageBuffMultiplierBonus)));
-		float ActualValue = DestComponent->GetSet<UAbilitySystemTestAttributeSet>()->Health;
-
-		Test->TestTrue(SKILL_TEST_TEXT("Damaged Applied. ActualValue: %.2f. ExpectedValue: %.2f.", ActualValue, ExpectedValue), (ActualValue == ExpectedValue));
-		Test->TestTrue(SKILL_TEST_TEXT("NumberOfActive Dest GameplayEffects %d == 1", DestComponent->GetNumActiveGameplayEffect()), DestComponent->GetNumActiveGameplayEffect() == 1);
-	}
-
-	{
-		// Clear DependantsUpdates stat
-#if SKILL_SYSTEM_AGGREGATOR_DEBUG
-		FAggregator::AllocationStats.DependantsUpdated = 0;
-#endif
-		// Remove the original buff buff
-		bool RemovedEffect = SourceComponent->RemoveActiveGameplayEffect(BuffBuffHandle);
-
-		float ActualValue = DestComponent->GetSet<UAbilitySystemTestAttributeSet>()->Health;
-		float ExpectedValue = (StartHealth + (DamageValue * (DamageBuffMultiplier + DamageBuffMultiplierBonus)));
-
-		Test->TestTrue(SKILL_TEST_TEXT("Damaged Applied. ActualValue: %.2f. ExpectedValue: %.2f.", ActualValue, ExpectedValue), (ActualValue == ExpectedValue));
-
-		// Check that we updated exactly 1 dependant by removing the BuffBuff (it should have forced the Buff to update - but not the applied Damage GE)
-#if SKILL_SYSTEM_AGGREGATOR_DEBUG
-		Test->TestTrue(SKILL_TEST_TEXT("DependantsUpdated %d == 1", FAggregator::AllocationStats.DependantsUpdated), (FAggregator::AllocationStats.DependantsUpdated == 1));
-#endif
-	}
-
-	{
-		ABILITY_LOG_SCOPE(TEXT("Remove Buff"));
-#if SKILL_SYSTEM_AGGREGATOR_DEBUG
-		FAggregator::AllocationStats.DependantsUpdated = 0;
-#endif
-		// Remove the buff
-		bool RemovedEffect = SourceComponent->RemoveActiveGameplayEffect(BuffHandle);
-
-		// No change to health since we applied a snapshot of the buff to the damage GE
-		float ExpectedValue = (StartHealth + (DamageValue * (DamageBuffMultiplier + DamageBuffMultiplierBonus)));
-		float ActualValue = DestComponent->GetSet<UAbilitySystemTestAttributeSet>()->Health;
-
-		Test->TestTrue(SKILL_TEST_TEXT("Damaged Applied. ActualValue: %.2f. ExpectedValue: %.2f.", ActualValue, ExpectedValue), (ActualValue == ExpectedValue));
-
-		// Check that we updated 0 dependants
-#if SKILL_SYSTEM_AGGREGATOR_DEBUG
-		Test->TestTrue(SKILL_TEST_TEXT("DependantsUpdated %d == 0", FAggregator::AllocationStats.DependantsUpdated), (FAggregator::AllocationStats.DependantsUpdated == 0));
-#endif
-	}
-
-	// DestComponent->PrintAllGameplayEffects();
-
-	World->EditorDestroyActor(SourceActor, false);
-	World->EditorDestroyActor(DestActor, false);
-
-	return true;
-}
-
-
-
-bool GameplayEffectsTest_DamageBuffBuff_FullLink(UWorld *World, FAutomationTestBase * Test)
-{
-	/**
-	*	Buff a Damage Buff, then apply damage
-	*/
-
-	const float StartHealth = 100.f;
-	const float DamageValue = -5.f;
-	const float DamageBuffMultiplier = 2.f;			// Damage is buff and multiplied by 2
-	const float DamageBuffMultiplierBonus = 1.f;	// The above multiplier receives a +1 bonus (we expect a final multiplier of 3)
-
-	AAbilitySystemTestPawn *SourceActor = World->SpawnActor<AAbilitySystemTestPawn>();
-	AAbilitySystemTestPawn *DestActor = World->SpawnActor<AAbilitySystemTestPawn>();
-
-	UProperty *HealthProperty = FindFieldChecked<UProperty>(UAbilitySystemTestAttributeSet::StaticClass(), GET_MEMBER_NAME_CHECKED(UAbilitySystemTestAttributeSet, Health));
-
-	UAbilitySystemComponent * SourceComponent = SourceActor->GetAbilitySystemComponent();
-	UAbilitySystemComponent * DestComponent = DestActor->GetAbilitySystemComponent();
-	SourceComponent->GetSet<UAbilitySystemTestAttributeSet>()->Health = StartHealth;
-	DestComponent->GetSet<UAbilitySystemTestAttributeSet>()->Health = StartHealth;
-
-	// Apply "Damage" but set to INFINITE_DURATION
-	// (An odd example for damage, but would make sense for something like run speed, etc)
-	FActiveGameplayEffectHandle BuffBuffHandle;
-	FActiveGameplayEffectHandle BuffHandle;
-	FActiveGameplayEffectHandle DamageHandle;
-
-	{
-		ABILITY_LOG_SCOPE(TEXT("Apply Buff Buff"))
-
-		// Here we are choosing to do this by adding an perm IncomingGE buff first. There are other ways to do this.
-
-		UGameplayEffect * BuffEffect = Cast<UGameplayEffect>(StaticConstructObject(UGameplayEffect::StaticClass(), GetTransientPackage(), FName(TEXT("BuffBuff"))));
-		BuffEffect->Modifiers.SetNum(1);
-		BuffEffect->Modifiers[0].Magnitude.SetValue(DamageBuffMultiplierBonus);
-		BuffEffect->Modifiers[0].ModifierType = EGameplayMod::IncomingGE;
-		BuffEffect->Modifiers[0].ModifierOp = EGameplayModOp::Additive;
-		BuffEffect->Modifiers[0].Attribute.SetUProperty(HealthProperty);
-		BuffEffect->Duration.Value = UGameplayEffect::INFINITE_DURATION;
-		BuffEffect->CopyPolicy = EGameplayEffectCopyPolicy::AlwaysLink;
-
-		// Apply to target
-		BuffBuffHandle = SourceComponent->ApplyGameplayEffectToTarget(BuffEffect, SourceComponent, 1.f);
-
-		Test->TestTrue(SKILL_TEST_TEXT("Number of Source GameplayEffect: %d == 1", SourceComponent->GetNumActiveGameplayEffect()), SourceComponent->GetNumActiveGameplayEffect() == 1);
-	}
-
-	{
-		ABILITY_LOG_SCOPE(TEXT("Apply Damage Buff"))
-
-		UGameplayEffect * BuffEffect = Cast<UGameplayEffect>(StaticConstructObject(UGameplayEffect::StaticClass(), GetTransientPackage(), FName(TEXT("Buff"))));
-		BuffEffect->Modifiers.SetNum(1);
-		BuffEffect->Modifiers[0].Magnitude.SetValue(DamageBuffMultiplier);
-		BuffEffect->Modifiers[0].ModifierType = EGameplayMod::OutgoingGE;
-		BuffEffect->Modifiers[0].ModifierOp = EGameplayModOp::Multiplicitive;
-		BuffEffect->Modifiers[0].Attribute.SetUProperty(HealthProperty);
-		BuffEffect->Duration.Value = UGameplayEffect::INFINITE_DURATION;
-		BuffEffect->CopyPolicy = EGameplayEffectCopyPolicy::AlwaysLink;
-
-		// Apply to target
-		BuffHandle = SourceComponent->ApplyGameplayEffectToTarget(BuffEffect, SourceComponent, 1.f);
-		Test->TestTrue(SKILL_TEST_TEXT("Number of Source GameplayEffect: %d == 2", SourceComponent->GetNumActiveGameplayEffect()), SourceComponent->GetNumActiveGameplayEffect() == 2);
-
-		// Check that the buff was buffed
-		float ExpectedBuffMagnitude = DamageBuffMultiplier + DamageBuffMultiplierBonus;
-		float ActualBuffMagnitude = SourceComponent->GetGameplayEffectMagnitude(BuffHandle, FGameplayAttribute(HealthProperty));
-
-		Test->TestTrue(SKILL_TEST_TEXT("Buff Applied. Check Magnitude. ActualValue: %.2f. ExpectedValue: %.2f.", ActualBuffMagnitude, ExpectedBuffMagnitude), (ActualBuffMagnitude == ExpectedBuffMagnitude));
-	}
-
-	{
-		ABILITY_LOG_SCOPE(TEXT("Apply Permanent (infinite duration) Damage"))
-
-		UGameplayEffect * BaseDmgEffect = Cast<UGameplayEffect>(StaticConstructObject(UGameplayEffect::StaticClass(), GetTransientPackage(), FName(TEXT("Damage"))));
-		BaseDmgEffect->Modifiers.SetNum(1);
-		BaseDmgEffect->Modifiers[0].Magnitude.SetValue(DamageValue);
-		BaseDmgEffect->Modifiers[0].ModifierType = EGameplayMod::Attribute;
-		BaseDmgEffect->Modifiers[0].ModifierOp = EGameplayModOp::Additive;
-		BaseDmgEffect->Modifiers[0].Attribute.SetUProperty(HealthProperty);
-		BaseDmgEffect->Modifiers[0].OwnedTags.AddTag(IGameplayTagsModule::RequestGameplayTag(FName(TEXT("Damage.Basic"))));
-		BaseDmgEffect->Duration.Value = UGameplayEffect::INFINITE_DURATION;
-
-		// Apply to target
-		DamageHandle = SourceComponent->ApplyGameplayEffectToTarget(BaseDmgEffect, DestComponent, 1.f);
-
-		float ExpectedValue = (StartHealth + (DamageValue * (DamageBuffMultiplier + DamageBuffMultiplierBonus)));
-		float ActualValue = DestComponent->GetSet<UAbilitySystemTestAttributeSet>()->Health;
-
-		Test->TestTrue(SKILL_TEST_TEXT("Damaged Applied. ActualValue: %.2f. ExpectedValue: %.2f.", ActualValue, ExpectedValue), (ActualValue == ExpectedValue));
-		Test->TestTrue(SKILL_TEST_TEXT("NumberOfActive Dest GameplayEffects %d == 1", DestComponent->GetNumActiveGameplayEffect()), DestComponent->GetNumActiveGameplayEffect() == 1);
-	}
-
-	{
-		// Clear DependantsUpdates stat
-#if SKILL_SYSTEM_AGGREGATOR_DEBUG
-		FAggregator::AllocationStats.DependantsUpdated = 0;
-#endif
-		// Remove the original buff buff
-		bool RemovedEffect = SourceComponent->RemoveActiveGameplayEffect(BuffBuffHandle);
-
-		float ActualValue = DestComponent->GetSet<UAbilitySystemTestAttributeSet>()->Health;
-		float ExpectedValue = (StartHealth + (DamageValue * (DamageBuffMultiplier)));
-
-		Test->TestTrue(SKILL_TEST_TEXT("Damaged Applied. ActualValue: %.2f. ExpectedValue: %.2f.", ActualValue, ExpectedValue), (ActualValue == ExpectedValue));
-
-		// Check that we updated exactly 3 dependant by removing the BuffBuff (it should have forced the Buff, the applied damage GE, and the attribute aggregator to update)
-#if SKILL_SYSTEM_AGGREGATOR_DEBUG
-		Test->TestTrue(SKILL_TEST_TEXT("DependantsUpdated %d == 2", FAggregator::AllocationStats.DependantsUpdated), (FAggregator::AllocationStats.DependantsUpdated == 3));
-#endif
-	}
-
-	{
-		ABILITY_LOG_SCOPE(TEXT("Remove Buff"));
-#if SKILL_SYSTEM_AGGREGATOR_DEBUG
-		FAggregator::AllocationStats.DependantsUpdated = 0;
-#endif
-		// Remove the buff
-		bool RemovedEffect = SourceComponent->RemoveActiveGameplayEffect(BuffHandle);
-
-		float ActualValue = DestComponent->GetSet<UAbilitySystemTestAttributeSet>()->Health;
-		float ExpectedValue = StartHealth + DamageValue;
-
-		Test->TestTrue(SKILL_TEST_TEXT("Damaged Applied. ActualValue: %.2f. ExpectedValue: %.2f.", ActualValue, ExpectedValue), (ActualValue == ExpectedValue));
-
-		// Check that we updated 2 dependants - the damage GE and the attribute aggregator
-#if SKILL_SYSTEM_AGGREGATOR_DEBUG
-		Test->TestTrue(SKILL_TEST_TEXT("DependantsUpdated %d == 2", FAggregator::AllocationStats.DependantsUpdated), (FAggregator::AllocationStats.DependantsUpdated == 2));
-#endif
-	}
-
-	// DestComponent->PrintAllGameplayEffects();
-
-	World->EditorDestroyActor(SourceActor, false);
-	World->EditorDestroyActor(DestActor, false);
-
-	return true;
-}
-
-bool GameplayEffectsTest_DamageBuffBuff_FullSnapshot(UWorld *World, FAutomationTestBase * Test)
-{
-	/**
-	*	Buff a Damage Buff, then apply damage
-	*/
-
-	const float StartHealth = 100.f;
-	const float DamageValue = -5.f;
-	const float DamageBuffMultiplier = 2.f;			// Damage is buff and multiplied by 2
-	const float DamageBuffMultiplierBonus = 1.f;	// The above multiplier receives a +1 bonus (we expect a final multiplier of 3)
-
-	AAbilitySystemTestPawn *SourceActor = World->SpawnActor<AAbilitySystemTestPawn>();
-	AAbilitySystemTestPawn *DestActor = World->SpawnActor<AAbilitySystemTestPawn>();
-
-	UProperty *HealthProperty = FindFieldChecked<UProperty>(UAbilitySystemTestAttributeSet::StaticClass(), GET_MEMBER_NAME_CHECKED(UAbilitySystemTestAttributeSet, Health));
-
-	UAbilitySystemComponent * SourceComponent = SourceActor->GetAbilitySystemComponent();
-	UAbilitySystemComponent * DestComponent = DestActor->GetAbilitySystemComponent();
-	SourceComponent->GetSet<UAbilitySystemTestAttributeSet>()->Health = StartHealth;
-	DestComponent->GetSet<UAbilitySystemTestAttributeSet>()->Health = StartHealth;
-
-	// Apply "Damage" but set to INFINITE_DURATION
-	// (An odd example for damage, but would make sense for something like run speed, etc)
-	FActiveGameplayEffectHandle BuffBuffHandle;
-	FActiveGameplayEffectHandle BuffHandle;
-	FActiveGameplayEffectHandle DamageHandle;
-
-	{
-		ABILITY_LOG_SCOPE(TEXT("Apply Buff Buff"))
-
-		// Here we are choosing to do this by adding an perm IncomingGE buff first. There are other ways to do this.
-
-		UGameplayEffect * BuffEffect = Cast<UGameplayEffect>(StaticConstructObject(UGameplayEffect::StaticClass(), GetTransientPackage(), FName(TEXT("BuffBuff"))));
-		BuffEffect->Modifiers.SetNum(1);
-		BuffEffect->Modifiers[0].Magnitude.SetValue(DamageBuffMultiplierBonus);
-		BuffEffect->Modifiers[0].ModifierType = EGameplayMod::IncomingGE;
-		BuffEffect->Modifiers[0].ModifierOp = EGameplayModOp::Additive;
-		BuffEffect->Modifiers[0].Attribute.SetUProperty(HealthProperty);
-		BuffEffect->Duration.Value = UGameplayEffect::INFINITE_DURATION;
-		BuffEffect->CopyPolicy = EGameplayEffectCopyPolicy::AlwaysSnapshot;
-
-		// Apply to target
-		BuffBuffHandle = SourceComponent->ApplyGameplayEffectToTarget(BuffEffect, SourceComponent, 1.f);
-
-		Test->TestTrue(SKILL_TEST_TEXT("Number of Source GameplayEffect: %d == 1", SourceComponent->GetNumActiveGameplayEffect()), SourceComponent->GetNumActiveGameplayEffect() == 1);
-	}
-
-	{
-		ABILITY_LOG_SCOPE(TEXT("Apply Damage Buff"))
-
-		UGameplayEffect * BuffEffect = Cast<UGameplayEffect>(StaticConstructObject(UGameplayEffect::StaticClass(), GetTransientPackage(), FName(TEXT("Buff"))));
-		BuffEffect->Modifiers.SetNum(1);
-		BuffEffect->Modifiers[0].Magnitude.SetValue(DamageBuffMultiplier);
-		BuffEffect->Modifiers[0].ModifierType = EGameplayMod::OutgoingGE;
-		BuffEffect->Modifiers[0].ModifierOp = EGameplayModOp::Multiplicitive;
-		BuffEffect->Modifiers[0].Attribute.SetUProperty(HealthProperty);
-		BuffEffect->Duration.Value = UGameplayEffect::INFINITE_DURATION;
-		BuffEffect->CopyPolicy = EGameplayEffectCopyPolicy::AlwaysSnapshot;
-
-		// Apply to target
-		BuffHandle = SourceComponent->ApplyGameplayEffectToTarget(BuffEffect, SourceComponent, 1.f);
-		Test->TestTrue(SKILL_TEST_TEXT("Number of Source GameplayEffect: %d == 2", SourceComponent->GetNumActiveGameplayEffect()), SourceComponent->GetNumActiveGameplayEffect() == 2);
-
-		// Check that the buff was buffed
-		float ExpectedBuffMagnitude = DamageBuffMultiplier + DamageBuffMultiplierBonus;
-		float ActualBuffMagnitude = SourceComponent->GetGameplayEffectMagnitude(BuffHandle, FGameplayAttribute(HealthProperty));
-
-		Test->TestTrue(SKILL_TEST_TEXT("Buff Applied. Check Magnitude. ActualValue: %.2f. ExpectedValue: %.2f.", ActualBuffMagnitude, ExpectedBuffMagnitude), (ActualBuffMagnitude == ExpectedBuffMagnitude));
-	}
-
-	{
-		ABILITY_LOG_SCOPE(TEXT("Apply Permanent (infinite duration) Damage"))
-
-		UGameplayEffect * BaseDmgEffect = Cast<UGameplayEffect>(StaticConstructObject(UGameplayEffect::StaticClass(), GetTransientPackage(), FName(TEXT("Damage"))));
-		BaseDmgEffect->Modifiers.SetNum(1);
-		BaseDmgEffect->Modifiers[0].Magnitude.SetValue(DamageValue);
-		BaseDmgEffect->Modifiers[0].ModifierType = EGameplayMod::Attribute;
-		BaseDmgEffect->Modifiers[0].ModifierOp = EGameplayModOp::Additive;
-		BaseDmgEffect->Modifiers[0].Attribute.SetUProperty(HealthProperty);
-		BaseDmgEffect->Modifiers[0].OwnedTags.AddTag(IGameplayTagsModule::RequestGameplayTag(FName(TEXT("Damage.Basic"))));
-		BaseDmgEffect->Duration.Value = UGameplayEffect::INFINITE_DURATION;
-
-		// Apply to target
-		DamageHandle = SourceComponent->ApplyGameplayEffectToTarget(BaseDmgEffect, DestComponent, 1.f);
-
-		float ExpectedValue = (StartHealth + (DamageValue * (DamageBuffMultiplier + DamageBuffMultiplierBonus)));
-		float ActualValue = DestComponent->GetSet<UAbilitySystemTestAttributeSet>()->Health;
-
-		Test->TestTrue(SKILL_TEST_TEXT("Damaged Applied. ActualValue: %.2f. ExpectedValue: %.2f.", ActualValue, ExpectedValue), (ActualValue == ExpectedValue));
-		Test->TestTrue(SKILL_TEST_TEXT("NumberOfActive Dest GameplayEffects %d == 1", DestComponent->GetNumActiveGameplayEffect()), DestComponent->GetNumActiveGameplayEffect() == 1);
-	}
-
-	{
-		// Clear DependantsUpdates stat
-#if SKILL_SYSTEM_AGGREGATOR_DEBUG
-		FAggregator::AllocationStats.DependantsUpdated = 0;
-#endif
-		// Remove the original buff buff
-		bool RemovedEffect = SourceComponent->RemoveActiveGameplayEffect(BuffBuffHandle);
-
-		float ActualValue = DestComponent->GetSet<UAbilitySystemTestAttributeSet>()->Health;
-		float ExpectedValue = (StartHealth + (DamageValue * (DamageBuffMultiplier + DamageBuffMultiplierBonus)));
-
-		Test->TestTrue(SKILL_TEST_TEXT("Damaged Applied. ActualValue: %.2f. ExpectedValue: %.2f.", ActualValue, ExpectedValue), (ActualValue == ExpectedValue));
-
-		// Check that we updated 0 dependants - since everything was applied via snapshot, no dependants should be updated
-#if SKILL_SYSTEM_AGGREGATOR_DEBUG
-		Test->TestTrue(SKILL_TEST_TEXT("DependantsUpdated %d == 2", FAggregator::AllocationStats.DependantsUpdated), (FAggregator::AllocationStats.DependantsUpdated == 0));
-#endif
-	}
-
-	{
-		ABILITY_LOG_SCOPE(TEXT("Remove Buff"));
-#if SKILL_SYSTEM_AGGREGATOR_DEBUG
-		FAggregator::AllocationStats.DependantsUpdated = 0;
-#endif
-		// Remove the buff
-		bool RemovedEffect = SourceComponent->RemoveActiveGameplayEffect(BuffHandle);
-
-		float ActualValue = DestComponent->GetSet<UAbilitySystemTestAttributeSet>()->Health;
-		float ExpectedValue = (StartHealth + (DamageValue * (DamageBuffMultiplier + DamageBuffMultiplierBonus)));
-
-		Test->TestTrue(SKILL_TEST_TEXT("Damaged Applied. ActualValue: %.2f. ExpectedValue: %.2f.", ActualValue, ExpectedValue), (ActualValue == ExpectedValue));
-
-		// Check that we updated 0 dependants
-#if SKILL_SYSTEM_AGGREGATOR_DEBUG
-		Test->TestTrue(SKILL_TEST_TEXT("DependantsUpdated %d == 0", FAggregator::AllocationStats.DependantsUpdated), (FAggregator::AllocationStats.DependantsUpdated == 0));
-#endif
-	}
-
-	// DestComponent->PrintAllGameplayEffects();
-
-	World->EditorDestroyActor(SourceActor, false);
-	World->EditorDestroyActor(DestActor, false);
-
-	return true;
-}
-
-bool GameplayEffectsTest_DamageBuffBuff_SnapshotLink(UWorld *World, FAutomationTestBase * Test)
-{
-	/**
-	*	Buff a Damage Buff, then apply damage
-	*/
-
-	const float StartHealth = 100.f;
-	const float DamageValue = -5.f;
-	const float DamageBuffMultiplier = 2.f;			// Damage is buff and multiplied by 2
-	const float DamageBuffMultiplierBonus = 1.f;	// The above multiplier receives a +1 bonus (we expect a final multiplier of 3)
-
-	AAbilitySystemTestPawn *SourceActor = World->SpawnActor<AAbilitySystemTestPawn>();
-	AAbilitySystemTestPawn *DestActor = World->SpawnActor<AAbilitySystemTestPawn>();
-
-	UProperty *HealthProperty = FindFieldChecked<UProperty>(UAbilitySystemTestAttributeSet::StaticClass(), GET_MEMBER_NAME_CHECKED(UAbilitySystemTestAttributeSet, Health));
-
-	UAbilitySystemComponent * SourceComponent = SourceActor->GetAbilitySystemComponent();
-	UAbilitySystemComponent * DestComponent = DestActor->GetAbilitySystemComponent();
-	SourceComponent->GetSet<UAbilitySystemTestAttributeSet>()->Health = StartHealth;
-	DestComponent->GetSet<UAbilitySystemTestAttributeSet>()->Health = StartHealth;
-
-	// Apply "Damage" but set to INFINITE_DURATION
-	// (An odd example for damage, but would make sense for something like run speed, etc)
-	FActiveGameplayEffectHandle BuffBuffHandle;
-	FActiveGameplayEffectHandle BuffHandle;
-	FActiveGameplayEffectHandle DamageHandle;
-
-	{
-		ABILITY_LOG_SCOPE(TEXT("Apply Buff Buff"))
-
-		// Here we are choosing to do this by adding an perm IncomingGE buff first. There are other ways to do this.
-
-		UGameplayEffect * BuffEffect = Cast<UGameplayEffect>(StaticConstructObject(UGameplayEffect::StaticClass(), GetTransientPackage(), FName(TEXT("BuffBuff"))));
-		BuffEffect->Modifiers.SetNum(1);
-		BuffEffect->Modifiers[0].Magnitude.SetValue(DamageBuffMultiplierBonus);
-		BuffEffect->Modifiers[0].ModifierType = EGameplayMod::IncomingGE;
-		BuffEffect->Modifiers[0].ModifierOp = EGameplayModOp::Additive;
-		BuffEffect->Modifiers[0].Attribute.SetUProperty(HealthProperty);
-		BuffEffect->Duration.Value = UGameplayEffect::INFINITE_DURATION;
-		BuffEffect->CopyPolicy = EGameplayEffectCopyPolicy::AlwaysSnapshot;
-
-		// Apply to target
-		BuffBuffHandle = SourceComponent->ApplyGameplayEffectToTarget(BuffEffect, SourceComponent, 1.f);
-
-		Test->TestTrue(SKILL_TEST_TEXT("Number of Source GameplayEffect: %d == 1", SourceComponent->GetNumActiveGameplayEffect()), SourceComponent->GetNumActiveGameplayEffect() == 1);
-	}
-
-	{
-		ABILITY_LOG_SCOPE(TEXT("Apply Damage Buff"))
-
-		UGameplayEffect * BuffEffect = Cast<UGameplayEffect>(StaticConstructObject(UGameplayEffect::StaticClass(), GetTransientPackage(), FName(TEXT("Buff"))));
-		BuffEffect->Modifiers.SetNum(1);
-		BuffEffect->Modifiers[0].Magnitude.SetValue(DamageBuffMultiplier);
-		BuffEffect->Modifiers[0].ModifierType = EGameplayMod::OutgoingGE;
-		BuffEffect->Modifiers[0].ModifierOp = EGameplayModOp::Multiplicitive;
-		BuffEffect->Modifiers[0].Attribute.SetUProperty(HealthProperty);
-		BuffEffect->Duration.Value = UGameplayEffect::INFINITE_DURATION;
-		BuffEffect->CopyPolicy = EGameplayEffectCopyPolicy::AlwaysLink;
-
-		// Apply to target
-		BuffHandle = SourceComponent->ApplyGameplayEffectToTarget(BuffEffect, SourceComponent, 1.f);
-		Test->TestTrue(SKILL_TEST_TEXT("Number of Source GameplayEffect: %d == 2", SourceComponent->GetNumActiveGameplayEffect()), SourceComponent->GetNumActiveGameplayEffect() == 2);
-
-		// Check that the buff was buffed
-		float ExpectedBuffMagnitude = DamageBuffMultiplier + DamageBuffMultiplierBonus;
-		float ActualBuffMagnitude = SourceComponent->GetGameplayEffectMagnitude(BuffHandle, FGameplayAttribute(HealthProperty));
-
-		Test->TestTrue(SKILL_TEST_TEXT("Buff Applied. Check Magnitude. ActualValue: %.2f. ExpectedValue: %.2f.", ActualBuffMagnitude, ExpectedBuffMagnitude), (ActualBuffMagnitude == ExpectedBuffMagnitude));
-	}
-
-	{
-		ABILITY_LOG_SCOPE(TEXT("Apply Permanent (infinite duration) Damage"))
-
-		UGameplayEffect * BaseDmgEffect = Cast<UGameplayEffect>(StaticConstructObject(UGameplayEffect::StaticClass(), GetTransientPackage(), FName(TEXT("Damage"))));
-		BaseDmgEffect->Modifiers.SetNum(1);
-		BaseDmgEffect->Modifiers[0].Magnitude.SetValue(DamageValue);
-		BaseDmgEffect->Modifiers[0].ModifierType = EGameplayMod::Attribute;
-		BaseDmgEffect->Modifiers[0].ModifierOp = EGameplayModOp::Additive;
-		BaseDmgEffect->Modifiers[0].Attribute.SetUProperty(HealthProperty);
-		BaseDmgEffect->Modifiers[0].OwnedTags.AddTag(IGameplayTagsModule::RequestGameplayTag(FName(TEXT("Damage.Basic"))));
-		BaseDmgEffect->Duration.Value = UGameplayEffect::INFINITE_DURATION;
-
-		// Apply to target
-		DamageHandle = SourceComponent->ApplyGameplayEffectToTarget(BaseDmgEffect, DestComponent, 1.f);
-
-		float ExpectedValue = (StartHealth + (DamageValue * (DamageBuffMultiplier + DamageBuffMultiplierBonus)));
-		float ActualValue = DestComponent->GetSet<UAbilitySystemTestAttributeSet>()->Health;
-
-		Test->TestTrue(SKILL_TEST_TEXT("Damaged Applied. ActualValue: %.2f. ExpectedValue: %.2f.", ActualValue, ExpectedValue), (ActualValue == ExpectedValue));
-		Test->TestTrue(SKILL_TEST_TEXT("NumberOfActive Dest GameplayEffects %d == 1", DestComponent->GetNumActiveGameplayEffect()), DestComponent->GetNumActiveGameplayEffect() == 1);
-	}
-
-	{
-		ABILITY_LOG_SCOPE(TEXT("Remove Buff Buff"));
-#if SKILL_SYSTEM_AGGREGATOR_DEBUG
-		FAggregator::AllocationStats.DependantsUpdated = 0;
-#endif
-		// Remove the original buff buff
-		bool RemovedEffect = SourceComponent->RemoveActiveGameplayEffect(BuffBuffHandle);
-
-		float ActualValue = DestComponent->GetSet<UAbilitySystemTestAttributeSet>()->Health;
-		float ExpectedValue = (StartHealth + (DamageValue * (DamageBuffMultiplier + DamageBuffMultiplierBonus)));
-
-		Test->TestTrue(SKILL_TEST_TEXT("Damaged Applied. ActualValue: %.2f. ExpectedValue: %.2f.", ActualValue, ExpectedValue), (ActualValue == ExpectedValue));
-
-		// Check that we updated 0 dependants - since everything was applied via snapshot, no dependants should be updated
-#if SKILL_SYSTEM_AGGREGATOR_DEBUG
-		Test->TestTrue(SKILL_TEST_TEXT("DependantsUpdated %d == 2", FAggregator::AllocationStats.DependantsUpdated), (FAggregator::AllocationStats.DependantsUpdated == 0));
-#endif
-	}
-
-	{
-		ABILITY_LOG_SCOPE(TEXT("Remove Buff"));
-#if SKILL_SYSTEM_AGGREGATOR_DEBUG
-		FAggregator::AllocationStats.DependantsUpdated = 0;
-#endif
-		// Remove the buff
-		bool RemovedEffect = SourceComponent->RemoveActiveGameplayEffect(BuffHandle);
-
-		float ActualValue = DestComponent->GetSet<UAbilitySystemTestAttributeSet>()->Health;
-		float ExpectedValue = StartHealth + DamageValue;
-
-		Test->TestTrue(SKILL_TEST_TEXT("Damaged Applied. ActualValue: %.2f. ExpectedValue: %.2f.", ActualValue, ExpectedValue), (ActualValue == ExpectedValue));
-
-		// Check that we updated 2 dependants - the damage GE and the attribute aggregator
-#if SKILL_SYSTEM_AGGREGATOR_DEBUG
-		Test->TestTrue(SKILL_TEST_TEXT("DependantsUpdated %d == 2", FAggregator::AllocationStats.DependantsUpdated), (FAggregator::AllocationStats.DependantsUpdated == 2));
-#endif
-	}
-
-	// DestComponent->PrintAllGameplayEffects();
-
-	World->EditorDestroyActor(SourceActor, false);
-	World->EditorDestroyActor(DestActor, false);
-
-	return true;
-}
-
-// Tests gameplay effects that apply other gameplay effects to the target
-bool GameplayEffectsTest_DamageAppliesBuff(UWorld *World, FAutomationTestBase * Test)
-{
-	const float StartHealth = 100.f;
-	const float DamageValue = 5.f;
-	const float DamageProtectionDivisor = 2.f;
-
-	AAbilitySystemTestPawn *SourceActor = World->SpawnActor<AAbilitySystemTestPawn>();
-	AAbilitySystemTestPawn *DestActor = World->SpawnActor<AAbilitySystemTestPawn>();
-
-	UProperty *DamageProperty = FindFieldChecked<UProperty>(UAbilitySystemTestAttributeSet::StaticClass(), GET_MEMBER_NAME_CHECKED(UAbilitySystemTestAttributeSet, Damage));
-
-	UAbilitySystemComponent * SourceComponent = SourceActor->GetAbilitySystemComponent();
-	UAbilitySystemComponent * DestComponent = DestActor->GetAbilitySystemComponent();
-	SourceComponent->GetSet<UAbilitySystemTestAttributeSet>()->Health = StartHealth;
-	DestComponent->GetSet<UAbilitySystemTestAttributeSet>()->Health = StartHealth;
-
-	// Apply damage and a buff that reduces incoming damage
-	{
-		ABILITY_LOG_SCOPE(TEXT("Apply DamageBuff and InstantDamage"))
-
-		UGameplayEffect * BuffEffect = Cast<UGameplayEffect>(StaticConstructObject(UGameplayEffect::StaticClass(), GetTransientPackage(), FName(TEXT("DamageBuff"))));
-		BuffEffect->Modifiers.SetNum(1);
-		BuffEffect->Modifiers[0].Magnitude.SetValue(DamageProtectionDivisor);
-		BuffEffect->Modifiers[0].ModifierType = EGameplayMod::IncomingGE;
-		BuffEffect->Modifiers[0].ModifierOp = EGameplayModOp::Division;
-		BuffEffect->Modifiers[0].Attribute.SetUProperty(DamageProperty);
-		BuffEffect->Modifiers[0].OwnedTags.AddTag(IGameplayTagsModule::RequestGameplayTag(FName(TEXT("Damage.Reduce"))));
-		BuffEffect->Duration.SetValue(UGameplayEffect::INFINITE_DURATION);
-
-		UGameplayEffect * BaseDmgEffect = Cast<UGameplayEffect>(StaticConstructObject(UGameplayEffect::StaticClass(), GetTransientPackage(), FName(TEXT("Damage"))));
-		BaseDmgEffect->Modifiers.SetNum(1);
-		BaseDmgEffect->Modifiers[0].Magnitude.SetValue(DamageValue);
-		BaseDmgEffect->Modifiers[0].ModifierType = EGameplayMod::Attribute;
-		BaseDmgEffect->Modifiers[0].ModifierOp = EGameplayModOp::Additive;
-		BaseDmgEffect->Modifiers[0].Attribute.SetUProperty(DamageProperty);
-		BaseDmgEffect->Modifiers[0].OwnedTags.AddTag(IGameplayTagsModule::RequestGameplayTag(FName(TEXT("Damage.Basic"))));
-		BaseDmgEffect->Duration.Value = UGameplayEffect::INSTANT_APPLICATION;
-		BaseDmgEffect->TargetEffects.Add(BuffEffect);
-
-		// Apply to target
-		SourceComponent->ApplyGameplayEffectToTarget(BaseDmgEffect, DestComponent, 1.f);
-
-		float ExpectedValue = (StartHealth - DamageValue);
-		float ActualValue = DestComponent->GetSet<UAbilitySystemTestAttributeSet>()->Health;
-
-		Test->TestTrue(SKILL_TEST_TEXT("Buff Instant Damage Applied"), (ActualValue == ExpectedValue));
-		ABILITY_LOG(Log, TEXT("Final Health: %.2f"), DestComponent->GetSet<UAbilitySystemTestAttributeSet>()->Health);
-	}
-
-	// Apply Damage
-	{
-		UGameplayEffect * BaseDmgEffect = Cast<UGameplayEffect>(StaticConstructObject(UGameplayEffect::StaticClass(), GetTransientPackage(), FName(TEXT("Damage"))));
-		BaseDmgEffect->Modifiers.SetNum(1);
-		BaseDmgEffect->Modifiers[0].Magnitude.SetValue(DamageValue);
-		BaseDmgEffect->Modifiers[0].ModifierType = EGameplayMod::Attribute;
-		BaseDmgEffect->Modifiers[0].ModifierOp = EGameplayModOp::Additive;
-		BaseDmgEffect->Modifiers[0].Attribute.SetUProperty(DamageProperty);
-		BaseDmgEffect->Modifiers[0].OwnedTags.AddTag(IGameplayTagsModule::RequestGameplayTag(FName(TEXT("Damage.Basic"))));
-		BaseDmgEffect->Duration.Value = UGameplayEffect::INSTANT_APPLICATION;
-
-		// Apply to target
-		SourceComponent->ApplyGameplayEffectToTarget(BaseDmgEffect, DestComponent, 1.f);
-
-		float ExpectedValue = (StartHealth - DamageValue - (DamageValue / DamageProtectionDivisor));
-		float ActualValue = DestComponent->GetSet<UAbilitySystemTestAttributeSet>()->Health;
-
-		Test->TestTrue(SKILL_TEST_TEXT("Buff Instant Damage Applied"), (ActualValue == ExpectedValue));
-		ABILITY_LOG(Log, TEXT("Final Health: %.2f"), DestComponent->GetSet<UAbilitySystemTestAttributeSet>()->Health);
-	}
-
-	World->EditorDestroyActor(SourceActor, false);
-	World->EditorDestroyActor(DestActor, false);
-
-	return true;
-}
-
-// Source has a buff that applies a buff to the target of all damage effects
-bool GameplayEffectsTest_BuffAppliesBuff(UWorld *World, FAutomationTestBase * Test)
-{
-	const float StartHealth = 100.f;
-	const float DamageValue = 5.f;
-	const float DamageProtectionDivisor = 2.f;
-
-	AAbilitySystemTestPawn *SourceActor = World->SpawnActor<AAbilitySystemTestPawn>();
-	AAbilitySystemTestPawn *DestActor = World->SpawnActor<AAbilitySystemTestPawn>();
-
-	UProperty *DamageProperty = FindFieldChecked<UProperty>(UAbilitySystemTestAttributeSet::StaticClass(), GET_MEMBER_NAME_CHECKED(UAbilitySystemTestAttributeSet, Damage));
-
-	UAbilitySystemComponent * SourceComponent = SourceActor->GetAbilitySystemComponent();
-	UAbilitySystemComponent * DestComponent = DestActor->GetAbilitySystemComponent();
-	SourceComponent->GetSet<UAbilitySystemTestAttributeSet>()->Health = StartHealth;
-	DestComponent->GetSet<UAbilitySystemTestAttributeSet>()->Health = StartHealth;
-
-	// Apply damage and a buff that reduces incoming damage
-	{
-		ABILITY_LOG_SCOPE(TEXT("Apply DamageBuff"))
-
-		UGameplayEffect * BuffEffect = Cast<UGameplayEffect>(StaticConstructObject(UGameplayEffect::StaticClass(), GetTransientPackage(), FName(TEXT("DamageBuff"))));
-		BuffEffect->Modifiers.SetNum(1);
-		BuffEffect->Modifiers[0].Magnitude.SetValue(DamageProtectionDivisor);
-		BuffEffect->Modifiers[0].ModifierType = EGameplayMod::IncomingGE;
-		BuffEffect->Modifiers[0].ModifierOp = EGameplayModOp::Division;
-		BuffEffect->Modifiers[0].Attribute.SetUProperty(DamageProperty);
-		BuffEffect->Modifiers[0].OwnedTags.AddTag(IGameplayTagsModule::RequestGameplayTag(FName(TEXT("Damage.Reduce"))));
-		BuffEffect->Duration.SetValue(UGameplayEffect::INFINITE_DURATION);
-
-		UGameplayEffect * DummyBuffEffect = Cast<UGameplayEffect>(StaticConstructObject(UGameplayEffect::StaticClass(), GetTransientPackage(), FName(TEXT("Dummy"))));
-		DummyBuffEffect->Modifiers.SetNum(1);
-		DummyBuffEffect->Modifiers[0].ModifierType = EGameplayMod::OutgoingGE;
-		DummyBuffEffect->Modifiers[0].EffectType = EGameplayModEffect::LinkedGameplayEffect;
-		DummyBuffEffect->Modifiers[0].RequiredTags.AddTag(IGameplayTagsModule::RequestGameplayTag(FName(TEXT("Damage.Buffable"))));
-		DummyBuffEffect->Modifiers[0].Attribute.SetUProperty(DamageProperty);
-		DummyBuffEffect->Modifiers[0].TargetEffect = BuffEffect;
-		DummyBuffEffect->Duration.SetValue(UGameplayEffect::INFINITE_DURATION);
-
-		// Apply to target
-		SourceComponent->ApplyGameplayEffectToTarget(DummyBuffEffect, SourceComponent, 1.f);
-	}
-
-	// apply damage to source to make sure it didn't get the protection buff
-	{
-		UGameplayEffect * BaseDmgEffect = Cast<UGameplayEffect>(StaticConstructObject(UGameplayEffect::StaticClass(), GetTransientPackage(), FName(TEXT("Damage"))));
-		BaseDmgEffect->Modifiers.SetNum(1);
-		BaseDmgEffect->Modifiers[0].Magnitude.SetValue(DamageValue);
-		BaseDmgEffect->Modifiers[0].ModifierType = EGameplayMod::Attribute;
-		BaseDmgEffect->Modifiers[0].ModifierOp = EGameplayModOp::Additive;
-		BaseDmgEffect->Modifiers[0].Attribute.SetUProperty(DamageProperty);
-		BaseDmgEffect->Modifiers[0].OwnedTags.AddTag(IGameplayTagsModule::RequestGameplayTag(FName(TEXT("Damage.Basic"))));
-		BaseDmgEffect->Duration.Value = UGameplayEffect::INSTANT_APPLICATION;
-
-		// Apply to target
-		DestComponent->ApplyGameplayEffectToTarget(BaseDmgEffect, SourceComponent, 1.f);
-
-		float ExpectedValue = StartHealth - DamageValue;
-		float ActualValue = SourceComponent->GetSet<UAbilitySystemTestAttributeSet>()->Health;
-
-		Test->TestTrue(SKILL_TEST_TEXT("Sending buffs test"), (ActualValue == ExpectedValue));
-		ABILITY_LOG(Log, TEXT("Final Health: %.2f"), SourceComponent->GetSet<UAbilitySystemTestAttributeSet>()->Health);
-	}
-
-	// Apply Damage
-	{
-		UGameplayEffect * BaseDmgEffect = Cast<UGameplayEffect>(StaticConstructObject(UGameplayEffect::StaticClass(), GetTransientPackage(), FName(TEXT("Damage"))));
-		BaseDmgEffect->Modifiers.SetNum(1);
-		BaseDmgEffect->Modifiers[0].Magnitude.SetValue(DamageValue);
-		BaseDmgEffect->Modifiers[0].ModifierType = EGameplayMod::Attribute;
-		BaseDmgEffect->Modifiers[0].ModifierOp = EGameplayModOp::Additive;
-		BaseDmgEffect->Modifiers[0].Attribute.SetUProperty(DamageProperty);
-		BaseDmgEffect->Modifiers[0].OwnedTags.AddTag(IGameplayTagsModule::RequestGameplayTag(FName(TEXT("Damage.Buffable"))));
-		BaseDmgEffect->Duration.Value = UGameplayEffect::INSTANT_APPLICATION;
-
-		// Apply to target
-		SourceComponent->ApplyGameplayEffectToTarget(BaseDmgEffect, DestComponent, 1.f);
-
-		float ExpectedValue = StartHealth - DamageValue;
-		float ActualValue = DestComponent->GetSet<UAbilitySystemTestAttributeSet>()->Health;
-
-		Test->TestTrue(SKILL_TEST_TEXT("Buff Instant Damage Applied"), (ActualValue == ExpectedValue));
-		ABILITY_LOG(Log, TEXT("Final Health: %.2f"), DestComponent->GetSet<UAbilitySystemTestAttributeSet>()->Health);
-	}
-
-	// Apply Damage
-	{
-		UGameplayEffect * BaseDmgEffect = Cast<UGameplayEffect>(StaticConstructObject(UGameplayEffect::StaticClass(), GetTransientPackage(), FName(TEXT("Damage"))));
-		BaseDmgEffect->Modifiers.SetNum(1);
-		BaseDmgEffect->Modifiers[0].Magnitude.SetValue(DamageValue);
-		BaseDmgEffect->Modifiers[0].ModifierType = EGameplayMod::Attribute;
-		BaseDmgEffect->Modifiers[0].ModifierOp = EGameplayModOp::Additive;
-		BaseDmgEffect->Modifiers[0].Attribute.SetUProperty(DamageProperty);
-		BaseDmgEffect->Modifiers[0].OwnedTags.AddTag(IGameplayTagsModule::RequestGameplayTag(FName(TEXT("Damage.Basic"))));
-		BaseDmgEffect->Duration.Value = UGameplayEffect::INSTANT_APPLICATION;
-
-		// Apply to target
-		SourceComponent->ApplyGameplayEffectToTarget(BaseDmgEffect, DestComponent, 1.f);
-
-		float ExpectedValue = (StartHealth - DamageValue - (DamageValue / DamageProtectionDivisor));
-		float ActualValue = DestComponent->GetSet<UAbilitySystemTestAttributeSet>()->Health;
-
-		Test->TestTrue(SKILL_TEST_TEXT("Instant Damage Applied"), (ActualValue == ExpectedValue));
-		ABILITY_LOG(Log, TEXT("Final Health: %.2f"), DestComponent->GetSet<UAbilitySystemTestAttributeSet>()->Health);
-	}
-
-	// Apply Damage again to make sure that the buff only applied once
-	{
-		UGameplayEffect * BaseDmgEffect = Cast<UGameplayEffect>(StaticConstructObject(UGameplayEffect::StaticClass(), GetTransientPackage(), FName(TEXT("Damage"))));
-		BaseDmgEffect->Modifiers.SetNum(1);
-		BaseDmgEffect->Modifiers[0].Magnitude.SetValue(DamageValue);
-		BaseDmgEffect->Modifiers[0].ModifierType = EGameplayMod::Attribute;
-		BaseDmgEffect->Modifiers[0].ModifierOp = EGameplayModOp::Additive;
-		BaseDmgEffect->Modifiers[0].Attribute.SetUProperty(DamageProperty);
-		BaseDmgEffect->Modifiers[0].OwnedTags.AddTag(IGameplayTagsModule::RequestGameplayTag(FName(TEXT("Damage.Basic"))));
-		BaseDmgEffect->Duration.Value = UGameplayEffect::INSTANT_APPLICATION;
-
-		// Apply to target
-		SourceComponent->ApplyGameplayEffectToTarget(BaseDmgEffect, DestComponent, 1.f);
-
-		float ExpectedValue = (StartHealth - DamageValue - (2 * DamageValue / DamageProtectionDivisor));
-		float ActualValue = DestComponent->GetSet<UAbilitySystemTestAttributeSet>()->Health;
-
-		Test->TestTrue(SKILL_TEST_TEXT("Instant Damage Applied"), (ActualValue == ExpectedValue));
-		ABILITY_LOG(Log, TEXT("Final Health: %.2f"), DestComponent->GetSet<UAbilitySystemTestAttributeSet>()->Health);
-	}
-
-	World->EditorDestroyActor(SourceActor, false);
-	World->EditorDestroyActor(DestActor, false);
-
-	return true;
-}
-
-bool GameplayEffectsTest_BuffIndirection(UWorld *World, FAutomationTestBase * Test)
-{
-	const float StartHealth = 100.f;
-	const float DamageValue = 5.f;
-	const float DamageProtectionDivisor = 2.f;
-
-	AAbilitySystemTestPawn *SourceActor = World->SpawnActor<AAbilitySystemTestPawn>();
-	AAbilitySystemTestPawn *DestActor = World->SpawnActor<AAbilitySystemTestPawn>();
-
-	UProperty *DamageProperty = FindFieldChecked<UProperty>(UAbilitySystemTestAttributeSet::StaticClass(), GET_MEMBER_NAME_CHECKED(UAbilitySystemTestAttributeSet, Damage));
-
-	UAbilitySystemComponent * SourceComponent = SourceActor->GetAbilitySystemComponent();
-	UAbilitySystemComponent * DestComponent = DestActor->GetAbilitySystemComponent();
-	SourceComponent->GetSet<UAbilitySystemTestAttributeSet>()->Health = StartHealth;
-	DestComponent->GetSet<UAbilitySystemTestAttributeSet>()->Health = StartHealth;
-
-	// Apply damage and a buff that reduces incoming damage
-	{
-		ABILITY_LOG_SCOPE(TEXT("Apply DamageBuff and InstantDamage"))
-
-		UGameplayEffect * BuffEffect = Cast<UGameplayEffect>(StaticConstructObject(UGameplayEffect::StaticClass(), GetTransientPackage(), FName(TEXT("DamageBuff"))));
-		BuffEffect->Modifiers.SetNum(1);
-		BuffEffect->Modifiers[0].Magnitude.SetValue(DamageProtectionDivisor);
-		BuffEffect->Modifiers[0].ModifierType = EGameplayMod::IncomingGE;
-		BuffEffect->Modifiers[0].ModifierOp = EGameplayModOp::Division;
-		BuffEffect->Modifiers[0].Attribute.SetUProperty(DamageProperty);
-		BuffEffect->Modifiers[0].OwnedTags.AddTag(IGameplayTagsModule::RequestGameplayTag(FName(TEXT("Damage.Reduce"))));
-		BuffEffect->Duration.SetValue(UGameplayEffect::INFINITE_DURATION);
-
-		UGameplayEffect * DummyBuffEffect = Cast<UGameplayEffect>(StaticConstructObject(UGameplayEffect::StaticClass(), GetTransientPackage(), FName(TEXT("Dummy"))));
-		DummyBuffEffect->Modifiers.SetNum(1);
-		DummyBuffEffect->Modifiers[0].ModifierType = EGameplayMod::OutgoingGE;
-		DummyBuffEffect->Modifiers[0].EffectType = EGameplayModEffect::LinkedGameplayEffect;
-		DummyBuffEffect->Modifiers[0].Attribute.SetUProperty(DamageProperty);
-		DummyBuffEffect->Modifiers[0].TargetEffect = BuffEffect;
-		DummyBuffEffect->Modifiers[0].RequiredTags.AddTag(IGameplayTagsModule::RequestGameplayTag(FName(TEXT("Damage.Buffable"))));
-		DummyBuffEffect->Duration.SetValue(UGameplayEffect::INFINITE_DURATION);
-
-		UGameplayEffect * DummyBuffEffect2 = Cast<UGameplayEffect>(StaticConstructObject(UGameplayEffect::StaticClass(), GetTransientPackage(), FName(TEXT("Dummy2"))));
-		DummyBuffEffect2->Modifiers.SetNum(1);
-		DummyBuffEffect2->Modifiers[0].ModifierType = EGameplayMod::OutgoingGE;
-		DummyBuffEffect2->Modifiers[0].EffectType = EGameplayModEffect::LinkedGameplayEffect;
-		DummyBuffEffect2->Modifiers[0].Attribute.SetUProperty(DamageProperty);
-		DummyBuffEffect2->Modifiers[0].TargetEffect = DummyBuffEffect;
-		DummyBuffEffect2->Duration.SetValue(UGameplayEffect::INFINITE_DURATION);
-
-		// Apply to target
-		SourceComponent->ApplyGameplayEffectToTarget(DummyBuffEffect2, SourceComponent, 1.f);
-	}
-
-	// Apply Damage
-	{
-		UGameplayEffect * BaseDmgEffect = Cast<UGameplayEffect>(StaticConstructObject(UGameplayEffect::StaticClass(), GetTransientPackage(), FName(TEXT("Damage"))));
-		BaseDmgEffect->Modifiers.SetNum(1);
-		BaseDmgEffect->Modifiers[0].Magnitude.SetValue(DamageValue);
-		BaseDmgEffect->Modifiers[0].ModifierType = EGameplayMod::Attribute;
-		BaseDmgEffect->Modifiers[0].ModifierOp = EGameplayModOp::Additive;
-		BaseDmgEffect->Modifiers[0].Attribute.SetUProperty(DamageProperty);
-		BaseDmgEffect->Modifiers[0].OwnedTags.AddTag(IGameplayTagsModule::RequestGameplayTag(FName(TEXT("Damage.Basic"))));
-		BaseDmgEffect->Duration.Value = UGameplayEffect::INSTANT_APPLICATION;
-
-		// Apply to target
-		SourceComponent->ApplyGameplayEffectToTarget(BaseDmgEffect, DestComponent, 1.f);
-
-		float ExpectedValue = StartHealth - DamageValue;
-		float ActualValue = DestComponent->GetSet<UAbilitySystemTestAttributeSet>()->Health;
-
-		Test->TestTrue(SKILL_TEST_TEXT("Buff Instant Damage Applied"), (ActualValue == ExpectedValue));
-		ABILITY_LOG(Log, TEXT("Final Health: %.2f"), DestComponent->GetSet<UAbilitySystemTestAttributeSet>()->Health);
-	}
-
-	// Apply Damage
-	{
-		UGameplayEffect * BaseDmgEffect = Cast<UGameplayEffect>(StaticConstructObject(UGameplayEffect::StaticClass(), GetTransientPackage(), FName(TEXT("Damage"))));
-		BaseDmgEffect->Modifiers.SetNum(1);
-		BaseDmgEffect->Modifiers[0].Magnitude.SetValue(DamageValue);
-		BaseDmgEffect->Modifiers[0].ModifierType = EGameplayMod::Attribute;
-		BaseDmgEffect->Modifiers[0].ModifierOp = EGameplayModOp::Additive;
-		BaseDmgEffect->Modifiers[0].Attribute.SetUProperty(DamageProperty);
-		BaseDmgEffect->Modifiers[0].OwnedTags.AddTag(IGameplayTagsModule::RequestGameplayTag(FName(TEXT("Damage.Buffable"))));
-		BaseDmgEffect->Duration.Value = UGameplayEffect::INSTANT_APPLICATION;
-
-		// Apply to target
-		DestComponent->ApplyGameplayEffectToTarget(BaseDmgEffect, SourceComponent, 1.f);
-
-		float ExpectedValue = StartHealth - DamageValue;
-		float ActualValue = SourceComponent->GetSet<UAbilitySystemTestAttributeSet>()->Health;
-
-		Test->TestTrue(SKILL_TEST_TEXT("Instant Damage Applied"), (ActualValue == ExpectedValue));
-		ABILITY_LOG(Log, TEXT("Final Health: %.2f"), SourceComponent->GetSet<UAbilitySystemTestAttributeSet>()->Health);
-	}
-
-	// Apply Damage again to make sure that the buff only applied once
-	{
-		UGameplayEffect * BaseDmgEffect = Cast<UGameplayEffect>(StaticConstructObject(UGameplayEffect::StaticClass(), GetTransientPackage(), FName(TEXT("Damage"))));
-		BaseDmgEffect->Modifiers.SetNum(1);
-		BaseDmgEffect->Modifiers[0].Magnitude.SetValue(DamageValue);
-		BaseDmgEffect->Modifiers[0].ModifierType = EGameplayMod::Attribute;
-		BaseDmgEffect->Modifiers[0].ModifierOp = EGameplayModOp::Additive;
-		BaseDmgEffect->Modifiers[0].Attribute.SetUProperty(DamageProperty);
-		BaseDmgEffect->Modifiers[0].OwnedTags.AddTag(IGameplayTagsModule::RequestGameplayTag(FName(TEXT("Damage.Basic"))));
-		BaseDmgEffect->Duration.Value = UGameplayEffect::INSTANT_APPLICATION;
-
-		// Apply to target
-		DestComponent->ApplyGameplayEffectToTarget(BaseDmgEffect, SourceComponent, 1.f);
-
-		float ExpectedValue = StartHealth - DamageValue - (DamageValue / DamageProtectionDivisor);
-		float ActualValue = SourceComponent->GetSet<UAbilitySystemTestAttributeSet>()->Health;
-
-		Test->TestTrue(SKILL_TEST_TEXT("Instant Damage Applied"), (ActualValue == ExpectedValue));
-		ABILITY_LOG(Log, TEXT("Final Health: %.2f"), SourceComponent->GetSet<UAbilitySystemTestAttributeSet>()->Health);
-	}
-
-	World->EditorDestroyActor(SourceActor, false);
-	World->EditorDestroyActor(DestActor, false);
-
-	return true;
-}
-
-bool GameplayEffectsTest_DurationDamage(UWorld *World, FAutomationTestBase * Test)
-{
-	const float StartHealth = 100.f;
-	const float DamageValue = -5.f;
-	float Duration = 5.f;
-	float StartTime = World->GetTimeSeconds();
-
-	AAbilitySystemTestPawn *SourceActor = World->SpawnActor<AAbilitySystemTestPawn>();
-	AAbilitySystemTestPawn *DestActor = World->SpawnActor<AAbilitySystemTestPawn>();
-
-	UProperty *HealthProperty = FindFieldChecked<UProperty>(UAbilitySystemTestAttributeSet::StaticClass(), GET_MEMBER_NAME_CHECKED(UAbilitySystemTestAttributeSet, Health));
-
-	UAbilitySystemComponent * SourceComponent = SourceActor->GetAbilitySystemComponent();
-	UAbilitySystemComponent * DestComponent = DestActor->GetAbilitySystemComponent();
-
-	SourceComponent->GetSet<UAbilitySystemTestAttributeSet>()->Health = StartHealth;
-	DestComponent->GetSet<UAbilitySystemTestAttributeSet>()->Health = StartHealth;
-
-	{
-		ABILITY_LOG_SCOPE(TEXT("Apply Temporary Damage"));
-
-		UGameplayEffect * BaseDmgEffect = Cast<UGameplayEffect>(StaticConstructObject(UGameplayEffect::StaticClass(), GetTransientPackage(), FName(TEXT("BaseDmgEffect"))));
-		BaseDmgEffect->Modifiers.SetNum(1);
-		BaseDmgEffect->Modifiers[0].Magnitude.SetValue(DamageValue);
-		BaseDmgEffect->Modifiers[0].ModifierType = EGameplayMod::Attribute;
-		BaseDmgEffect->Modifiers[0].ModifierOp = EGameplayModOp::Additive;
-		BaseDmgEffect->Modifiers[0].Attribute.SetUProperty(HealthProperty);
-		BaseDmgEffect->Modifiers[0].OwnedTags.AddTag(IGameplayTagsModule::RequestGameplayTag(FName(TEXT("Damage.Basic"))));
-		BaseDmgEffect->Duration.Value = Duration;
-		BaseDmgEffect->Period.Value = UGameplayEffect::NO_PERIOD;
-
-		SourceComponent->ApplyGameplayEffectToTarget(BaseDmgEffect, DestComponent, 1.f);
-
-		// The effect should instantly apply without ticking (for now at least)
-		float ActualValue = DestComponent->GetSet<UAbilitySystemTestAttributeSet>()->Health;
-		float ExpectedValue = StartHealth + (DamageValue);
-
-		Test->TestTrue(SKILL_TEST_TEXT("Damage Applied. Actual: %.2f == Exected: %.2f", ActualValue, ExpectedValue), ActualValue == ExpectedValue);
-
-		GameplayTest_TickWorld(World, 0.0001f); // Move our Effects from the pending stack to the active stack inside the timer manager, this starts the clock for effect execution
-		StartTime = World->GetTimeSeconds();
-	}
-
-	// Tick until the effect should expire
-	for (int32 i = 0; i < 10; ++i)
-	{
-		GameplayTest_TickWorld(World, 1.f);
-		if (World->GetTimeSeconds() > StartTime + Duration + KINDA_SMALL_NUMBER)
-		{
-			break;
-		}
-
-		// The temporary effect is still in place
-		float ActualValue = DestComponent->GetSet<UAbilitySystemTestAttributeSet>()->Health;
-		float ExpectedValue = StartHealth + (DamageValue);
-
-		Test->TestTrue(SKILL_TEST_TEXT("Damage Applied. Duration (left) %.2f. Actual: %.2f == Exected: %.2f", Duration, ActualValue, ExpectedValue), ActualValue == ExpectedValue);
-	}
-
-	// Ensure the effect expired
-	{
-		// The effect should instantly execute one time without ticking (for now at least)
-		float ActualValue = DestComponent->GetSet<UAbilitySystemTestAttributeSet>()->Health;
-		float ExpectedValue = StartHealth;
-
-		Test->TestTrue(SKILL_TEST_TEXT("Damage Applied. Duration (left) %.2f. Actual: %.2f == Exected: %.2f", Duration, ActualValue, ExpectedValue), ActualValue == ExpectedValue);
-
-		int32 NumEffects = DestComponent->GetNumActiveGameplayEffect();
-		Test->TestTrue(SKILL_TEST_TEXT("NumberOfActive Dest GameplayEffects %d == 0", NumEffects), NumEffects == 0);
-	}
-
-	World->EditorDestroyActor(SourceActor, false);
-	World->EditorDestroyActor(DestActor, false);
-
-	return true;
-}
-
-bool GameplayEffectsTest_PeriodicDamage(UWorld *World, FAutomationTestBase * Test)
-{
-	const float StartHealth = 100.f;
-	const float DamageValue = -5.f;
-	float Duration = 5.f;
-	float StartTime = World->GetTimeSeconds();
-
-	AAbilitySystemTestPawn *SourceActor = World->SpawnActor<AAbilitySystemTestPawn>();
-	AAbilitySystemTestPawn *DestActor = World->SpawnActor<AAbilitySystemTestPawn>();
-
-	UProperty *HealthProperty = FindFieldChecked<UProperty>(UAbilitySystemTestAttributeSet::StaticClass(), GET_MEMBER_NAME_CHECKED(UAbilitySystemTestAttributeSet, Health));
-
-	UAbilitySystemComponent * SourceComponent = SourceActor->GetAbilitySystemComponent();
-	UAbilitySystemComponent * DestComponent = DestActor->GetAbilitySystemComponent();
-
-	SourceComponent->GetSet<UAbilitySystemTestAttributeSet>()->Health = StartHealth;
-	DestComponent->GetSet<UAbilitySystemTestAttributeSet>()->Health = StartHealth;
-
-	float ApplyCount = 0;
-	{
-		ABILITY_LOG_SCOPE(TEXT("Apply Dot"))
-
-		UGameplayEffect * BaseDmgEffect = Cast<UGameplayEffect>(StaticConstructObject(UGameplayEffect::StaticClass(), GetTransientPackage(), FName(TEXT("BaseDmgEffect"))));
-		BaseDmgEffect->Modifiers.SetNum(1);
-		BaseDmgEffect->Modifiers[0].Magnitude.SetValue(DamageValue);
-		BaseDmgEffect->Modifiers[0].ModifierType = EGameplayMod::Attribute;
-		BaseDmgEffect->Modifiers[0].ModifierOp = EGameplayModOp::Additive;
-		BaseDmgEffect->Modifiers[0].Attribute.SetUProperty(HealthProperty);
-		BaseDmgEffect->Modifiers[0].OwnedTags.AddTag(IGameplayTagsModule::RequestGameplayTag(FName(TEXT("Damage.Basic"))));
-		BaseDmgEffect->Duration.Value = Duration;
-		BaseDmgEffect->Period.Value = 1.f; // Apply every 1 second
-		BaseDmgEffect->GameplayCues.Add( FGameplayEffectCue( IGameplayTagsModule::RequestGameplayTag(FName(TEXT("GameplayCue.Burning"))), 1.f, 10.f) );
-
-		SourceComponent->ApplyGameplayEffectToTarget(BaseDmgEffect, DestComponent, 5.f);
-
-		// The effect should execute on the next tick
-		float ActualValue = DestComponent->GetSet<UAbilitySystemTestAttributeSet>()->Health;
-		float ExpectedValue = StartHealth + (DamageValue * ApplyCount);
-
-		Test->TestTrue(SKILL_TEST_TEXT("Damage Applied. Actual: %.2f == Exected: %.2f", ActualValue, ExpectedValue), ActualValue == ExpectedValue );
-	}
-
-	GameplayTest_TickWorld(World, 0.0001f); // Move our Effects from the pending stack to the active stack inside the timer manager, this starts the clock for effect execution
-	ApplyCount++; // the effect will execute as soon as we tick any amount of time
-	GameplayTest_TickWorld(World, 0.1f); // Offset the current time from the start of the period so that floating point issues don't affect the test
-
-	for (int32 i=0; i < 10; ++i)
-	{
-		GameplayTest_TickWorld(World, 1.f);
-		if (World->GetTimeSeconds() <= StartTime + Duration)
-		{
-			// We should have applied as long as there was still some duration left
-			ApplyCount++;
-		}
-
-		// The effect should instantly execute one time without ticking (for now at least)
-		float ActualValue = DestComponent->GetSet<UAbilitySystemTestAttributeSet>()->Health;
-		float ExpectedValue = StartHealth + (DamageValue * ApplyCount);
-
-		Test->TestTrue(SKILL_TEST_TEXT("Damage Applied. Duration (left) %.2f. Actual: %.2f == Exected: %.2f", Duration, ActualValue, ExpectedValue), ActualValue == ExpectedValue);
-	}
-
-	// Ensure the effect expired
-	{
-		int32 NumEffects = DestComponent->GetNumActiveGameplayEffect();
-		Test->TestTrue(SKILL_TEST_TEXT("NumberOfActive Dest GameplayEffects %d == 0", NumEffects), NumEffects == 0);
-	}
-
-	World->EditorDestroyActor(SourceActor, false);
-	World->EditorDestroyActor(DestActor, false);
-
-	return true;
-}
-
-bool GameplayEffectsTest_LifestealExtension(UWorld *World, FAutomationTestBase * Test)
-{
-	const float StartHealth = 100.f;
-	const float DamageValue = -10.f;
-	const float LifestealPCT = 0.20f;
-	float StartTime = World->GetTimeSeconds();
-
-	AAbilitySystemTestPawn *SourceActor = World->SpawnActor<AAbilitySystemTestPawn>();
-	AAbilitySystemTestPawn *DestActor = World->SpawnActor<AAbilitySystemTestPawn>();
-
-	UProperty *HealthProperty = FindFieldChecked<UProperty>(UAbilitySystemTestAttributeSet::StaticClass(), GET_MEMBER_NAME_CHECKED(UAbilitySystemTestAttributeSet, Health));
-	UProperty *DamageProperty = FindFieldChecked<UProperty>(UAbilitySystemTestAttributeSet::StaticClass(), GET_MEMBER_NAME_CHECKED(UAbilitySystemTestAttributeSet, Damage));
-
-	UAbilitySystemComponent * SourceComponent = SourceActor->GetAbilitySystemComponent();
-	UAbilitySystemComponent * DestComponent = DestActor->GetAbilitySystemComponent();
-
-	SourceComponent->GetSet<UAbilitySystemTestAttributeSet>()->Health = StartHealth;
-	DestComponent->GetSet<UAbilitySystemTestAttributeSet>()->Health = StartHealth;
-
-	{
-		ABILITY_LOG_SCOPE(TEXT("Apply Lifesteal"));
-
-		FGameplayModifierCallback LifestealCallback;
-		LifestealCallback.ExtensionClass = UGameplayEffectExtension_LifestealTest::StaticClass();
-
-		UGameplayEffect * LifestealEffect = Cast<UGameplayEffect>(StaticConstructObject(UGameplayEffect::StaticClass(), GetTransientPackage(), FName(TEXT("LifestealPassive"))));
-		LifestealEffect->Modifiers.SetNum(1);
-		LifestealEffect->Modifiers[0].Magnitude.SetValue(LifestealPCT);
-		LifestealEffect->Modifiers[0].ModifierType = EGameplayMod::OutgoingGE;
-		LifestealEffect->Modifiers[0].ModifierOp = EGameplayModOp::Callback;
-		LifestealEffect->Modifiers[0].Attribute.SetUProperty(HealthProperty);
-		LifestealEffect->Modifiers[0].OwnedTags.AddTag(IGameplayTagsModule::RequestGameplayTag(FName(TEXT("Lifesteal"))));
-		LifestealEffect->Modifiers[0].Callbacks.Add( LifestealCallback );
-		LifestealEffect->Duration.Value = UGameplayEffect::INFINITE_DURATION;
-		LifestealEffect->Period.Value = UGameplayEffect::NO_PERIOD;
-
-		SourceComponent->ApplyGameplayEffectToSelf(LifestealEffect, 1.f, SourceComponent->GetEffectContext());
-	}
-	
-	{
-		ABILITY_LOG_SCOPE(TEXT("Apply Damage"));
-
-		UGameplayEffect * BaseDmgEffect = Cast<UGameplayEffect>(StaticConstructObject(UGameplayEffect::StaticClass(), GetTransientPackage(), FName(TEXT("BaseDmgEffect"))));
-		BaseDmgEffect->Modifiers.SetNum(1);
-		BaseDmgEffect->Modifiers[0].Magnitude.SetValue(DamageValue);
-		BaseDmgEffect->Modifiers[0].ModifierType = EGameplayMod::Attribute;
-		BaseDmgEffect->Modifiers[0].ModifierOp = EGameplayModOp::Additive;
-		BaseDmgEffect->Modifiers[0].Attribute.SetUProperty(HealthProperty);
-		BaseDmgEffect->Modifiers[0].OwnedTags.AddTag(IGameplayTagsModule::RequestGameplayTag(FName(TEXT("Damage.Basic"))));
-		BaseDmgEffect->Duration.Value = UGameplayEffect::INSTANT_APPLICATION;
-		BaseDmgEffect->Period.Value = UGameplayEffect::NO_PERIOD;
-
-		SourceComponent->ApplyGameplayEffectToTarget(BaseDmgEffect, DestComponent, 5.f);
-
-		// The effect should instantly execute one time without ticking (for now at least)
-		{
-			float ActualValue = DestComponent->GetSet<UAbilitySystemTestAttributeSet>()->Health;
-			float ExpectedValue = StartHealth + (DamageValue);
-
-			Test->TestTrue(SKILL_TEST_TEXT("Damage Applied. Actual: %.2f == Exected: %.2f", ActualValue, ExpectedValue), ActualValue == ExpectedValue);
-		}
-
-		// Test that the source received extra health back
-		{
-			float ActualValue = SourceComponent->GetSet<UAbilitySystemTestAttributeSet>()->Health;
-			float ExpectedValue = StartHealth + (-DamageValue * LifestealPCT);
-
-			Test->TestTrue(SKILL_TEST_TEXT("Health after lifesteal. Actual: %.2f == Exected: %.2f", ActualValue, ExpectedValue), ActualValue == ExpectedValue);
-		}		
-
-	}
-
-	World->EditorDestroyActor(SourceActor, false);
-	World->EditorDestroyActor(DestActor, false);
-
-	return true;
-}
-
-bool GameplayEffectsTest_ShieldExtension(UWorld *World, FAutomationTestBase * Test)
-{
-	const float StartHealth = 100.f;
-	const float DamageValue = -10.f;
-	const float ShieldAmount = 20.0f;
-	float StartTime = World->GetTimeSeconds();
-
-	AAbilitySystemTestPawn *SourceActor = World->SpawnActor<AAbilitySystemTestPawn>();
-	AAbilitySystemTestPawn *DestActor = World->SpawnActor<AAbilitySystemTestPawn>();
-
-	UProperty *HealthProperty = FindFieldChecked<UProperty>(UAbilitySystemTestAttributeSet::StaticClass(), GET_MEMBER_NAME_CHECKED(UAbilitySystemTestAttributeSet, Health));
-	UProperty *DamageProperty = FindFieldChecked<UProperty>(UAbilitySystemTestAttributeSet::StaticClass(), GET_MEMBER_NAME_CHECKED(UAbilitySystemTestAttributeSet, Damage));
-
-	UAbilitySystemComponent * SourceComponent = SourceActor->GetAbilitySystemComponent();
-	UAbilitySystemComponent * DestComponent = DestActor->GetAbilitySystemComponent();
-
-	SourceComponent->GetSet<UAbilitySystemTestAttributeSet>()->Health = StartHealth;
-	DestComponent->GetSet<UAbilitySystemTestAttributeSet>()->Health = StartHealth;
-
-	FActiveGameplayEffectHandle AppliedHandle;
-	{
-		ABILITY_LOG_SCOPE(TEXT("Apply Shield"));
-
-		FGameplayModifierCallback ShieldCallback;
-		ShieldCallback.ExtensionClass = UGameplayEffectExtension_ShieldTest::StaticClass();
-
-		UGameplayEffect * ShieldEffect = Cast<UGameplayEffect>(StaticConstructObject(UGameplayEffect::StaticClass(), GetTransientPackage(), FName(TEXT("ShieldPassive"))));
-		ShieldEffect->Modifiers.SetNum(1);
-		ShieldEffect->Modifiers[0].Magnitude.SetValue(ShieldAmount);
-		ShieldEffect->Modifiers[0].ModifierType = EGameplayMod::IncomingGE;
-		ShieldEffect->Modifiers[0].ModifierOp = EGameplayModOp::Callback;
-		ShieldEffect->Modifiers[0].Attribute.SetUProperty(HealthProperty);
-		ShieldEffect->Modifiers[0].OwnedTags.AddTag(IGameplayTagsModule::RequestGameplayTag(FName(TEXT("Shield"))));
-		ShieldEffect->Modifiers[0].Callbacks.Add(ShieldCallback);
-		ShieldEffect->Duration.Value = UGameplayEffect::INFINITE_DURATION;
-		ShieldEffect->Period.Value = UGameplayEffect::NO_PERIOD;
-
-		AppliedHandle = DestComponent->ApplyGameplayEffectToSelf(ShieldEffect, 1.f, DestComponent->GetEffectContext());
-	}
-
-	{
-		ABILITY_LOG_SCOPE(TEXT("Apply Damage"));
-
-		UGameplayEffect * BaseDmgEffect = Cast<UGameplayEffect>(StaticConstructObject(UGameplayEffect::StaticClass(), GetTransientPackage(), FName(TEXT("BaseDmgEffect"))));
-		BaseDmgEffect->Modifiers.SetNum(1);
-		BaseDmgEffect->Modifiers[0].Magnitude.SetValue(DamageValue);
-		BaseDmgEffect->Modifiers[0].ModifierType = EGameplayMod::Attribute;
-		BaseDmgEffect->Modifiers[0].ModifierOp = EGameplayModOp::Additive;
-		BaseDmgEffect->Modifiers[0].Attribute.SetUProperty(HealthProperty);
-		BaseDmgEffect->Modifiers[0].OwnedTags.AddTag(IGameplayTagsModule::RequestGameplayTag(FName(TEXT("Damage.Basic"))));
-		BaseDmgEffect->Duration.Value = UGameplayEffect::INSTANT_APPLICATION;
-		BaseDmgEffect->Period.Value = UGameplayEffect::NO_PERIOD;
-
-		
-		// Apply 1
-		{
-			SourceComponent->ApplyGameplayEffectToTarget(BaseDmgEffect, DestComponent, 5.f);
-
-			// Health should be the same
-			{
-				float ActualValue = DestComponent->GetSet<UAbilitySystemTestAttributeSet>()->Health;
-				float ExpectedValue = StartHealth;
-				Test->TestTrue(SKILL_TEST_TEXT("Damage Applied. Actual: %.2f == Exected: %.2f", ActualValue, ExpectedValue), ActualValue == ExpectedValue);
-			}
-
-			// Shield still up but weakened
-			{
-				float ActualValue = DestComponent->GetGameplayEffectMagnitude(AppliedHandle, FGameplayAttribute(HealthProperty));
-				float ExpectedValue = ShieldAmount + DamageValue;
-				Test->TestTrue(SKILL_TEST_TEXT("Damage Applied. Actual: %.2f == Exected: %.2f", ActualValue, ExpectedValue), ActualValue == ExpectedValue);
-			}
-		}
-
-		// Apply 2
-		{
-			SourceComponent->ApplyGameplayEffectToTarget(BaseDmgEffect, DestComponent, 5.f);
-
-			// Health should be the same
-			{
-				float ActualValue = DestComponent->GetSet<UAbilitySystemTestAttributeSet>()->Health;
-				float ExpectedValue = StartHealth;
-				Test->TestTrue(SKILL_TEST_TEXT("Damage Applied. Actual: %.2f == Exected: %.2f", ActualValue, ExpectedValue), ActualValue == ExpectedValue);
-			}
-
-			// Shield should be done now (it absorbed the damage and then removed itself)
-			{
-				bool Removed = DestComponent->IsGameplayEffectActive(AppliedHandle);
-				Test->TestTrue(SKILL_TEST_TEXT("Shield removed (Expected: 0 Actual: %d", Removed), !Removed);
-			}
-		}
-
-		// Apply 3
-		{
-			// Now we lose health
-			SourceComponent->ApplyGameplayEffectToTarget(BaseDmgEffect, DestComponent, 5.f);
-
-			// Now we should have lost some health
-			float ActualValue = DestComponent->GetSet<UAbilitySystemTestAttributeSet>()->Health;
-			float ExpectedValue = StartHealth + DamageValue;;
-			Test->TestTrue(SKILL_TEST_TEXT("Damage Applied. Actual: %.2f == Exected: %.2f", ActualValue, ExpectedValue), ActualValue == ExpectedValue);
-
-			// For funsies, confirm shield is still definitely not there
-			{
-				bool Removed = DestComponent->IsGameplayEffectActive(AppliedHandle);
-				Test->TestTrue(SKILL_TEST_TEXT("Shield removed (Expected: 0 Actual: %d", Removed), !Removed);
-			}
-		}
-	}
-
-	World->EditorDestroyActor(SourceActor, false);
-	World->EditorDestroyActor(DestActor, false);
-
-	return true;
-}
-
-bool GameplayEffectsTest_ShieldExtensionMultiple(UWorld *World, FAutomationTestBase * Test)
-{
-	// This applies 2 instances of the shield and confirms that only 1 will absorb damage at at time
-
-	const float StartHealth = 100.f;
-	const float DamageValueSmall = -10.f;
-	const float DamageValueLarge = -20.f;
-	const float ShieldAmount = 20.0f;
-	float StartTime = World->GetTimeSeconds();
-
-	AAbilitySystemTestPawn *SourceActor = World->SpawnActor<AAbilitySystemTestPawn>();
-	AAbilitySystemTestPawn *DestActor = World->SpawnActor<AAbilitySystemTestPawn>();
-
-	UProperty *HealthProperty = FindFieldChecked<UProperty>(UAbilitySystemTestAttributeSet::StaticClass(), GET_MEMBER_NAME_CHECKED(UAbilitySystemTestAttributeSet, Health));
-	UProperty *DamageProperty = FindFieldChecked<UProperty>(UAbilitySystemTestAttributeSet::StaticClass(), GET_MEMBER_NAME_CHECKED(UAbilitySystemTestAttributeSet, Damage));
-
-	UAbilitySystemComponent * SourceComponent = SourceActor->GetAbilitySystemComponent();
-	UAbilitySystemComponent * DestComponent = DestActor->GetAbilitySystemComponent();
-
-	SourceComponent->GetSet<UAbilitySystemTestAttributeSet>()->Health = StartHealth;
-	DestComponent->GetSet<UAbilitySystemTestAttributeSet>()->Health = StartHealth;
-
-	FActiveGameplayEffectHandle AppliedHandle_1;
-	FActiveGameplayEffectHandle AppliedHandle_2;
-	{
-		ABILITY_LOG_SCOPE(TEXT("Apply Shields"));
-
-		FGameplayModifierCallback ShieldCallback;
-		ShieldCallback.ExtensionClass = UGameplayEffectExtension_ShieldTest::StaticClass();
-
-		UGameplayEffect * ShieldEffect = Cast<UGameplayEffect>(StaticConstructObject(UGameplayEffect::StaticClass(), GetTransientPackage(), FName(TEXT("ShieldPassive"))));
-		ShieldEffect->Modifiers.SetNum(1);
-		ShieldEffect->Modifiers[0].Magnitude.SetValue(ShieldAmount);
-		ShieldEffect->Modifiers[0].ModifierType = EGameplayMod::IncomingGE;
-		ShieldEffect->Modifiers[0].ModifierOp = EGameplayModOp::Callback;
-		ShieldEffect->Modifiers[0].Attribute.SetUProperty(HealthProperty);
-		ShieldEffect->Modifiers[0].OwnedTags.AddTag(IGameplayTagsModule::RequestGameplayTag(FName(TEXT("Shield"))));
-		ShieldEffect->Modifiers[0].Callbacks.Add(ShieldCallback);
-		ShieldEffect->Duration.Value = UGameplayEffect::INFINITE_DURATION;
-		ShieldEffect->Period.Value = UGameplayEffect::NO_PERIOD;
-
-		AppliedHandle_1 = DestComponent->ApplyGameplayEffectToSelf(ShieldEffect, 1.f, DestComponent->GetEffectContext());
-		AppliedHandle_2 = DestComponent->ApplyGameplayEffectToSelf(ShieldEffect, 1.f, DestComponent->GetEffectContext());
-	}
-
-	{
-		ABILITY_LOG_SCOPE(TEXT("Apply Damage"));
-
-		UGameplayEffect * SmallDmgEffect = Cast<UGameplayEffect>(StaticConstructObject(UGameplayEffect::StaticClass(), GetTransientPackage(), FName(TEXT("SmallDmgEffect"))));
-		SmallDmgEffect->Modifiers.SetNum(1);
-		SmallDmgEffect->Modifiers[0].Magnitude.SetValue(DamageValueSmall);
-		SmallDmgEffect->Modifiers[0].ModifierType = EGameplayMod::Attribute;
-		SmallDmgEffect->Modifiers[0].ModifierOp = EGameplayModOp::Additive;
-		SmallDmgEffect->Modifiers[0].Attribute.SetUProperty(HealthProperty);
-		SmallDmgEffect->Modifiers[0].OwnedTags.AddTag(IGameplayTagsModule::RequestGameplayTag(FName(TEXT("Damage.Basic"))));
-		SmallDmgEffect->Duration.Value = UGameplayEffect::INSTANT_APPLICATION;
-		SmallDmgEffect->Period.Value = UGameplayEffect::NO_PERIOD;
-
-		UGameplayEffect * LargeDmgEffect = Cast<UGameplayEffect>(StaticConstructObject(UGameplayEffect::StaticClass(), GetTransientPackage(), FName(TEXT("LargeDmgEffect"))));
-		LargeDmgEffect->Modifiers.SetNum(1);
-		LargeDmgEffect->Modifiers[0].Magnitude.SetValue(DamageValueLarge);
-		LargeDmgEffect->Modifiers[0].ModifierType = EGameplayMod::Attribute;
-		LargeDmgEffect->Modifiers[0].ModifierOp = EGameplayModOp::Additive;
-		LargeDmgEffect->Modifiers[0].Attribute.SetUProperty(HealthProperty);
-		LargeDmgEffect->Modifiers[0].OwnedTags.AddTag(IGameplayTagsModule::RequestGameplayTag(FName(TEXT("Damage.Basic"))));
-		LargeDmgEffect->Duration.Value = UGameplayEffect::INSTANT_APPLICATION;
-		LargeDmgEffect->Period.Value = UGameplayEffect::NO_PERIOD;
-
-
-		// Apply small damage
-		{
-			SourceComponent->ApplyGameplayEffectToTarget(SmallDmgEffect, DestComponent, 5.f);
-
-			// Health should be the same
-			{
-				float ActualValue = DestComponent->GetSet<UAbilitySystemTestAttributeSet>()->Health;
-				float ExpectedValue = StartHealth;
-				Test->TestTrue(SKILL_TEST_TEXT("Damage Applied. Actual: %.2f == Exected: %.2f", ActualValue, ExpectedValue), ActualValue == ExpectedValue);
-			}
-
-			// Shield 1 still up but weakened
-			{
-				float ActualValue = DestComponent->GetGameplayEffectMagnitude(AppliedHandle_1, FGameplayAttribute(HealthProperty));
-				float ExpectedValue = ShieldAmount + DamageValueSmall;
-				Test->TestTrue(SKILL_TEST_TEXT("Shield 1. Actual: %.2f == Exected: %.2f", ActualValue, ExpectedValue), ActualValue == ExpectedValue);
-			}
-
-			// Shield 2 untouched
-			{
-				float ActualValue = DestComponent->GetGameplayEffectMagnitude(AppliedHandle_2, FGameplayAttribute(HealthProperty));
-				float ExpectedValue = ShieldAmount;
-				Test->TestTrue(SKILL_TEST_TEXT("Shield 1. Actual: %.2f == Exected: %.2f", ActualValue, ExpectedValue), ActualValue == ExpectedValue);
-			}
-		}
-
-		// Apply large damage
-		{
-			SourceComponent->ApplyGameplayEffectToTarget(LargeDmgEffect, DestComponent, 5.f);
-
-			// Health should still be the same
-			{
-				float ActualValue = DestComponent->GetSet<UAbilitySystemTestAttributeSet>()->Health;
-				float ExpectedValue = StartHealth;
-				Test->TestTrue(SKILL_TEST_TEXT("Damage Applied. Actual: %.2f == Exected: %.2f", ActualValue, ExpectedValue), ActualValue == ExpectedValue);
-			}
-
-			// Shield 1 should be gone
-			{
-				bool Exists = DestComponent->IsGameplayEffectActive(AppliedHandle_1);
-				Test->TestTrue(SKILL_TEST_TEXT("Shield removed (Expected: 0 Actual: %d", Exists), !Exists);
-			}
-
-			// Shield 2 should be weakened
-			{
-				float DamageShield2Took = ShieldAmount + DamageValueSmall + DamageValueLarge;
-
-				float ActualValue = DestComponent->GetGameplayEffectMagnitude(AppliedHandle_2, FGameplayAttribute(HealthProperty));
-				float ExpectedValue = ShieldAmount + DamageShield2Took;
-				Test->TestTrue(SKILL_TEST_TEXT("Shield 1. Actual: %.2f == Exected: %.2f", ActualValue, ExpectedValue), ActualValue == ExpectedValue);
-			}
-		}
-
-		// Apply large damage again
-		{
-			// Now we lose health
-			SourceComponent->ApplyGameplayEffectToTarget(LargeDmgEffect, DestComponent, 5.f);
-
-			float HealthDelta = ShieldAmount + ShieldAmount + DamageValueSmall + DamageValueLarge + DamageValueLarge;
-
-			// Now we should have lost some health
-			float ActualValue = DestComponent->GetSet<UAbilitySystemTestAttributeSet>()->Health;
-			float ExpectedValue = StartHealth + HealthDelta;
-			Test->TestTrue(SKILL_TEST_TEXT("Damage Applied. Actual: %.2f == Exected: %.2f", ActualValue, ExpectedValue), ActualValue == ExpectedValue);
-
-			// For funsies, confirm shield is still definitely not there
-			{
-				bool Exists = DestComponent->IsGameplayEffectActive(AppliedHandle_2);
-				Test->TestTrue(SKILL_TEST_TEXT("Shield removed (Expected: 0 Actual: %d", Exists), !Exists);
-			}
-		}
-	}
-
-	World->EditorDestroyActor(SourceActor, false);
-	World->EditorDestroyActor(DestActor, false);
-
-	return true;
-}
-
-
-
-UCurveTable * SetGlobalCurveTable()
-{
-	FString CSV(TEXT(", 0, 1, 100\r\nStandardHealth, 0, 1, 100\r\nStandardDamage, 0, 1, 100\r\nLinearCurve, 0, 1, 100"));
-
-	UCurveTable * CurveTable = Cast<UCurveTable>(StaticConstructObject(UCurveTable::StaticClass(), GetTransientPackage(), FName(TEXT("TempCurveTable"))));
-	CurveTable->CreateTableFromCSVString(CSV);
-
-	FRichCurve * RichCurve = CurveTable->FindCurve(FName(TEXT("StandardHealth")), TEXT("Test"));
-	if (RichCurve)
-	{
-		float Value = RichCurve->Eval(5.f);
-		check(Value == 5.f);
-	}
-
-	IGameplayAbilitiesModule::Get().GetAbilitySystemGlobals()->AutomationTestOnly_SetGlobalCurveTable(CurveTable);
-	return CurveTable;
-}
-
-void ClearGlobalCurveTable()
-{
-	IGameplayAbilitiesModule::Get().GetAbilitySystemGlobals()->AutomationTestOnly_SetGlobalCurveTable(NULL);
-}
-
-UCurveTable * GetStandardDamageOverrideCurveTable(float Factor)
-{
-	FString CSV = FString::Printf(TEXT(", 0, 1, 100\r\nStandardDamage, 0, %.2f, %.2f"), Factor * 1.f, Factor * 100.f);
-
-	UCurveTable * CurveTable = Cast<UCurveTable>(StaticConstructObject(UCurveTable::StaticClass(), GetTransientPackage(), FName(TEXT("TempCurveTable"))));
-	CurveTable->CreateTableFromCSVString(CSV);
-
-	FRichCurve * RichCurve = CurveTable->FindCurve(FName(TEXT("StandardDamage")), TEXT("Test"));
-	if (RichCurve)
-	{
-		float Value = RichCurve->Eval(5.f);
-		check(Value == 5.f * Factor);
-	}
-	
-	return CurveTable;
-}
-
-UDataTable * SetGlobalDataTable()
-{
-	// set up a test table where SpellDamage stacks and PhysicalDamage does not.
-	FString CSV(TEXT(",BaseValue,MinValue,MaxValue,DerivedAttributeInfo,bCanStack\r\nStackingAttribute1,0.0,-999.9,999.9,,True\r\nStackingAttribute2,0.0,-999.9,999.9,,True\r\nNoStackAttribute,0.0,-999.9,999.9,,False\r\n"));
-
-	UDataTable * DataTable = Cast<UDataTable>(StaticConstructObject(UDataTable::StaticClass(), GetTransientPackage(), FName(TEXT("TempDataTable"))));
-	DataTable->RowStruct = FAttributeMetaData::StaticStruct();
-	DataTable->CreateTableFromCSVString(CSV);
-
-	FAttributeMetaData * Row = (FAttributeMetaData*)DataTable->RowMap["StackingAttribute1"];
-	if (Row)
-	{
-		check(Row->bCanStack);
-	}
-	Row = (FAttributeMetaData*)DataTable->RowMap["NoStackAttribute"];
-	if (Row)
-	{
-		check(!Row->bCanStack);
-	}
-
-	IGameplayAbilitiesModule::Get().GetAbilitySystemGlobals()->AutomationTestOnly_SetGlobalAttributeDataTable(DataTable);
-	return DataTable;
-}
-
-void ClearGlobalDataTable()
-{
-	IGameplayAbilitiesModule::Get().GetAbilitySystemGlobals()->AutomationTestOnly_SetGlobalAttributeDataTable(NULL);
-}
-
-static UDataTable* CreateGameplayDataTable()
-{
-	FString CSV(TEXT(",Tag,CategoryText,"));
-	CSV.Append(TEXT("\r\n0,Damage"));
-	CSV.Append(TEXT("\r\n1,Damage.Basic"));
-	CSV.Append(TEXT("\r\n2,Damage.Type1"));
-	CSV.Append(TEXT("\r\n3,Damage.Type2"));
-	CSV.Append(TEXT("\r\n4,Damage.Reduce"));
-	CSV.Append(TEXT("\r\n5,Damage.Buffable"));
-	CSV.Append(TEXT("\r\n6,Damage.Buff"));
-	CSV.Append(TEXT("\r\n7,Damage.Physical"));
-	CSV.Append(TEXT("\r\n8,Damage.Fire"));
-	CSV.Append(TEXT("\r\n9,Damage.Buffed.FireBuff"));
-	CSV.Append(TEXT("\r\n10,Damage.Mitigated.Armor"));
-	CSV.Append(TEXT("\r\n11,Lifesteal"));
-	CSV.Append(TEXT("\r\n12,Shield"));
-	CSV.Append(TEXT("\r\n13,Buff"));
-	CSV.Append(TEXT("\r\n14,Immune"));
-	CSV.Append(TEXT("\r\n15,FireDamage"));
-	CSV.Append(TEXT("\r\n16,ShieldAbsorb"));
-	CSV.Append(TEXT("\r\n17,Stackable"));
-	CSV.Append(TEXT("\r\n18,Stack"));
-	CSV.Append(TEXT("\r\n19,Stack.CappedNumber"));
-	CSV.Append(TEXT("\r\n20,Stack.DiminishingReturns"));
-	CSV.Append(TEXT("\r\n21,Protect.Damage"));
-	CSV.Append(TEXT("\r\n22,SpellDmg.Buff"));
-	CSV.Append(TEXT("\r\n23,GameplayCue.Burning"));
-
-	UDataTable * DataTable = Cast<UDataTable>(StaticConstructObject(UDataTable::StaticClass(), GetTransientPackage(), FName(TEXT("TempDataTable"))));
-	DataTable->RowStruct = FGameplayTagTableRow::StaticStruct();
-	DataTable->CreateTableFromCSVString(CSV);
-
-	FGameplayTagTableRow * Row = (FGameplayTagTableRow*)DataTable->RowMap["0"];
-	if (Row)
-	{
-		check(Row->Tag == TEXT("Damage"));
-	}
-	return DataTable;
-}
-
-bool GameplayEffectsTest_InstantDamage_ScalingExplicit(UWorld *World, FAutomationTestBase * Test)
-{
-	// This example uses explicit scaling in a GameplayEffect. We explicitly specify the curve table to use in the GameplayEffect
-
-	const float StartHealth = 100.f;
-	const float SourceDamageScale = 1.f;
-	const float LevelOfDamage = 5.f;
-
-	// Make sure no global curve table is setup
-	ClearGlobalCurveTable();
-
-	// Sets up a linear curve table f(x)=x for StandardDamage
-	UCurveTable * SourceCurveTableOverrides = GetStandardDamageOverrideCurveTable(SourceDamageScale);
-
-	AAbilitySystemTestPawn *SourceActor = World->SpawnActor<AAbilitySystemTestPawn>();
-	AAbilitySystemTestPawn *DestActor = World->SpawnActor<AAbilitySystemTestPawn>();
-
-	UProperty *DamageProperty = FindFieldChecked<UProperty>(UAbilitySystemTestAttributeSet::StaticClass(), GET_MEMBER_NAME_CHECKED(UAbilitySystemTestAttributeSet, Damage));
-
-	UAbilitySystemComponent * SourceComponent = SourceActor->GetAbilitySystemComponent();
-	UAbilitySystemComponent * DestComponent = DestActor->GetAbilitySystemComponent();
-	SourceComponent->GetSet<UAbilitySystemTestAttributeSet>()->Health = StartHealth;
-	DestComponent->GetSet<UAbilitySystemTestAttributeSet>()->Health = StartHealth;
-
-	// Source now has SourceDamageScale (2x) damage over standard damage 
-	SourceComponent->PushGlobalCurveOveride(SourceCurveTableOverrides);
-
-	{
-		ABILITY_LOG_SCOPE(TEXT("Apply InstantDamage"));
-
-		UGameplayEffect * BaseDmgEffect = Cast<UGameplayEffect>(StaticConstructObject(UGameplayEffect::StaticClass(), GetTransientPackage(), FName(TEXT("BaseDmgEffect"))));
-		BaseDmgEffect->Modifiers.SetNum(1);
-		BaseDmgEffect->Modifiers[0].Magnitude.SetScalingValue(1.f, FName(TEXT("StandardDamage")), SourceCurveTableOverrides); // do "1*StandardDamage[Level]"
-		BaseDmgEffect->Modifiers[0].ModifierType = EGameplayMod::Attribute;
-		BaseDmgEffect->Modifiers[0].ModifierOp = EGameplayModOp::Additive;
-		BaseDmgEffect->Modifiers[0].Attribute.SetUProperty(DamageProperty);
-		BaseDmgEffect->Modifiers[0].OwnedTags.AddTag(IGameplayTagsModule::RequestGameplayTag(FName(TEXT("Damage.Basic"))));
-		BaseDmgEffect->Duration.Value = UGameplayEffect::INSTANT_APPLICATION;
-
-		SourceComponent->ApplyGameplayEffectToTarget(BaseDmgEffect, DestComponent, LevelOfDamage);
-
-		float ActualValue = DestComponent->GetSet<UAbilitySystemTestAttributeSet>()->Health;
-		float ExpectedValue = StartHealth - (LevelOfDamage * SourceDamageScale);
-		Test->TestTrue(SKILL_TEST_TEXT("Damage Applied. Health: Actual: %.2f == Exected: %.2f", ActualValue, ExpectedValue), ActualValue == ExpectedValue);
-	}
-
-	World->EditorDestroyActor(SourceActor, false);
-	World->EditorDestroyActor(DestActor, false);
-	
-	return true;
-}
-
-bool GameplayEffectsTest_InstantDamage_ScalingGlobal(UWorld *World, FAutomationTestBase * Test)
-{
-	// This example uses global scaling. The gameplay effect doesn't specify which table it uses, just that its StandardDamage. 
-	// The GameplayEffects code will fall back to the GlobalCurveTable.
-
-	const float StartHealth = 100.f;
-	const float SourceDamageScale = 2.f;
-	const float LevelOfDamage = 5.f;
-
-	SetGlobalCurveTable();
-
-	AAbilitySystemTestPawn *SourceActor = World->SpawnActor<AAbilitySystemTestPawn>();
-	AAbilitySystemTestPawn *DestActor = World->SpawnActor<AAbilitySystemTestPawn>();
-	
-	UProperty *DamageProperty = FindFieldChecked<UProperty>(UAbilitySystemTestAttributeSet::StaticClass(), GET_MEMBER_NAME_CHECKED(UAbilitySystemTestAttributeSet, Damage));
-
-	UAbilitySystemComponent * SourceComponent = SourceActor->GetAbilitySystemComponent();
-	UAbilitySystemComponent * DestComponent = DestActor->GetAbilitySystemComponent();
-	SourceComponent->GetSet<UAbilitySystemTestAttributeSet>()->Health = StartHealth;
-	DestComponent->GetSet<UAbilitySystemTestAttributeSet>()->Health = StartHealth;
-
-	{
-		ABILITY_LOG_SCOPE(TEXT("Apply InstantDamage"));
-
-		UGameplayEffect * BaseDmgEffect = Cast<UGameplayEffect>(StaticConstructObject(UGameplayEffect::StaticClass(), GetTransientPackage(), FName(TEXT("BaseDmgEffect"))));
-		BaseDmgEffect->Modifiers.SetNum(1);
-		BaseDmgEffect->Modifiers[0].Magnitude.SetScalingValue(1.f, FName(TEXT("StandardDamage")), NULL); // do "1*StandardDamage[Level]"
-		BaseDmgEffect->Modifiers[0].ModifierType = EGameplayMod::Attribute;
-		BaseDmgEffect->Modifiers[0].ModifierOp = EGameplayModOp::Additive;
-		BaseDmgEffect->Modifiers[0].Attribute.SetUProperty(DamageProperty);
-		BaseDmgEffect->Modifiers[0].OwnedTags.AddTag(IGameplayTagsModule::RequestGameplayTag(FName(TEXT("Damage.Basic"))));
-		BaseDmgEffect->Duration.Value = UGameplayEffect::INSTANT_APPLICATION;
-
-		SourceComponent->ApplyGameplayEffectToTarget(BaseDmgEffect, DestComponent, LevelOfDamage);
-
-		float ActualValue = DestComponent->GetSet<UAbilitySystemTestAttributeSet>()->Health;
-		float ExpectedValue = StartHealth - (LevelOfDamage);
-		Test->TestTrue(SKILL_TEST_TEXT("Damage Applied. Health: Actual: %.2f == Exected: %.2f", ActualValue, ExpectedValue), ActualValue == ExpectedValue);
-	}
-
-	ClearGlobalCurveTable();
-
-	World->EditorDestroyActor(SourceActor, false);
-	World->EditorDestroyActor(DestActor, false);
-
-	return true;
-}
-
-bool GameplayEffectsTest_InstantDamage_OverrideScaling(UWorld *World, FAutomationTestBase * Test)
-{
-	// This example overrides global scaling. The setup is the same as GameplayEffectsTest_InstantDamage_ScalingGlobal except now the source
-	// has an explicit override table that will take precedent of the global table.
-
-	const float StartHealth = 100.f;
-	const float SourceDamageScale = 2.f;
-	const float LevelOfDamage = 5.f;
-
-	SetGlobalCurveTable();
-	UCurveTable * SourceCurveTableOverrides = GetStandardDamageOverrideCurveTable(SourceDamageScale);
-
-	AAbilitySystemTestPawn *SourceActor = World->SpawnActor<AAbilitySystemTestPawn>();
-	AAbilitySystemTestPawn *DestActor = World->SpawnActor<AAbilitySystemTestPawn>();
-
-	UProperty *DamageProperty = FindFieldChecked<UProperty>(UAbilitySystemTestAttributeSet::StaticClass(), GET_MEMBER_NAME_CHECKED(UAbilitySystemTestAttributeSet, Damage));
-
-	UAbilitySystemComponent * SourceComponent = SourceActor->GetAbilitySystemComponent();
-	UAbilitySystemComponent * DestComponent = DestActor->GetAbilitySystemComponent();
-	SourceComponent->GetSet<UAbilitySystemTestAttributeSet>()->Health = StartHealth;
-	DestComponent->GetSet<UAbilitySystemTestAttributeSet>()->Health = StartHealth;
-
-	// Source now has SourceDamageScale (2x) damage over standard damage 
-	SourceComponent->PushGlobalCurveOveride(SourceCurveTableOverrides);
-
-	{
-		ABILITY_LOG_SCOPE(TEXT("Apply InstantDamage"));
-
-		UGameplayEffect * BaseDmgEffect = Cast<UGameplayEffect>(StaticConstructObject(UGameplayEffect::StaticClass(), GetTransientPackage(), FName(TEXT("BaseDmgEffect"))));
-		BaseDmgEffect->Modifiers.SetNum(1);
-		BaseDmgEffect->Modifiers[0].Magnitude.SetScalingValue(1.f, FName(TEXT("StandardDamage")), NULL); // do "1*StandardDamage[Level]"
-		BaseDmgEffect->Modifiers[0].ModifierType = EGameplayMod::Attribute;
-		BaseDmgEffect->Modifiers[0].ModifierOp = EGameplayModOp::Additive;
-		BaseDmgEffect->Modifiers[0].Attribute.SetUProperty(DamageProperty);
-		BaseDmgEffect->Modifiers[0].OwnedTags.AddTag(IGameplayTagsModule::RequestGameplayTag(FName(TEXT("Damage.Basic"))));
-		BaseDmgEffect->Duration.Value = UGameplayEffect::INSTANT_APPLICATION;
-
-		SourceComponent->ApplyGameplayEffectToTarget(BaseDmgEffect, DestComponent, LevelOfDamage);
-
-		float ActualValue = DestComponent->GetSet<UAbilitySystemTestAttributeSet>()->Health;
-		float ExpectedValue = StartHealth - (LevelOfDamage * SourceDamageScale);
-		Test->TestTrue(SKILL_TEST_TEXT("Damage Applied. Health: Actual: %.2f == Exected: %.2f", ActualValue, ExpectedValue), ActualValue == ExpectedValue);
-	}
-
-	ClearGlobalCurveTable();
-
-	World->EditorDestroyActor(SourceActor, false);
-	World->EditorDestroyActor(DestActor, false);
-
-	return true;
-}
-
-bool GameplayEffectsTest_InstantDamageRequiredTag(UWorld *World, FAutomationTestBase * Test)
-{
-	const float StartHealth = 100.f;
-	const float DamageValue = -5.f;
-	const float DamageProtectionDivisor = 2.f;
-
-	AAbilitySystemTestPawn *SourceActor = World->SpawnActor<AAbilitySystemTestPawn>();
-	AAbilitySystemTestPawn *DestActor = World->SpawnActor<AAbilitySystemTestPawn>();
-
-	UProperty *HealthProperty = FindFieldChecked<UProperty>(UAbilitySystemTestAttributeSet::StaticClass(), GET_MEMBER_NAME_CHECKED(UAbilitySystemTestAttributeSet, Health));
-
-	UAbilitySystemComponent * SourceComponent = SourceActor->GetAbilitySystemComponent();
-	UAbilitySystemComponent * DestComponent = DestActor->GetAbilitySystemComponent();
-	SourceComponent->GetSet<UAbilitySystemTestAttributeSet>()->Health = StartHealth;
-	DestComponent->GetSet<UAbilitySystemTestAttributeSet>()->Health = StartHealth;
-
-	// Setup a GE to modify IncomingGEs
-	{
-		ABILITY_LOG_SCOPE(TEXT("Apply ProtectionBuff"))
-
-		UGameplayEffect* BaseProtectEffect = Cast<UGameplayEffect>(StaticConstructObject(UGameplayEffect::StaticClass(), GetTransientPackage(), FName(TEXT("ProtectBuff"))));
-		BaseProtectEffect->Modifiers.SetNum(1);
-		BaseProtectEffect->Modifiers[0].Magnitude.SetValue(DamageProtectionDivisor);
-		BaseProtectEffect->Modifiers[0].ModifierType = EGameplayMod::IncomingGE;
-		BaseProtectEffect->Modifiers[0].ModifierOp = EGameplayModOp::Division;
-		BaseProtectEffect->Modifiers[0].Attribute.SetUProperty(HealthProperty);
-		BaseProtectEffect->Modifiers[0].OwnedTags.AddTag(IGameplayTagsModule::RequestGameplayTag(FName(TEXT("Protect.Damage"))));
-		BaseProtectEffect->Duration.SetValue(UGameplayEffect::INFINITE_DURATION);
-		BaseProtectEffect->CopyPolicy = EGameplayEffectCopyPolicy::AlwaysLink;
-		BaseProtectEffect->GameplayEffectRequiredTags.AddTag(IGameplayTagsModule::RequestGameplayTag(FName(TEXT("Damage.Type2"))));
-
-		// Apply to self
-		DestComponent->ApplyGameplayEffectToTarget(BaseProtectEffect, DestComponent, 1.f);
-	}
-
-	// Apply Damage
-	{
-		ABILITY_LOG_SCOPE(TEXT("Apply InstantDamage"))
-
-		UGameplayEffect * BaseDmgEffect = Cast<UGameplayEffect>(StaticConstructObject(UGameplayEffect::StaticClass(), GetTransientPackage(), FName(TEXT("Damage"))));
-		BaseDmgEffect->Modifiers.SetNum(1);
-		BaseDmgEffect->Modifiers[0].Magnitude.SetValue(DamageValue);
-		BaseDmgEffect->Modifiers[0].ModifierType = EGameplayMod::Attribute;
-		BaseDmgEffect->Modifiers[0].ModifierOp = EGameplayModOp::Additive;
-		BaseDmgEffect->Modifiers[0].Attribute.SetUProperty(HealthProperty);
-		BaseDmgEffect->Modifiers[0].OwnedTags.AddTag(IGameplayTagsModule::RequestGameplayTag(FName(TEXT("Damage.Basic"))));
-		BaseDmgEffect->Duration.Value = UGameplayEffect::INSTANT_APPLICATION;
-		BaseDmgEffect->GameplayEffectTags.AddTag(IGameplayTagsModule::RequestGameplayTag(FName(TEXT("Damage.Type1"))));
-
-		// Apply to target
-		SourceComponent->ApplyGameplayEffectToTarget(BaseDmgEffect, DestComponent, 1.f);
-
-		float ExpectedValue = (StartHealth + DamageValue);
-		
-		Test->TestTrue(SKILL_TEST_TEXT("Instant Damage Required Tag No Protection"), (DestComponent->GetSet<UAbilitySystemTestAttributeSet>()->Health == ExpectedValue));
-		ABILITY_LOG(Log, TEXT("Final Health: %.2f"), DestComponent->GetSet<UAbilitySystemTestAttributeSet>()->Health);
-	}
-
-	// reset health
-	DestComponent->GetSet<UAbilitySystemTestAttributeSet>()->Health = StartHealth;
-
-	// Apply Damage
-	{
-		ABILITY_LOG_SCOPE(TEXT("Apply InstantDamage"))
-
-		UGameplayEffect * BaseDmgEffect = Cast<UGameplayEffect>(StaticConstructObject(UGameplayEffect::StaticClass(), GetTransientPackage(), FName(TEXT("Damage"))));
-		BaseDmgEffect->Modifiers.SetNum(1);
-		BaseDmgEffect->Modifiers[0].Magnitude.SetValue(DamageValue);
-		BaseDmgEffect->Modifiers[0].ModifierType = EGameplayMod::Attribute;
-		BaseDmgEffect->Modifiers[0].ModifierOp = EGameplayModOp::Additive;
-		BaseDmgEffect->Modifiers[0].Attribute.SetUProperty(HealthProperty);
-		BaseDmgEffect->Modifiers[0].OwnedTags.AddTag(IGameplayTagsModule::RequestGameplayTag(FName(TEXT("Damage.Basic"))));
-		BaseDmgEffect->Duration.Value = UGameplayEffect::INSTANT_APPLICATION;
-		BaseDmgEffect->GameplayEffectTags.AddTag(IGameplayTagsModule::RequestGameplayTag(FName(TEXT("Damage.Type2"))));
-
-		// Apply to target
-		SourceComponent->ApplyGameplayEffectToTarget(BaseDmgEffect, DestComponent, 1.f);
-
-		float ExpectedValue = (StartHealth + (DamageValue / DamageProtectionDivisor));
-		float ActualValue = DestComponent->GetSet<UAbilitySystemTestAttributeSet>()->Health;
-		
-		Test->TestTrue(SKILL_TEST_TEXT("Damage Applied. Health: Actual: %.2f == Exected: %.2f", ActualValue, ExpectedValue), ActualValue == ExpectedValue);
-	}
-
-	World->EditorDestroyActor(SourceActor, false);
-	World->EditorDestroyActor(DestActor, false);
-
-	return true;
-}
-
-bool GameplayEffectsTest_InstantDamageIgnoreTag(UWorld *World, FAutomationTestBase * Test)
-{
-	const float StartHealth = 100.f;
-	const float DamageValue = -5.f;
-	const float DamageProtectionDivisor = 2.f;
-
-	AAbilitySystemTestPawn *SourceActor = World->SpawnActor<AAbilitySystemTestPawn>();
-	AAbilitySystemTestPawn *DestActor = World->SpawnActor<AAbilitySystemTestPawn>();
-
-	UProperty *HealthProperty = FindFieldChecked<UProperty>(UAbilitySystemTestAttributeSet::StaticClass(), GET_MEMBER_NAME_CHECKED(UAbilitySystemTestAttributeSet, Health));
-
-	UAbilitySystemComponent * SourceComponent = SourceActor->GetAbilitySystemComponent();
-	UAbilitySystemComponent * DestComponent = DestActor->GetAbilitySystemComponent();
-	SourceComponent->GetSet<UAbilitySystemTestAttributeSet>()->Health = StartHealth;
-	DestComponent->GetSet<UAbilitySystemTestAttributeSet>()->Health = StartHealth;
-
-	// Setup a GE to modify IncomingGEs
-	{
-		ABILITY_LOG_SCOPE(TEXT("Apply ProtectionBuff"))
-
-		UGameplayEffect* BaseProtectEffect = Cast<UGameplayEffect>(StaticConstructObject(UGameplayEffect::StaticClass(), GetTransientPackage(), FName(TEXT("ProtectBuff"))));
-		BaseProtectEffect->Modifiers.SetNum(1);
-		BaseProtectEffect->Modifiers[0].Magnitude.SetValue(DamageProtectionDivisor);
-		BaseProtectEffect->Modifiers[0].ModifierType = EGameplayMod::IncomingGE;
-		BaseProtectEffect->Modifiers[0].ModifierOp = EGameplayModOp::Division;
-		BaseProtectEffect->Modifiers[0].Attribute.SetUProperty(HealthProperty);
-		BaseProtectEffect->Modifiers[0].OwnedTags.AddTag(IGameplayTagsModule::RequestGameplayTag(FName(TEXT("Protect.Damage"))));
-		BaseProtectEffect->Duration.SetValue(UGameplayEffect::INFINITE_DURATION);
-		BaseProtectEffect->CopyPolicy = EGameplayEffectCopyPolicy::AlwaysLink;
-		BaseProtectEffect->GameplayEffectIgnoreTags.AddTag(IGameplayTagsModule::RequestGameplayTag(FName(TEXT("Damage.Type1"))));
-
-		// Apply to self
-		DestComponent->ApplyGameplayEffectToTarget(BaseProtectEffect, DestComponent, 1.f);
-	}
-
-	// Apply Damage
-	{
-		ABILITY_LOG_SCOPE(TEXT("Apply InstantDamage"))
-
-		UGameplayEffect * BaseDmgEffect = Cast<UGameplayEffect>(StaticConstructObject(UGameplayEffect::StaticClass(), GetTransientPackage(), FName(TEXT("Damage"))));
-		BaseDmgEffect->Modifiers.SetNum(1);
-		BaseDmgEffect->Modifiers[0].Magnitude.SetValue(DamageValue);
-		BaseDmgEffect->Modifiers[0].ModifierType = EGameplayMod::Attribute;
-		BaseDmgEffect->Modifiers[0].ModifierOp = EGameplayModOp::Additive;
-		BaseDmgEffect->Modifiers[0].Attribute.SetUProperty(HealthProperty);
-		BaseDmgEffect->Modifiers[0].OwnedTags.AddTag(IGameplayTagsModule::RequestGameplayTag(FName(TEXT("Damage.Basic"))));
-		BaseDmgEffect->Duration.Value = UGameplayEffect::INSTANT_APPLICATION;
-		BaseDmgEffect->GameplayEffectTags.AddTag(IGameplayTagsModule::RequestGameplayTag(FName(TEXT("Damage.Type1"))));
-
-		// Apply to target
-		SourceComponent->ApplyGameplayEffectToTarget(BaseDmgEffect, DestComponent, 1.f);
-
-		float ExpectedValue = (StartHealth + DamageValue);
-
-		Test->TestTrue(SKILL_TEST_TEXT("Instant Damage Ignore Tag No Protection"), (DestComponent->GetSet<UAbilitySystemTestAttributeSet>()->Health == ExpectedValue));
-		ABILITY_LOG(Log, TEXT("Final Health: %.2f"), DestComponent->GetSet<UAbilitySystemTestAttributeSet>()->Health);
-	}
-
-	// reset health
-	DestComponent->GetSet<UAbilitySystemTestAttributeSet>()->Health = StartHealth;
-
-	// Apply Damage
-	{
-		ABILITY_LOG_SCOPE(TEXT("Apply InstantDamage"))
-
-		UGameplayEffect * BaseDmgEffect = Cast<UGameplayEffect>(StaticConstructObject(UGameplayEffect::StaticClass(), GetTransientPackage(), FName(TEXT("Damage"))));
-		BaseDmgEffect->Modifiers.SetNum(1);
-		BaseDmgEffect->Modifiers[0].Magnitude.SetValue(DamageValue);
-		BaseDmgEffect->Modifiers[0].ModifierType = EGameplayMod::Attribute;
-		BaseDmgEffect->Modifiers[0].ModifierOp = EGameplayModOp::Additive;
-		BaseDmgEffect->Modifiers[0].Attribute.SetUProperty(HealthProperty);
-		BaseDmgEffect->Modifiers[0].OwnedTags.AddTag(IGameplayTagsModule::RequestGameplayTag(FName(TEXT("Damage.Basic"))));
-		BaseDmgEffect->Duration.Value = UGameplayEffect::INSTANT_APPLICATION;
-		BaseDmgEffect->GameplayEffectTags.AddTag(IGameplayTagsModule::RequestGameplayTag(FName(TEXT("Damage.Type2"))));
-
-		// Apply to target
-		SourceComponent->ApplyGameplayEffectToTarget(BaseDmgEffect, DestComponent, 1.f);
-
-		float ExpectedValue = (StartHealth + (DamageValue / DamageProtectionDivisor));
-
-		Test->TestTrue(SKILL_TEST_TEXT("Instant Damage Ignore Tag Protected"), (DestComponent->GetSet<UAbilitySystemTestAttributeSet>()->Health == ExpectedValue));
-		ABILITY_LOG(Log, TEXT("Final Health: %.2f"), DestComponent->GetSet<UAbilitySystemTestAttributeSet>()->Health);
-	}
-
-	World->EditorDestroyActor(SourceActor, false);
-	World->EditorDestroyActor(DestActor, false);
-
-	return true;
-}
-
-bool GameplayEffectsTest_InstantDamageModifierPassesTag(UWorld *World, FAutomationTestBase * Test)
-{
-	const float StartHealth = 100.f;
-	const float DamageValue = -5.f;
-	const float BonusDamageMultiplier = 2.f;
-	const float DamageProtectionDivisor = 2.f;
-
-	AAbilitySystemTestPawn *SourceActor = World->SpawnActor<AAbilitySystemTestPawn>();
-	AAbilitySystemTestPawn *DestActor = World->SpawnActor<AAbilitySystemTestPawn>();
-
-	UProperty *HealthProperty = FindFieldChecked<UProperty>(UAbilitySystemTestAttributeSet::StaticClass(), GET_MEMBER_NAME_CHECKED(UAbilitySystemTestAttributeSet, Health));
-
-	UAbilitySystemComponent * SourceComponent = SourceActor->GetAbilitySystemComponent();
-	UAbilitySystemComponent * DestComponent = DestActor->GetAbilitySystemComponent();
-	SourceComponent->GetSet<UAbilitySystemTestAttributeSet>()->Health = StartHealth;
-	DestComponent->GetSet<UAbilitySystemTestAttributeSet>()->Health = StartHealth;
-
-	// Setup a GE to modify OutgoingGEs
-	{
-		ABILITY_LOG_SCOPE(TEXT("Apply DamageBuff"))
-
-		UGameplayEffect * BaseDmgEffect = Cast<UGameplayEffect>(StaticConstructObject(UGameplayEffect::StaticClass(), GetTransientPackage(), FName(TEXT("DamageBuff"))));
-		BaseDmgEffect->Modifiers.SetNum(1);
-		BaseDmgEffect->Modifiers[0].Magnitude.SetValue(BonusDamageMultiplier);
-		BaseDmgEffect->Modifiers[0].ModifierType = EGameplayMod::OutgoingGE;
-		BaseDmgEffect->Modifiers[0].ModifierOp = EGameplayModOp::Multiplicitive;
-		BaseDmgEffect->Modifiers[0].Attribute.SetUProperty(HealthProperty);
-		BaseDmgEffect->Modifiers[0].OwnedTags.AddTag(IGameplayTagsModule::RequestGameplayTag(FName(TEXT("Damage.Type1"))));		// When I am applied, the damage modifier gets this tag.
-		BaseDmgEffect->Duration.SetValue(UGameplayEffect::INFINITE_DURATION);
-
-		// Apply to self
-		SourceComponent->ApplyGameplayEffectToTarget(BaseDmgEffect, SourceComponent, 1.f);
-	}
-
-	// Setup a GE to modify IncomingGEs
-	{
-		ABILITY_LOG_SCOPE(TEXT("Apply ProtectionBuff"))
-
-		UGameplayEffect* BaseProtectEffect = Cast<UGameplayEffect>(StaticConstructObject(UGameplayEffect::StaticClass(), GetTransientPackage(), FName(TEXT("ProtectBuff"))));
-		BaseProtectEffect->Modifiers.SetNum(1);
-		BaseProtectEffect->Modifiers[0].Magnitude.SetValue(DamageProtectionDivisor);
-		BaseProtectEffect->Modifiers[0].ModifierType = EGameplayMod::IncomingGE;
-		BaseProtectEffect->Modifiers[0].ModifierOp = EGameplayModOp::Division;
-		BaseProtectEffect->Modifiers[0].Attribute.SetUProperty(HealthProperty);
-		BaseProtectEffect->Modifiers[0].OwnedTags.AddTag(IGameplayTagsModule::RequestGameplayTag(FName(TEXT("Protect.Damage"))));
-		BaseProtectEffect->Modifiers[0].RequiredTags.AddTag(IGameplayTagsModule::RequestGameplayTag(FName(TEXT("Damage.Type1"))));
-		BaseProtectEffect->Duration.SetValue(UGameplayEffect::INFINITE_DURATION);
-		BaseProtectEffect->CopyPolicy = EGameplayEffectCopyPolicy::AlwaysLink;
-
-		// Apply to self
-		DestComponent->ApplyGameplayEffectToTarget(BaseProtectEffect, DestComponent, 1.f);
-	}
-
-	// Apply Damage
-	{
-		ABILITY_LOG_SCOPE(TEXT("Apply InstantDamage"))
-
-		UGameplayEffect * BaseDmgEffect = Cast<UGameplayEffect>(StaticConstructObject(UGameplayEffect::StaticClass(), GetTransientPackage(), FName(TEXT("Damage"))));
-		BaseDmgEffect->Modifiers.SetNum(1);
-		BaseDmgEffect->Modifiers[0].Magnitude.SetValue(DamageValue);
-		BaseDmgEffect->Modifiers[0].ModifierType = EGameplayMod::Attribute;
-		BaseDmgEffect->Modifiers[0].ModifierOp = EGameplayModOp::Additive;
-		BaseDmgEffect->Modifiers[0].Attribute.SetUProperty(HealthProperty);
-		BaseDmgEffect->Modifiers[0].OwnedTags.AddTag(IGameplayTagsModule::RequestGameplayTag(FName(TEXT("Damage.Basic"))));
-		BaseDmgEffect->Duration.Value = UGameplayEffect::INSTANT_APPLICATION;
-
-		// Apply to target
-		SourceComponent->ApplyGameplayEffectToTarget(BaseDmgEffect, DestComponent, 1.f);
-
-		float ExpectedValue = (StartHealth + ((DamageValue * BonusDamageMultiplier) / DamageProtectionDivisor));
-		float ActualValue = DestComponent->GetSet<UAbilitySystemTestAttributeSet>()->Health;
-
-		Test->TestTrue(SKILL_TEST_TEXT("Instant Damage Required Tag No Protection.  Actual: %.2f == Exected: %.2f", ActualValue, ExpectedValue), (ActualValue == ExpectedValue));
-		ABILITY_LOG(Log, TEXT("Final Health: %.2f"), DestComponent->GetSet<UAbilitySystemTestAttributeSet>()->Health);
-	}
-
-	World->EditorDestroyActor(SourceActor, false);
-	World->EditorDestroyActor(DestActor, false);
-
-	return true;
-}
-
-bool GameplayEffectsTest_InstantDamageModifierTag(UWorld *World, FAutomationTestBase * Test)
-{
-	const float StartHealth = 100.f;
-	const float DamageValue = -5.f;
-	const float BonusDamageValue = -10.f;
-	const float DamageProtectionDivisor = 2.f;
-
-	AAbilitySystemTestPawn *SourceActor = World->SpawnActor<AAbilitySystemTestPawn>();
-	AAbilitySystemTestPawn *DestActor = World->SpawnActor<AAbilitySystemTestPawn>();
-
-	UProperty *HealthProperty = FindFieldChecked<UProperty>(UAbilitySystemTestAttributeSet::StaticClass(), GET_MEMBER_NAME_CHECKED(UAbilitySystemTestAttributeSet, Health));
-
-	UAbilitySystemComponent * SourceComponent = SourceActor->GetAbilitySystemComponent();
-	UAbilitySystemComponent * DestComponent = DestActor->GetAbilitySystemComponent();
-	SourceComponent->GetSet<UAbilitySystemTestAttributeSet>()->Health = StartHealth;
-	DestComponent->GetSet<UAbilitySystemTestAttributeSet>()->Health = StartHealth;
-
-	// Setup a GE to modify OutgoingGEs
-	{
-		ABILITY_LOG_SCOPE(TEXT("Apply DamageBuff"))
-
-		UGameplayEffect * BaseDmgEffect = Cast<UGameplayEffect>(StaticConstructObject(UGameplayEffect::StaticClass(), GetTransientPackage(), FName(TEXT("DamageBuff"))));
-		BaseDmgEffect->Modifiers.SetNum(1);
-		BaseDmgEffect->Modifiers[0].Magnitude.SetValue(BonusDamageValue);
-		BaseDmgEffect->Modifiers[0].ModifierType = EGameplayMod::OutgoingGE;
-		BaseDmgEffect->Modifiers[0].ModifierOp = EGameplayModOp::Additive;
-		BaseDmgEffect->Modifiers[0].Attribute.SetUProperty(HealthProperty);
-		BaseDmgEffect->Modifiers[0].OwnedTags.AddTag(IGameplayTagsModule::RequestGameplayTag(FName(TEXT("Damage.Type1"))));
-		BaseDmgEffect->Duration.SetValue(UGameplayEffect::INFINITE_DURATION);
-
-		// Apply to self
-		SourceComponent->ApplyGameplayEffectToTarget(BaseDmgEffect, SourceComponent, 1.f);
-	}
-
-	// Setup a GE to modify IncomingGEs
-	{
-		ABILITY_LOG_SCOPE(TEXT("Apply ProtectionBuff"))
-
-		UGameplayEffect* BaseProtectEffect = Cast<UGameplayEffect>(StaticConstructObject(UGameplayEffect::StaticClass(), GetTransientPackage(), FName(TEXT("ProtectBuff"))));
-		BaseProtectEffect->Modifiers.SetNum(1);
-		BaseProtectEffect->Modifiers[0].Magnitude.SetValue(DamageProtectionDivisor);
-		BaseProtectEffect->Modifiers[0].ModifierType = EGameplayMod::IncomingGE;
-		BaseProtectEffect->Modifiers[0].ModifierOp = EGameplayModOp::Division;
-		BaseProtectEffect->Modifiers[0].Attribute.SetUProperty(HealthProperty);
-		BaseProtectEffect->Modifiers[0].OwnedTags.AddTag(IGameplayTagsModule::RequestGameplayTag(FName(TEXT("Protect.Damage"))));
-		BaseProtectEffect->Modifiers[0].RequiredTags.AddTag(IGameplayTagsModule::RequestGameplayTag(FName(TEXT("Damage.Type1"))));
-		BaseProtectEffect->Duration.SetValue(UGameplayEffect::INFINITE_DURATION);
-		BaseProtectEffect->CopyPolicy = EGameplayEffectCopyPolicy::AlwaysLink;
-
-		// Apply to self
-		DestComponent->ApplyGameplayEffectToTarget(BaseProtectEffect, DestComponent, 1.f);
-	}
-
-	// Apply Damage
-	{
-		ABILITY_LOG_SCOPE(TEXT("Apply InstantDamage"))
-
-		UGameplayEffect * BaseDmgEffect = Cast<UGameplayEffect>(StaticConstructObject(UGameplayEffect::StaticClass(), GetTransientPackage(), FName(TEXT("Damage"))));
-		BaseDmgEffect->Modifiers.SetNum(1);
-		BaseDmgEffect->Modifiers[0].Magnitude.SetValue(DamageValue);
-		BaseDmgEffect->Modifiers[0].ModifierType = EGameplayMod::Attribute;
-		BaseDmgEffect->Modifiers[0].ModifierOp = EGameplayModOp::Additive;
-		BaseDmgEffect->Modifiers[0].Attribute.SetUProperty(HealthProperty);
-		BaseDmgEffect->Modifiers[0].OwnedTags.AddTag(IGameplayTagsModule::RequestGameplayTag(FName(TEXT("Damage.Basic"))));
-		BaseDmgEffect->Duration.Value = UGameplayEffect::INSTANT_APPLICATION;
-
-		// Apply to target
-		SourceComponent->ApplyGameplayEffectToTarget(BaseDmgEffect, DestComponent, 1.f);
-
-		float ExpectedValue = (StartHealth + ((DamageValue + BonusDamageValue) / DamageProtectionDivisor));
-		float ActualValue = DestComponent->GetSet<UAbilitySystemTestAttributeSet>()->Health;
-
-		Test->TestTrue(SKILL_TEST_TEXT("Buff Instant Damage Applied.  Actual: %.2f == Exected: %.2f", ActualValue, ExpectedValue), (ActualValue == ExpectedValue));
-	}
-
-	World->EditorDestroyActor(SourceActor, false);
-	World->EditorDestroyActor(DestActor, false);
-
-	return true;
-}
-
-
-bool GameplayEffectsTest_InstantDamage_ScalingProperty(UWorld *World, FAutomationTestBase * Test)
-{
-	// This example we scale Damage based off the instigator's PhysicalDamage attribute
-
-	const float StartHealth = 100.f;
-	const float PhysicalDamage = 10.f;
-	const float GameplayEffectScaling = 1.f;
-
-	SetGlobalCurveTable();
-
-	AAbilitySystemTestPawn *SourceActor = World->SpawnActor<AAbilitySystemTestPawn>();
-	AAbilitySystemTestPawn *DestActor = World->SpawnActor<AAbilitySystemTestPawn>();
-
-	UProperty *DamageProperty = FindFieldChecked<UProperty>(UAbilitySystemTestAttributeSet::StaticClass(), GET_MEMBER_NAME_CHECKED(UAbilitySystemTestAttributeSet, Damage));
-	UProperty *PhysicalDamageProperty = FindFieldChecked<UProperty>(UAbilitySystemTestAttributeSet::StaticClass(), GET_MEMBER_NAME_CHECKED(UAbilitySystemTestAttributeSet, PhysicalDamage));
-
-	UAbilitySystemComponent * SourceComponent = SourceActor->GetAbilitySystemComponent();
-	UAbilitySystemComponent * DestComponent = DestActor->GetAbilitySystemComponent();
-	SourceComponent->GetSet<UAbilitySystemTestAttributeSet>()->Health = StartHealth;
-	DestComponent->GetSet<UAbilitySystemTestAttributeSet>()->Health = StartHealth;
-
-	SourceComponent->GetSet<UAbilitySystemTestAttributeSet>()->PhysicalDamage = PhysicalDamage;
-
-	{
-		ABILITY_LOG_SCOPE(TEXT("Apply InstantDamage"));
-
-		// This effects do Damage = 1.f * LinearCurve[LevelOfGameplayEffect].
-		// This translate into 1.f * PhysicalDamage.
-
-		UGameplayEffect * BaseDmgEffect = Cast<UGameplayEffect>(StaticConstructObject(UGameplayEffect::StaticClass(), GetTransientPackage(), FName(TEXT("BaseDmgEffect"))));
-		BaseDmgEffect->Modifiers.SetNum(1);
-		BaseDmgEffect->Modifiers[0].Magnitude.SetScalingValue(GameplayEffectScaling, FName(TEXT("LinearCurve")), NULL); // do "1*StandardDamage[Level]"
-		BaseDmgEffect->Modifiers[0].ModifierType = EGameplayMod::Attribute;
-		BaseDmgEffect->Modifiers[0].ModifierOp = EGameplayModOp::Additive;
-		BaseDmgEffect->Modifiers[0].Attribute.SetUProperty(DamageProperty);
-		BaseDmgEffect->Modifiers[0].OwnedTags.AddTag(IGameplayTagsModule::RequestGameplayTag(FName(TEXT("Damage.Basic"))));
-		BaseDmgEffect->Duration.Value = UGameplayEffect::INSTANT_APPLICATION;		
-		BaseDmgEffect->LevelInfo.Attribute.SetUProperty(PhysicalDamageProperty);
-		BaseDmgEffect->LevelInfo.InheritLevelFromOwner = false;
-		
-
-		SourceComponent->ApplyGameplayEffectToTarget(BaseDmgEffect, DestComponent);
-
-		float ActualValue = DestComponent->GetSet<UAbilitySystemTestAttributeSet>()->Health;
-		float ExpectedValue = StartHealth - (PhysicalDamage * GameplayEffectScaling);
-		Test->TestTrue(SKILL_TEST_TEXT("Damage Applied. Health: Actual: %.2f == Exected: %.2f", ActualValue, ExpectedValue), ActualValue == ExpectedValue);
-	}
-
-	ClearGlobalCurveTable();
-
-	World->EditorDestroyActor(SourceActor, false);
-	World->EditorDestroyActor(DestActor, false);
-
-	return true;
-}
-
-bool GameplayEffectsTest_InstantDamage_ScalingPropertyNested(UWorld *World, FAutomationTestBase * Test)
-{
-	// This accomplishes the same as GameplayEffectsTest_InstantDamage_ScalingProperty but the leveling info is specified at the modifier, not gameplayeffect, level.
-
-	const float StartHealth = 100.f;
-	const float PhysicalDamage = 10.f;
-	const float GameplayEffectScaling = 1.f;
-
-	SetGlobalCurveTable();
-
-	AAbilitySystemTestPawn *SourceActor = World->SpawnActor<AAbilitySystemTestPawn>();
-	AAbilitySystemTestPawn *DestActor = World->SpawnActor<AAbilitySystemTestPawn>();
-
-	UProperty *DamageProperty = FindFieldChecked<UProperty>(UAbilitySystemTestAttributeSet::StaticClass(), GET_MEMBER_NAME_CHECKED(UAbilitySystemTestAttributeSet, Damage));
-	UProperty *PhysicalDamageProperty = FindFieldChecked<UProperty>(UAbilitySystemTestAttributeSet::StaticClass(), GET_MEMBER_NAME_CHECKED(UAbilitySystemTestAttributeSet, PhysicalDamage));
-
-	UAbilitySystemComponent * SourceComponent = SourceActor->GetAbilitySystemComponent();
-	UAbilitySystemComponent * DestComponent = DestActor->GetAbilitySystemComponent();
-	SourceComponent->GetSet<UAbilitySystemTestAttributeSet>()->Health = StartHealth;
-	DestComponent->GetSet<UAbilitySystemTestAttributeSet>()->Health = StartHealth;
-
-	SourceComponent->GetSet<UAbilitySystemTestAttributeSet>()->PhysicalDamage = PhysicalDamage;
-
-	{
-		ABILITY_LOG_SCOPE(TEXT("Apply InstantDamage"));
-
-		// This effects do Damage = 1.f * LinearCurve[LevelOfGameplayEffect].
-		// This translate into 1.f * PhysicalDamage.
-
-		UGameplayEffect * BaseDmgEffect = Cast<UGameplayEffect>(StaticConstructObject(UGameplayEffect::StaticClass(), GetTransientPackage(), FName(TEXT("BaseDmgEffect"))));
-		BaseDmgEffect->Modifiers.SetNum(1);
-		BaseDmgEffect->Modifiers[0].Magnitude.SetScalingValue(GameplayEffectScaling, FName(TEXT("LinearCurve")), NULL); // do "1*StandardDamage[Level]"
-		BaseDmgEffect->Modifiers[0].ModifierType = EGameplayMod::Attribute;
-		BaseDmgEffect->Modifiers[0].ModifierOp = EGameplayModOp::Additive;
-		BaseDmgEffect->Modifiers[0].Attribute.SetUProperty(DamageProperty);
-		BaseDmgEffect->Modifiers[0].OwnedTags.AddTag(IGameplayTagsModule::RequestGameplayTag(FName(TEXT("Damage.Basic"))));
-		BaseDmgEffect->Modifiers[0].LevelInfo.Attribute.SetUProperty(PhysicalDamageProperty);
-		BaseDmgEffect->Modifiers[0].LevelInfo.InheritLevelFromOwner = false;
-		BaseDmgEffect->Duration.Value = UGameplayEffect::INSTANT_APPLICATION;
-
-		SourceComponent->ApplyGameplayEffectToTarget(BaseDmgEffect, DestComponent);
-
-		float ActualValue = DestComponent->GetSet<UAbilitySystemTestAttributeSet>()->Health;
-		float ExpectedValue = StartHealth - (PhysicalDamage * GameplayEffectScaling);
-		Test->TestTrue(SKILL_TEST_TEXT("Damage Applied. Health: Actual: %.2f == Exected: %.2f", ActualValue, ExpectedValue), ActualValue == ExpectedValue);
-	}
-
-	ClearGlobalCurveTable();
-
-	World->EditorDestroyActor(SourceActor, false);
-	World->EditorDestroyActor(DestActor, false);
-
-	return true;
-}
-
-bool GameplayEffectsTest_DotDamage_ScalingProperty_Snapshot(UWorld *World, FAutomationTestBase * Test)
-{
-	// Add a dot that is powered by SpellDamage. Increase SpellDamage after applying, confirm it doesn't add extra damage to subsequent ticks.
-
-	const float StartHealth = 100.f;
-	const float SpellDamage = 10.f;
-	const float SpellDamage2 = 50.f;
-	const float GameplayEffectScaling = 1.f;
-
-
-	SetGlobalCurveTable();
-
-	AAbilitySystemTestPawn *SourceActor = World->SpawnActor<AAbilitySystemTestPawn>();
-	AAbilitySystemTestPawn *DestActor = World->SpawnActor<AAbilitySystemTestPawn>();
-
-	UProperty *DamageProperty = FindFieldChecked<UProperty>(UAbilitySystemTestAttributeSet::StaticClass(), GET_MEMBER_NAME_CHECKED(UAbilitySystemTestAttributeSet, Damage));
-	UProperty *SpellDamageProperty = FindFieldChecked<UProperty>(UAbilitySystemTestAttributeSet::StaticClass(), GET_MEMBER_NAME_CHECKED(UAbilitySystemTestAttributeSet, SpellDamage));
-
-	UAbilitySystemComponent * SourceComponent = SourceActor->GetAbilitySystemComponent();
-	UAbilitySystemComponent * DestComponent = DestActor->GetAbilitySystemComponent();
-	SourceComponent->GetSet<UAbilitySystemTestAttributeSet>()->Health = StartHealth;
-	DestComponent->GetSet<UAbilitySystemTestAttributeSet>()->Health = StartHealth;
-
-	SourceComponent->GetSet<UAbilitySystemTestAttributeSet>()->SpellDamage = SpellDamage;
-
-	{
-		ABILITY_LOG_SCOPE(TEXT("Apply InstantDamage"));
-
-		// This effects do Damage = 1.f * LinearCurve[LevelOfGameplayEffect].
-		// This translate into 1.f * PhysicalDamage.
-
-		UGameplayEffect * BaseDmgEffect = Cast<UGameplayEffect>(StaticConstructObject(UGameplayEffect::StaticClass(), GetTransientPackage(), FName(TEXT("BaseDmgEffect"))));
-		BaseDmgEffect->Modifiers.SetNum(1);
-		BaseDmgEffect->Modifiers[0].Magnitude.SetScalingValue(GameplayEffectScaling, FName(TEXT("LinearCurve")), NULL); // do "1*StandardDamage[Level]"
-		BaseDmgEffect->Modifiers[0].ModifierType = EGameplayMod::Attribute;
-		BaseDmgEffect->Modifiers[0].ModifierOp = EGameplayModOp::Additive;
-		BaseDmgEffect->Modifiers[0].Attribute.SetUProperty(DamageProperty);				// Modifies target's "Damage" attribute (-health)
-		BaseDmgEffect->Modifiers[0].OwnedTags.AddTag(IGameplayTagsModule::RequestGameplayTag(FName(TEXT("Damage.Basic"))));
-		BaseDmgEffect->Duration.Value = UGameplayEffect::INFINITE_DURATION;
-		BaseDmgEffect->Period.Value = 1.f;
-
-		BaseDmgEffect->LevelInfo.Attribute.SetUProperty(SpellDamageProperty);			// Powered by instigators SpellDamage
-		BaseDmgEffect->LevelInfo.InheritLevelFromOwner = false;
-		BaseDmgEffect->LevelInfo.TakeSnapshotOnInit = true;								// But just a snapshot of their SpellDamage when we are applied
-
-		SourceComponent->ApplyGameplayEffectToTarget(BaseDmgEffect, DestComponent);
-	}
-
-
-	{
-		// Increase spell damage on instigator (after we already applied the DOT)
-		UGameplayEffect * SpellDmgEffect = Cast<UGameplayEffect>(StaticConstructObject(UGameplayEffect::StaticClass(), GetTransientPackage(), FName(TEXT("SpellDmgEffect"))));
-		SpellDmgEffect->Modifiers.SetNum(1);
-		SpellDmgEffect->Modifiers[0].Magnitude.SetValue(SpellDamage2);
-		SpellDmgEffect->Modifiers[0].ModifierType = EGameplayMod::Attribute;
-		SpellDmgEffect->Modifiers[0].ModifierOp = EGameplayModOp::Override;
-		SpellDmgEffect->Modifiers[0].Attribute.SetUProperty(SpellDamageProperty);				// Modifies target's "Damage" attribute (-health)
-		SpellDmgEffect->Modifiers[0].OwnedTags.AddTag(IGameplayTagsModule::RequestGameplayTag(FName(TEXT("SpellDmg.Buff"))));
-		SpellDmgEffect->Duration.Value = UGameplayEffect::INSTANT_APPLICATION;
-
-		SourceComponent->ApplyGameplayEffectToTarget(SpellDmgEffect, SourceComponent);
-	}
-
-	{
-		GameplayTest_TickWorld(World, 0.0001f); // Move our Effects from the pending stack to the active stack inside the timer manager, this starts the clock for effect execution
-		GameplayTest_TickWorld(World, 0.1f); // Offset the current time from the start of the period so that floating point issues don't affect the test
-	}
-
-	{
-		// Tick once
-		GameplayTest_TickWorld(World, 1.f);
-		
-		float ActualValue = DestComponent->GetSet<UAbilitySystemTestAttributeSet>()->Health;
-		float ExpectedValue = StartHealth - (SpellDamage * GameplayEffectScaling) * 2.f;	// we trigger twice on this tick
-		Test->TestTrue(SKILL_TEST_TEXT("Damage Applied. Health: Actual: %.2f == Exected: %.2f", ActualValue, ExpectedValue), ActualValue == ExpectedValue);
-	}	
-
-	ClearGlobalCurveTable();
-
-	World->GetTimerManager().ClearAllTimersForObject(DestComponent);
-
-	World->EditorDestroyActor(SourceActor, false);
-	World->EditorDestroyActor(DestActor, false);
-
-	return true;
-}
-
-bool GameplayEffectsTest_DotDamage_ScalingProperty_Dynamic(UWorld *World, FAutomationTestBase * Test)
-{
-	// Add a dot that is powered by SpellDamage. Increase SpellDamage after applying, confirm it doesn't add extra damage to subsequent ticks.
-
-	const float StartHealth = 100.f;
-	const float SpellDamage = 10.f;
-	const float SpellDamage2 = 50.f;
-	const float GameplayEffectScaling = 1.f;
-
-	SetGlobalCurveTable();
-
-	AAbilitySystemTestPawn *SourceActor = World->SpawnActor<AAbilitySystemTestPawn>();
-	AAbilitySystemTestPawn *DestActor = World->SpawnActor<AAbilitySystemTestPawn>();
-
-	UProperty *DamageProperty = FindFieldChecked<UProperty>(UAbilitySystemTestAttributeSet::StaticClass(), GET_MEMBER_NAME_CHECKED(UAbilitySystemTestAttributeSet, Damage));
-	UProperty *SpellDamageProperty = FindFieldChecked<UProperty>(UAbilitySystemTestAttributeSet::StaticClass(), GET_MEMBER_NAME_CHECKED(UAbilitySystemTestAttributeSet, SpellDamage));
-
-	UAbilitySystemComponent * SourceComponent = SourceActor->GetAbilitySystemComponent();
-	UAbilitySystemComponent * DestComponent = DestActor->GetAbilitySystemComponent();
-	SourceComponent->GetSet<UAbilitySystemTestAttributeSet>()->Health = StartHealth;
-	DestComponent->GetSet<UAbilitySystemTestAttributeSet>()->Health = StartHealth;
-
-	SourceComponent->GetSet<UAbilitySystemTestAttributeSet>()->SpellDamage = SpellDamage;
-
-	{
-		ABILITY_LOG_SCOPE(TEXT("Apply InstantDamage"));
-
-		// This effects do Damage = 1.f * LinearCurve[LevelOfGameplayEffect].
-		// This translate into 1.f * PhysicalDamage.
-
-		UGameplayEffect * BaseDmgEffect = Cast<UGameplayEffect>(StaticConstructObject(UGameplayEffect::StaticClass(), GetTransientPackage(), FName(TEXT("BaseDmgEffect"))));
-		BaseDmgEffect->Modifiers.SetNum(1);
-		BaseDmgEffect->Modifiers[0].Magnitude.SetScalingValue(GameplayEffectScaling, FName(TEXT("LinearCurve")), NULL); // do "1*StandardDamage[Level]"
-		BaseDmgEffect->Modifiers[0].ModifierType = EGameplayMod::Attribute;
-		BaseDmgEffect->Modifiers[0].ModifierOp = EGameplayModOp::Additive;
-		BaseDmgEffect->Modifiers[0].Attribute.SetUProperty(DamageProperty);				// Modifies target's "Damage" attribute (-health)
-		BaseDmgEffect->Modifiers[0].OwnedTags.AddTag(IGameplayTagsModule::RequestGameplayTag(FName(TEXT("Damage.Basic"))));
-		BaseDmgEffect->Duration.Value = UGameplayEffect::INFINITE_DURATION;
-		BaseDmgEffect->Period.Value = 1.f;
-
-		BaseDmgEffect->LevelInfo.Attribute.SetUProperty(SpellDamageProperty);			// Powered by instigators SpellDamage
-		BaseDmgEffect->LevelInfo.InheritLevelFromOwner = false;
-		BaseDmgEffect->LevelInfo.TakeSnapshotOnInit = false;							// But level is dynamic, if SpellDamage changers after we apply, we update.
-
-		SourceComponent->ApplyGameplayEffectToTarget(BaseDmgEffect, DestComponent);
-		GameplayTest_TickWorld(World, 0.0001f); // Move our Effects from the pending stack to the active stack inside the timer manager, this starts the clock for effect execution
-
-		float SpellDamageTest = SourceComponent->GetSet<UAbilitySystemTestAttributeSet>()->SpellDamage;
-		check(SpellDamageTest == SpellDamage);
-
-		GameplayTest_TickWorld(World, 0.1f); // Offset the current time from the start of the period so that floating point issues don't affect the test
-	}
-
-	{
-		// Tick once
-		GameplayTest_TickWorld(World, 1.f);
-
-		float ActualValue = DestComponent->GetSet<UAbilitySystemTestAttributeSet>()->Health;
-		float ExpectedValue = StartHealth - (SpellDamage * GameplayEffectScaling * 2);	// We've ticked twice
-		Test->TestTrue(SKILL_TEST_TEXT("Damage Applied. Health: Actual: %.2f == Exected: %.2f", ActualValue, ExpectedValue), ActualValue == ExpectedValue);
-	}
-
-	{
-		// Increase spell damage on instigator (after we already applied the DOT)
-		UGameplayEffect * SpellDmgEffect = Cast<UGameplayEffect>(StaticConstructObject(UGameplayEffect::StaticClass(), GetTransientPackage(), FName(TEXT("SpellDmgEffect"))));
-		SpellDmgEffect->Modifiers.SetNum(1);
-		SpellDmgEffect->Modifiers[0].Magnitude.SetValue(SpellDamage2);
-		SpellDmgEffect->Modifiers[0].ModifierType = EGameplayMod::Attribute;
-		SpellDmgEffect->Modifiers[0].ModifierOp = EGameplayModOp::Override;
-		SpellDmgEffect->Modifiers[0].Attribute.SetUProperty(SpellDamageProperty);				// Modifies target's "Damage" attribute (-health)
-		SpellDmgEffect->Modifiers[0].OwnedTags.AddTag(IGameplayTagsModule::RequestGameplayTag(FName(TEXT("SpellDmg.Buff"))));
-		SpellDmgEffect->Duration.Value = UGameplayEffect::INSTANT_APPLICATION;
-
-		SourceComponent->ApplyGameplayEffectToTarget(SpellDmgEffect, SourceComponent);
-		GameplayTest_TickWorld(World, 0.0001f); // Move our Effects from the pending stack to the active stack inside the timer manager, this starts the clock for effect execution
-
-		float ActualValue = SourceComponent->GetSet<UAbilitySystemTestAttributeSet>()->SpellDamage;
-		float ExpectedValue = SpellDamage2;
-		Test->TestTrue(SKILL_TEST_TEXT("Spell Damage Mod: Actual: %.2f == Exected: %.2f", ActualValue, ExpectedValue), ActualValue == ExpectedValue);
-
-		GameplayTest_TickWorld(World, 0.1f); // Offset the current time from the start of the period so that floating point issues don't affect the test
-	}
-
-	{
-		// Tick once
-		GameplayTest_TickWorld(World, 1.f);
-
-		float ActualValue = DestComponent->GetSet<UAbilitySystemTestAttributeSet>()->Health;
-		float ExpectedValue = StartHealth - (SpellDamage * GameplayEffectScaling * 2) - (SpellDamage2 * GameplayEffectScaling);	// two unbuffed ticks, one buffed tick
-		Test->TestTrue(SKILL_TEST_TEXT("Damage Applied. Health: Actual: %.2f == Exected: %.2f", ActualValue, ExpectedValue), ActualValue == ExpectedValue);
-	}
-
-	ClearGlobalCurveTable();
-
-	World->EditorDestroyActor(SourceActor, false);
-	World->EditorDestroyActor(DestActor, false);
-
-	return true;
-}
-
-bool GameplayEffectsTest_MetaAttributes(UWorld *World, FAutomationTestBase * Test)
-{
-	// Sets up a GameplayEffect to give the source a constant +Health powered by the source's strength
-
-	const float StartHealth = 100.f;
-	const float MaxHealthPerStrength = 3.f;
-	const float StrengthValue = 10.f;
-
-	SetGlobalCurveTable();
-
-	AAbilitySystemTestPawn *SourceActor = World->SpawnActor<AAbilitySystemTestPawn>();
-	
-
-	UProperty *MaxHealthProperty = FindFieldChecked<UProperty>(UAbilitySystemTestAttributeSet::StaticClass(), GET_MEMBER_NAME_CHECKED(UAbilitySystemTestAttributeSet, MaxHealth));
-	UProperty *StrengthProperty = FindFieldChecked<UProperty>(UAbilitySystemTestAttributeSet::StaticClass(), GET_MEMBER_NAME_CHECKED(UAbilitySystemTestAttributeSet, Strength));
-
-	UAbilitySystemComponent * SourceComponent = SourceActor->GetAbilitySystemComponent();
-	SourceComponent->GetSet<UAbilitySystemTestAttributeSet>()->MaxHealth = StartHealth;
-	SourceComponent->GetSet<UAbilitySystemTestAttributeSet>()->Strength = 0.f;
-
-	{
-		ABILITY_LOG_SCOPE(TEXT("Setup meta stat"));
-
-		// This effects do Damage = 1.f * LinearCurve[LevelOfGameplayEffect].
-		// This translate into 1.f * PhysicalDamage.
-
-		UGameplayEffect * StrengthMaxHealhEffect = Cast<UGameplayEffect>(StaticConstructObject(UGameplayEffect::StaticClass(), GetTransientPackage(), FName(TEXT("StrengthMaxHealhEffect"))));
-		StrengthMaxHealhEffect->Modifiers.SetNum(1);
-		StrengthMaxHealhEffect->Modifiers[0].Magnitude.SetScalingValue(MaxHealthPerStrength, FName(TEXT("LinearCurve")), NULL); // do "1*StandardDamage[Level]"
-		StrengthMaxHealhEffect->Modifiers[0].ModifierType = EGameplayMod::Attribute;
-		StrengthMaxHealhEffect->Modifiers[0].ModifierOp = EGameplayModOp::Additive;
-		StrengthMaxHealhEffect->Modifiers[0].Attribute.SetUProperty(MaxHealthProperty);				// Modifies target's "Damage" attribute (-health)
-		StrengthMaxHealhEffect->Modifiers[0].OwnedTags.AddTag(IGameplayTagsModule::RequestGameplayTag(FName(TEXT("Damage.Basic"))));
-		StrengthMaxHealhEffect->Duration.Value = UGameplayEffect::INFINITE_DURATION;
-		StrengthMaxHealhEffect->Period.Value = UGameplayEffect::NO_PERIOD;
-
-		StrengthMaxHealhEffect->LevelInfo.Attribute.SetUProperty(StrengthProperty);			// Powered by instigators SpellDamage
-		StrengthMaxHealhEffect->LevelInfo.InheritLevelFromOwner = false;
-		StrengthMaxHealhEffect->LevelInfo.TakeSnapshotOnInit = false;						// But level is dynamic, if SpellDamage changers after we apply, we update.
-
-		SourceComponent->ApplyGameplayEffectToTarget(StrengthMaxHealhEffect, SourceComponent);
-
-		// Strength starts at 0, so confirm it did nothing yet.
-		float ActualValue = SourceComponent->GetSet<UAbilitySystemTestAttributeSet>()->MaxHealth;
-		float ExpectedValue = StartHealth;
-		Test->TestTrue(SKILL_TEST_TEXT("Damage Applied. Health: Actual: %.2f == Exected: %.2f", ActualValue, ExpectedValue), ActualValue == ExpectedValue);
-	}
-
-	{
-		// Set strength to 10. Confirm this adds 30 to MaxHeatlh.
-		UGameplayEffect * StrEffect = Cast<UGameplayEffect>(StaticConstructObject(UGameplayEffect::StaticClass(), GetTransientPackage(), FName(TEXT("StrEffect"))));
-		StrEffect->Modifiers.SetNum(1);
-		StrEffect->Modifiers[0].Magnitude.SetValue(StrengthValue);
-		StrEffect->Modifiers[0].ModifierType = EGameplayMod::Attribute;
-		StrEffect->Modifiers[0].ModifierOp = EGameplayModOp::Additive;
-		StrEffect->Modifiers[0].Attribute.SetUProperty(StrengthProperty);				// Modifies target's "Damage" attribute (-health)
-		StrEffect->Modifiers[0].OwnedTags.AddTag(IGameplayTagsModule::RequestGameplayTag(FName(TEXT("SpellDmg.Buff"))));
-		StrEffect->Duration.Value = UGameplayEffect::INSTANT_APPLICATION;
-
-		SourceComponent->ApplyGameplayEffectToTarget(StrEffect, SourceComponent);
-
-		{
-			float ActualValue = SourceComponent->GetSet<UAbilitySystemTestAttributeSet>()->Strength;
-			float ExpectedValue = StrengthValue;
-			Test->TestTrue(SKILL_TEST_TEXT("Strength: Actual: %.2f == Exected: %.2f", ActualValue, ExpectedValue), ActualValue == ExpectedValue);
-		}
-
-		{
-			float ActualValue = SourceComponent->GetSet<UAbilitySystemTestAttributeSet>()->MaxHealth;
-			float ExpectedValue = StartHealth + (MaxHealthPerStrength * StrengthValue);
-			Test->TestTrue(SKILL_TEST_TEXT("MaxHealth: Actual: %.2f == Exected: %.2f", ActualValue, ExpectedValue), ActualValue == ExpectedValue);
-		}
-	}
-
-	ClearGlobalCurveTable();
-
-	World->EditorDestroyActor(SourceActor, false);
-
-	return true;
-}
-
-bool GameplayEffectsTest_TagOrdering(UWorld *World, FAutomationTestBase * Test)
-{
-	const float StartHealth = 100.f;
-	const float DamageValue = 5.f;
-	const float BonusDamageMultiplier = 2.f;
-
-	AAbilitySystemTestPawn *SourceActor = World->SpawnActor<AAbilitySystemTestPawn>();
-	AAbilitySystemTestPawn *DestActor = World->SpawnActor<AAbilitySystemTestPawn>();
-
-	UProperty *HealthProperty = FindFieldChecked<UProperty>(UAbilitySystemTestAttributeSet::StaticClass(), GET_MEMBER_NAME_CHECKED(UAbilitySystemTestAttributeSet, Health));
-	UProperty *DamageProperty = FindFieldChecked<UProperty>(UAbilitySystemTestAttributeSet::StaticClass(), GET_MEMBER_NAME_CHECKED(UAbilitySystemTestAttributeSet, Damage));
-
-	UAbilitySystemComponent * SourceComponent = SourceActor->GetAbilitySystemComponent();
-	UAbilitySystemComponent * DestComponent = DestActor->GetAbilitySystemComponent();
-	SourceComponent->GetSet<UAbilitySystemTestAttributeSet>()->Health = StartHealth;
-	DestComponent->GetSet<UAbilitySystemTestAttributeSet>()->Health = StartHealth;
-
-	// Setup a GE to modify OutgoingGEs
-	{
-		ABILITY_LOG_SCOPE(TEXT("FireDamageBuff"));
-
-		UGameplayEffect * BaseDmgEffect = Cast<UGameplayEffect>(StaticConstructObject(UGameplayEffect::StaticClass(), GetTransientPackage(), FName(TEXT("FireDamageBuff"))));
-		BaseDmgEffect->Modifiers.SetNum(1);
-		BaseDmgEffect->Modifiers[0].Magnitude.SetValue(BonusDamageMultiplier);
-		BaseDmgEffect->Modifiers[0].ModifierType = EGameplayMod::OutgoingGE;
-		BaseDmgEffect->Modifiers[0].ModifierOp = EGameplayModOp::Multiplicitive;
-		BaseDmgEffect->Modifiers[0].Attribute.SetUProperty(DamageProperty);
-		BaseDmgEffect->Modifiers[0].OwnedTags.AddTag(IGameplayTagsModule::RequestGameplayTag(FName(TEXT("Damage.Buffed.FireBuff"))));
-		BaseDmgEffect->Modifiers[0].RequiredTags.AddTag(IGameplayTagsModule::RequestGameplayTag(FName(TEXT("Damage.Fire"))));
-		BaseDmgEffect->Duration.SetValue(UGameplayEffect::INFINITE_DURATION);
-		
-		BaseDmgEffect->GameplayEffectTags.AddTag(IGameplayTagsModule::RequestGameplayTag(FName(TEXT("Buff"))));
-		BaseDmgEffect->GameplayEffectRequiredTags.AddTag(IGameplayTagsModule::RequestGameplayTag(FName(TEXT("Damage"))));
-
-		// Apply to self
-		SourceComponent->ApplyGameplayEffectToTarget(BaseDmgEffect, SourceComponent, 1.f);
-	}
-
-	{
-		ABILITY_LOG_SCOPE(TEXT("MakeFireDamage"));
-
-		UGameplayEffect * BaseDmgEffect = Cast<UGameplayEffect>(StaticConstructObject(UGameplayEffect::StaticClass(), GetTransientPackage(), FName(TEXT("MakeFireDamage"))));
-		BaseDmgEffect->Modifiers.SetNum(1);
-		BaseDmgEffect->Modifiers[0].Magnitude.SetValue(0.f);
-		BaseDmgEffect->Modifiers[0].ModifierType = EGameplayMod::OutgoingGE;
-		BaseDmgEffect->Modifiers[0].ModifierOp = EGameplayModOp::Additive;
-		BaseDmgEffect->Modifiers[0].Attribute.SetUProperty(DamageProperty);
-		BaseDmgEffect->Modifiers[0].OwnedTags.AddTag(IGameplayTagsModule::RequestGameplayTag(FName(TEXT("Damage.Fire"))));
-		BaseDmgEffect->Modifiers[0].RequiredTags.AddTag(IGameplayTagsModule::RequestGameplayTag(FName(TEXT("Damage.Physical"))));
-		BaseDmgEffect->Duration.SetValue(UGameplayEffect::INFINITE_DURATION);
-		
-		BaseDmgEffect->GameplayEffectTags.AddTag(IGameplayTagsModule::RequestGameplayTag(FName(TEXT("Buff"))));
-		BaseDmgEffect->GameplayEffectRequiredTags.AddTag(IGameplayTagsModule::RequestGameplayTag(FName(TEXT("Damage"))));
-
-		// Apply to self
-		SourceComponent->ApplyGameplayEffectToTarget(BaseDmgEffect, SourceComponent, 1.f);
-	}
-
-	// Apply Damage
-	{
-		ABILITY_LOG_SCOPE(TEXT("Apply InstantDamage"));
-
-		UGameplayEffect * BaseDmgEffect = Cast<UGameplayEffect>(StaticConstructObject(UGameplayEffect::StaticClass(), GetTransientPackage(), FName(TEXT("Damage"))));
-		BaseDmgEffect->Modifiers.SetNum(1);
-		BaseDmgEffect->Modifiers[0].Magnitude.SetValue(DamageValue);
-		BaseDmgEffect->Modifiers[0].ModifierType = EGameplayMod::Attribute;
-		BaseDmgEffect->Modifiers[0].ModifierOp = EGameplayModOp::Additive;
-		BaseDmgEffect->Modifiers[0].Attribute.SetUProperty(DamageProperty);
-		BaseDmgEffect->Modifiers[0].OwnedTags.AddTag(IGameplayTagsModule::RequestGameplayTag(FName(TEXT("Damage.Physical"))));
-		BaseDmgEffect->Duration.Value = UGameplayEffect::INSTANT_APPLICATION;
-		
-		BaseDmgEffect->GameplayEffectTags.AddTag(IGameplayTagsModule::RequestGameplayTag(FName(TEXT("Damage"))));
-
-		// Apply to target
-		SourceComponent->ApplyGameplayEffectToTarget(BaseDmgEffect, DestComponent, 1.f);
-
-		float ExpectedValue = (StartHealth - (DamageValue * BonusDamageMultiplier));
-		float ActualValue = DestComponent->GetSet<UAbilitySystemTestAttributeSet>()->Health;
-		
-		Test->TestTrue(SKILL_TEST_TEXT("MaxHealth: Actual: %.2f == Exected: %.2f", ActualValue, ExpectedValue), ActualValue == ExpectedValue);
-	}
-
-	World->EditorDestroyActor(SourceActor, false);
-	World->EditorDestroyActor(DestActor, false);
-
-	return true;
-}
-
-// test the rule that stacks based on the highest gameplay effect
-bool GameplayEffectsTest_StackingHighest(UWorld *World, FAutomationTestBase * Test)
-{
-	const float StackingValue = 5.f;
-
-	SetGlobalCurveTable();
-	SetGlobalDataTable();
-
-	AAbilitySystemTestPawn *SourceActor = World->SpawnActor<AAbilitySystemTestPawn>();
-	AAbilitySystemTestPawn *DestActor = World->SpawnActor<AAbilitySystemTestPawn>();
-
-	UProperty *StackingProperty = FindFieldChecked<UProperty>(UAbilitySystemTestAttributeSet::StaticClass(), GET_MEMBER_NAME_CHECKED(UAbilitySystemTestAttributeSet, StackingAttribute1));
-
-	UAbilitySystemComponent * SourceComponent = SourceActor->GetAbilitySystemComponent();
-	UAbilitySystemComponent * DestComponent = DestActor->GetAbilitySystemComponent();
-
-	{
-		ABILITY_LOG_SCOPE(TEXT("Apply StackingEffect"));
-
-		UGameplayEffect * BaseStackedEffect = Cast<UGameplayEffect>(StaticConstructObject(UGameplayEffect::StaticClass(), GetTransientPackage(), FName(TEXT("StackingEffect1"))));
-		BaseStackedEffect->Modifiers.SetNum(1);
-		BaseStackedEffect->Modifiers[0].Magnitude.SetValue(StackingValue);
-		BaseStackedEffect->Modifiers[0].ModifierType = EGameplayMod::Attribute;
-		BaseStackedEffect->Modifiers[0].ModifierOp = EGameplayModOp::Additive;
-		BaseStackedEffect->Modifiers[0].Attribute.SetUProperty(StackingProperty);
-		BaseStackedEffect->Duration.Value = UGameplayEffect::INFINITE_DURATION;
-		BaseStackedEffect->Period.Value = 1.f;
-		BaseStackedEffect->StackingPolicy = EGameplayEffectStackingPolicy::Highest;
-		BaseStackedEffect->StackedAttribName = FName(*StackingProperty->GetName());
-
-		BaseStackedEffect->ValidateGameplayEffect();
-
-		SourceComponent->ApplyGameplayEffectToTarget(BaseStackedEffect, DestComponent);
-
-		GameplayTest_TickWorld(World, SMALL_NUMBER); // Move our Effects from the pending stack to the active stack inside the timer manager, this starts the clock for effect execution
-		{
-			float ActualValue = DestComponent->GetSet<UAbilitySystemTestAttributeSet>()->StackingAttribute1;
-			float ExpectedValue = StackingValue;
-			Test->TestTrue(SKILL_TEST_TEXT("Stacking Attribute: Actual: %.2f == Exected: %.2f", ActualValue, ExpectedValue), ActualValue == ExpectedValue);
-		}
-		GameplayTest_TickWorld(World, 0.1f); // Offset the current time from the start of the period so that floating point issues don't affect the test
-	}
-
-	{
-		// Tick once
-		GameplayTest_TickWorld(World, 1.f);
-
-		float ActualValue = DestComponent->GetSet<UAbilitySystemTestAttributeSet>()->StackingAttribute1;
-		float ExpectedValue = StackingValue * 2.f;	// effect will execute twice
-		Test->TestTrue(SKILL_TEST_TEXT("Stacking Attribute: Actual: %.2f == Exected: %.2f", ActualValue, ExpectedValue), ActualValue == ExpectedValue);
-	}
-
-	{
-		ABILITY_LOG_SCOPE(TEXT("Apply StackingEffect"));
-
-		UGameplayEffect * BaseStackedEffect = Cast<UGameplayEffect>(StaticConstructObject(UGameplayEffect::StaticClass(), GetTransientPackage(), FName(TEXT("StackingEffect2"))));
-		BaseStackedEffect->Modifiers.SetNum(1);
-		BaseStackedEffect->Modifiers[0].Magnitude.SetValue(StackingValue * 2.f);
-		BaseStackedEffect->Modifiers[0].ModifierType = EGameplayMod::Attribute;
-		BaseStackedEffect->Modifiers[0].ModifierOp = EGameplayModOp::Additive;
-		BaseStackedEffect->Modifiers[0].Attribute.SetUProperty(StackingProperty);
-		BaseStackedEffect->Duration.Value = UGameplayEffect::INFINITE_DURATION;
-		BaseStackedEffect->Period.Value = 1.f;
-		BaseStackedEffect->StackingPolicy = EGameplayEffectStackingPolicy::Highest;
-		BaseStackedEffect->StackedAttribName = FName(*StackingProperty->GetName());
-
-		BaseStackedEffect->ValidateGameplayEffect();
-
-		SourceComponent->ApplyGameplayEffectToTarget(BaseStackedEffect, DestComponent);
-
-		GameplayTest_TickWorld(World, 0.0001f); // Move our Effects from the pending stack to the active stack inside the timer manager, this starts the clock for effect execution
-		GameplayTest_TickWorld(World, 0.1f); // Offset the current time from the start of the period so that floating point issues don't affect the test
-	}
-
-	{
-		// Tick once
-		GameplayTest_TickWorld(World, 1.f);
-
-		float ActualValue = DestComponent->GetSet<UAbilitySystemTestAttributeSet>()->StackingAttribute1;
-		float ExpectedValue = StackingValue * 4.f;	// 2 for first effect, 2 for the last tick with both effects
-		Test->TestTrue(SKILL_TEST_TEXT("Stacking Attribute: Actual: %.2f == Exected: %.2f", ActualValue, ExpectedValue), ActualValue == ExpectedValue);
-	}
-
-	ClearGlobalCurveTable();
-	ClearGlobalDataTable();
-
-	World->EditorDestroyActor(SourceActor, false);
-	World->EditorDestroyActor(DestActor, false);
-
-	return true;
-}
-
-// test the rule that stacks based on the smallest gameplay effect
-bool GameplayEffectsTest_StackingLowest(UWorld *World, FAutomationTestBase * Test)
-{
-	const float StackingValue = 5.f;
-
-	SetGlobalCurveTable();
-	SetGlobalDataTable();
-
-	AAbilitySystemTestPawn *SourceActor = World->SpawnActor<AAbilitySystemTestPawn>();
-	AAbilitySystemTestPawn *DestActor = World->SpawnActor<AAbilitySystemTestPawn>();
-
-	UProperty *StackingProperty = FindFieldChecked<UProperty>(UAbilitySystemTestAttributeSet::StaticClass(), GET_MEMBER_NAME_CHECKED(UAbilitySystemTestAttributeSet, StackingAttribute1));
-
-	UAbilitySystemComponent * SourceComponent = SourceActor->GetAbilitySystemComponent();
-	UAbilitySystemComponent * DestComponent = DestActor->GetAbilitySystemComponent();
-
-	{
-		ABILITY_LOG_SCOPE(TEXT("Apply StackingEffect"));
-
-		UGameplayEffect * BaseStackedEffect = Cast<UGameplayEffect>(StaticConstructObject(UGameplayEffect::StaticClass(), GetTransientPackage(), FName(TEXT("StackingEffect1"))));
-		BaseStackedEffect->Modifiers.SetNum(1);
-		BaseStackedEffect->Modifiers[0].Magnitude.SetValue(StackingValue);
-		BaseStackedEffect->Modifiers[0].ModifierType = EGameplayMod::Attribute;
-		BaseStackedEffect->Modifiers[0].ModifierOp = EGameplayModOp::Additive;
-		BaseStackedEffect->Modifiers[0].Attribute.SetUProperty(StackingProperty);
-		BaseStackedEffect->Duration.Value = UGameplayEffect::INFINITE_DURATION;
-		BaseStackedEffect->Period.Value = 1.f;
-		BaseStackedEffect->StackingPolicy = EGameplayEffectStackingPolicy::Lowest;
-		BaseStackedEffect->StackedAttribName = FName(*StackingProperty->GetName());
-
-		BaseStackedEffect->ValidateGameplayEffect();
-
-		SourceComponent->ApplyGameplayEffectToTarget(BaseStackedEffect, DestComponent);
-
-		GameplayTest_TickWorld(World, 0.0001f); // Move our Effects from the pending stack to the active stack inside the timer manager, this starts the clock for effect execution
-		GameplayTest_TickWorld(World, 0.1f); // Offset the current time from the start of the period so that floating point issues don't affect the test
-	}
-
-	{
-		// Tick once
-		GameplayTest_TickWorld(World, 1.f);
-
-		float ActualValue = DestComponent->GetSet<UAbilitySystemTestAttributeSet>()->StackingAttribute1;
-		float ExpectedValue = StackingValue * 2.f;	// the effect should execute twice here
-		Test->TestTrue(SKILL_TEST_TEXT("Stacking Attribute: Actual: %.2f == Exected: %.2f", ActualValue, ExpectedValue), ActualValue == ExpectedValue);
-	}
-
-	{
-		ABILITY_LOG_SCOPE(TEXT("Apply StackingEffect"));
-
-		UGameplayEffect * BaseStackedEffect = Cast<UGameplayEffect>(StaticConstructObject(UGameplayEffect::StaticClass(), GetTransientPackage(), FName(TEXT("StackingEffect2"))));
-		BaseStackedEffect->Modifiers.SetNum(1);
-		BaseStackedEffect->Modifiers[0].Magnitude.SetValue(StackingValue * 2.f);
-		BaseStackedEffect->Modifiers[0].ModifierType = EGameplayMod::Attribute;
-		BaseStackedEffect->Modifiers[0].ModifierOp = EGameplayModOp::Additive;
-		BaseStackedEffect->Modifiers[0].Attribute.SetUProperty(StackingProperty);
-		BaseStackedEffect->Duration.Value = UGameplayEffect::INFINITE_DURATION;
-		BaseStackedEffect->Period.Value = 1.f;
-		BaseStackedEffect->StackingPolicy = EGameplayEffectStackingPolicy::Lowest;
-		BaseStackedEffect->StackedAttribName = FName(*StackingProperty->GetName());
-
-		BaseStackedEffect->ValidateGameplayEffect();
-
-		SourceComponent->ApplyGameplayEffectToTarget(BaseStackedEffect, DestComponent);
-
-		GameplayTest_TickWorld(World, 0.0001f); // Move our Effects from the pending stack to the active stack inside the timer manager, this starts the clock for effect execution
-		GameplayTest_TickWorld(World, 0.1f); // Offset the current time from the start of the period so that floating point issues don't affect the test
-	}
-
-	{
-		// Tick once
-		GameplayTest_TickWorld(World, 1.f);
-
-		float ActualValue = DestComponent->GetSet<UAbilitySystemTestAttributeSet>()->StackingAttribute1;
-		float ExpectedValue = StackingValue * 3.f;	// the first effect has executed 3 times, the second hasn't executed
-		Test->TestTrue(SKILL_TEST_TEXT("Stacking Attribute: Actual: %.2f == Exected: %.2f", ActualValue, ExpectedValue), ActualValue == ExpectedValue);
-	}
-
-	ClearGlobalCurveTable();
-	ClearGlobalDataTable();
-
-	World->EditorDestroyActor(SourceActor, false);
-	World->EditorDestroyActor(DestActor, false);
-
-	return true;
-}
-
-// test the rule that ignores stacking
-bool GameplayEffectsTest_StackingUnlimited(UWorld *World, FAutomationTestBase * Test)
-{
-	const float StackingValue = 5.f;
-
-	SetGlobalCurveTable();
-	SetGlobalDataTable();
-
-	AAbilitySystemTestPawn *SourceActor = World->SpawnActor<AAbilitySystemTestPawn>();
-	AAbilitySystemTestPawn *DestActor = World->SpawnActor<AAbilitySystemTestPawn>();
-
-	UProperty *StackingProperty = FindFieldChecked<UProperty>(UAbilitySystemTestAttributeSet::StaticClass(), GET_MEMBER_NAME_CHECKED(UAbilitySystemTestAttributeSet, StackingAttribute1));
-
-	UAbilitySystemComponent * SourceComponent = SourceActor->GetAbilitySystemComponent();
-	UAbilitySystemComponent * DestComponent = DestActor->GetAbilitySystemComponent();
-
-	{
-		ABILITY_LOG_SCOPE(TEXT("Apply StackingEffect"));
-
-		UGameplayEffect * BaseStackedEffect = Cast<UGameplayEffect>(StaticConstructObject(UGameplayEffect::StaticClass(), GetTransientPackage(), FName(TEXT("StackingEffect1"))));
-		BaseStackedEffect->Modifiers.SetNum(1);
-		BaseStackedEffect->Modifiers[0].Magnitude.SetValue(StackingValue);
-		BaseStackedEffect->Modifiers[0].ModifierType = EGameplayMod::Attribute;
-		BaseStackedEffect->Modifiers[0].ModifierOp = EGameplayModOp::Additive;
-		BaseStackedEffect->Modifiers[0].Attribute.SetUProperty(StackingProperty);
-		BaseStackedEffect->Duration.Value = UGameplayEffect::INFINITE_DURATION;
-		BaseStackedEffect->Period.Value = 1.f;
-		BaseStackedEffect->StackingPolicy = EGameplayEffectStackingPolicy::Unlimited;
-		BaseStackedEffect->StackedAttribName = FName(*StackingProperty->GetName());
-
-		BaseStackedEffect->ValidateGameplayEffect();
-
-		SourceComponent->ApplyGameplayEffectToTarget(BaseStackedEffect, DestComponent);
-	}
-
-	{
-		ABILITY_LOG_SCOPE(TEXT("Apply StackingEffect"));
-
-		UGameplayEffect * BaseStackedEffect = Cast<UGameplayEffect>(StaticConstructObject(UGameplayEffect::StaticClass(), GetTransientPackage(), FName(TEXT("StackingEffect2"))));
-		BaseStackedEffect->Modifiers.SetNum(1);
-		BaseStackedEffect->Modifiers[0].Magnitude.SetValue(StackingValue * 2.f);
-		BaseStackedEffect->Modifiers[0].ModifierType = EGameplayMod::Attribute;
-		BaseStackedEffect->Modifiers[0].ModifierOp = EGameplayModOp::Additive;
-		BaseStackedEffect->Modifiers[0].Attribute.SetUProperty(StackingProperty);
-		BaseStackedEffect->Duration.Value = UGameplayEffect::INFINITE_DURATION;
-		BaseStackedEffect->Period.Value = 1.f;
-		BaseStackedEffect->StackingPolicy = EGameplayEffectStackingPolicy::Unlimited;
-		BaseStackedEffect->StackedAttribName = FName(*StackingProperty->GetName());
-
-		BaseStackedEffect->ValidateGameplayEffect();
-
-		SourceComponent->ApplyGameplayEffectToTarget(BaseStackedEffect, DestComponent);
-
-		GameplayTest_TickWorld(World, 0.0001f); // Move our Effects from the pending stack to the active stack inside the timer manager, this starts the clock for effect execution
-		GameplayTest_TickWorld(World, 0.1f); // Offset the current time from the start of the period so that floating point issues don't affect the test
-	}
-
-	{
-		// Tick once
-		GameplayTest_TickWorld(World, 1.f);
-
-		float ActualValue = DestComponent->GetSet<UAbilitySystemTestAttributeSet>()->StackingAttribute1;
-		float ExpectedValue = StackingValue * 6.f;	// 1 for first GE, 2 for the second GE, 3 for both GEs during the tick
-		Test->TestTrue(SKILL_TEST_TEXT("Stacking Attribute: Actual: %.2f == Exected: %.2f", ActualValue, ExpectedValue), ActualValue == ExpectedValue);
-	}
-
-	ClearGlobalCurveTable();
-	ClearGlobalDataTable();
-
-	World->EditorDestroyActor(SourceActor, false);
-	World->EditorDestroyActor(DestActor, false);
-
-	return true;
-}
-
-// test to make sure that stacking updates correctly when a gameplay effect is removed
-bool GameplayEffectsTest_StackingRemoval(UWorld *World, FAutomationTestBase * Test)
-{
-	const float StackingValue = 5.f;
-
-	SetGlobalCurveTable();
-	SetGlobalDataTable();
-
-	AAbilitySystemTestPawn *SourceActor = World->SpawnActor<AAbilitySystemTestPawn>();
-	AAbilitySystemTestPawn *DestActor = World->SpawnActor<AAbilitySystemTestPawn>();
-
-	UProperty *StackingProperty = FindFieldChecked<UProperty>(UAbilitySystemTestAttributeSet::StaticClass(), GET_MEMBER_NAME_CHECKED(UAbilitySystemTestAttributeSet, StackingAttribute1));
-
-	UAbilitySystemComponent * SourceComponent = SourceActor->GetAbilitySystemComponent();
-	UAbilitySystemComponent * DestComponent = DestActor->GetAbilitySystemComponent();
-
-	{
-		ABILITY_LOG_SCOPE(TEXT("Apply StackingEffect"));
-
-		UGameplayEffect * BaseStackedEffect = Cast<UGameplayEffect>(StaticConstructObject(UGameplayEffect::StaticClass(), GetTransientPackage(), FName(TEXT("StackingEffect1"))));
-		BaseStackedEffect->Modifiers.SetNum(1);
-		BaseStackedEffect->Modifiers[0].Magnitude.SetValue(StackingValue);
-		BaseStackedEffect->Modifiers[0].ModifierType = EGameplayMod::Attribute;
-		BaseStackedEffect->Modifiers[0].ModifierOp = EGameplayModOp::Additive;
-		BaseStackedEffect->Modifiers[0].Attribute.SetUProperty(StackingProperty);
-		BaseStackedEffect->Duration.Value = UGameplayEffect::INFINITE_DURATION;
-		BaseStackedEffect->Period.Value = 1.f;
-		BaseStackedEffect->StackingPolicy = EGameplayEffectStackingPolicy::Highest;
-		BaseStackedEffect->StackedAttribName = FName(*StackingProperty->GetName());
-
-		BaseStackedEffect->ValidateGameplayEffect();
-
-		SourceComponent->ApplyGameplayEffectToTarget(BaseStackedEffect, DestComponent);
-	}
-
-	{
-		ABILITY_LOG_SCOPE(TEXT("Apply StackingEffect"));
-
-		UGameplayEffect * BaseStackedEffect = Cast<UGameplayEffect>(StaticConstructObject(UGameplayEffect::StaticClass(), GetTransientPackage(), FName(TEXT("StackingEffect2"))));
-		BaseStackedEffect->Modifiers.SetNum(1);
-		BaseStackedEffect->Modifiers[0].Magnitude.SetValue(StackingValue * 2.f);
-		BaseStackedEffect->Modifiers[0].ModifierType = EGameplayMod::Attribute;
-		BaseStackedEffect->Modifiers[0].ModifierOp = EGameplayModOp::Additive;
-		BaseStackedEffect->Modifiers[0].Attribute.SetUProperty(StackingProperty);
-		BaseStackedEffect->Duration.Value = 1.f;// UGameplayEffect::INFINITE_DURATION;
-		BaseStackedEffect->Period.Value = 1.f;
-		BaseStackedEffect->StackingPolicy = EGameplayEffectStackingPolicy::Highest;
-		BaseStackedEffect->StackedAttribName = FName(*StackingProperty->GetName());
-
-		BaseStackedEffect->ValidateGameplayEffect();
-
-		SourceComponent->ApplyGameplayEffectToTarget(BaseStackedEffect, DestComponent);
-
-		GameplayTest_TickWorld(World, 0.0001f); // Move our Effects from the pending stack to the active stack inside the timer manager, this starts the clock for effect execution
-		GameplayTest_TickWorld(World, 0.1f); // Offset the current time from the start of the period so that floating point issues don't affect the test
-	}
-
-	{
-		// Tick once
-		GameplayTest_TickWorld(World, 1.f);
-
-		float ActualValue = DestComponent->GetSet<UAbilitySystemTestAttributeSet>()->StackingAttribute1;
-		float ExpectedValue = StackingValue * 4.f;	// 2 for second effect, 2 for second effect during tick
-		Test->TestTrue(SKILL_TEST_TEXT("Stacking Attribute: Actual: %.2f == Exected: %.2f", ActualValue, ExpectedValue), ActualValue == ExpectedValue);
-
-		// Tick once
-		GameplayTest_TickWorld(World, 1.f);
-
-		ActualValue = DestComponent->GetSet<UAbilitySystemTestAttributeSet>()->StackingAttribute1;
-		ExpectedValue = StackingValue * 5.f;	// 2 for second effect, 2 for second effect during first tick, 1 for first effect during second tick
-		Test->TestTrue(SKILL_TEST_TEXT("Stacking Attribute: Actual: %.2f == Exected: %.2f", ActualValue, ExpectedValue), ActualValue == ExpectedValue);
-	}
-
-	ClearGlobalCurveTable();
-	ClearGlobalDataTable();
-
-	World->EditorDestroyActor(SourceActor, false);
-	World->EditorDestroyActor(DestActor, false);
-
-	return true;
-}
-
-// test that attributes that aren't allowed to stack won't stack even if the stacking rule says they should
-bool GameplayEffectsTest_StackingNoStack(UWorld *World, FAutomationTestBase * Test)
-{
-	const float StackingValue = 5.f;
-
-	SetGlobalCurveTable();
-	SetGlobalDataTable();
-
-	AAbilitySystemTestPawn *SourceActor = World->SpawnActor<AAbilitySystemTestPawn>();
-	AAbilitySystemTestPawn *DestActor = World->SpawnActor<AAbilitySystemTestPawn>();
-
-	UProperty *NoStackProperty = FindFieldChecked<UProperty>(UAbilitySystemTestAttributeSet::StaticClass(), GET_MEMBER_NAME_CHECKED(UAbilitySystemTestAttributeSet, NoStackAttribute));
-
-	UAbilitySystemComponent * SourceComponent = SourceActor->GetAbilitySystemComponent();
-	UAbilitySystemComponent * DestComponent = DestActor->GetAbilitySystemComponent();
-
-	{
-		ABILITY_LOG_SCOPE(TEXT("Apply StackingEffect"));
-
-		UGameplayEffect * UnstackableEffect = Cast<UGameplayEffect>(StaticConstructObject(UGameplayEffect::StaticClass(), GetTransientPackage(), FName(TEXT("NoStackEffect1"))));
-		UnstackableEffect->Modifiers.SetNum(1);
-		UnstackableEffect->Modifiers[0].Magnitude.SetValue(StackingValue);
-		UnstackableEffect->Modifiers[0].ModifierType = EGameplayMod::Attribute;
-		UnstackableEffect->Modifiers[0].ModifierOp = EGameplayModOp::Additive;
-		UnstackableEffect->Modifiers[0].Attribute.SetUProperty(NoStackProperty);
-		UnstackableEffect->Duration.Value = UGameplayEffect::INFINITE_DURATION;
-		UnstackableEffect->Period.Value = 1.f;
-		UnstackableEffect->StackingPolicy = EGameplayEffectStackingPolicy::Highest;
-		UnstackableEffect->StackedAttribName = FName(*NoStackProperty->GetName());
-
-		UnstackableEffect->ValidateGameplayEffect();
-
-		SourceComponent->ApplyGameplayEffectToTarget(UnstackableEffect, DestComponent);
-	}
-
-	{
-		ABILITY_LOG_SCOPE(TEXT("Apply StackingEffect"));
-
-		UGameplayEffect * UnstackableEffect = Cast<UGameplayEffect>(StaticConstructObject(UGameplayEffect::StaticClass(), GetTransientPackage(), FName(TEXT("NoStackEffect2"))));
-		UnstackableEffect->Modifiers.SetNum(1);
-		UnstackableEffect->Modifiers[0].Magnitude.SetValue(StackingValue * 2.f);
-		UnstackableEffect->Modifiers[0].ModifierType = EGameplayMod::Attribute;
-		UnstackableEffect->Modifiers[0].ModifierOp = EGameplayModOp::Additive;
-		UnstackableEffect->Modifiers[0].Attribute.SetUProperty(NoStackProperty);
-		UnstackableEffect->Duration.Value = UGameplayEffect::INFINITE_DURATION;
-		UnstackableEffect->Period.Value = 1.f;
-		UnstackableEffect->StackingPolicy = EGameplayEffectStackingPolicy::Highest;
-		UnstackableEffect->StackedAttribName = FName(*NoStackProperty->GetName());
-
-		UnstackableEffect->ValidateGameplayEffect();
-
-		SourceComponent->ApplyGameplayEffectToTarget(UnstackableEffect, DestComponent);
-
-		GameplayTest_TickWorld(World, 0.0001f); // Move our Effects from the pending stack to the active stack inside the timer manager, this starts the clock for effect execution
-		GameplayTest_TickWorld(World, 0.1f); // Offset the current time from the start of the period so that floating point issues don't affect the test
-	}
-
-	{
-		// Tick once
-		GameplayTest_TickWorld(World, 1.f);
-
-		float ActualValue = DestComponent->GetSet<UAbilitySystemTestAttributeSet>()->NoStackAttribute;
-		float ExpectedValue = StackingValue * 6.f;	// 1 for first effect, 2 for second effect, 3 for both effects during tick
-		Test->TestTrue(SKILL_TEST_TEXT("Stacking Attribute: Actual: %.2f == Exected: %.2f", ActualValue, ExpectedValue), ActualValue == ExpectedValue);
-	}
-
-	ClearGlobalCurveTable();
-	ClearGlobalDataTable();
-
-	World->EditorDestroyActor(SourceActor, false);
-	World->EditorDestroyActor(DestActor, false);
-
-	return true;
-}
-
-// test the custom rule that stacks based on capping the number of effects applied
-bool GameplayEffectsTest_StackingCustomCapped(UWorld *World, FAutomationTestBase * Test)
-{
-	const float StackingValue = 5.f;
-
-	SetGlobalCurveTable();
-	SetGlobalDataTable();
-
-	AAbilitySystemTestPawn *SourceActor = World->SpawnActor<AAbilitySystemTestPawn>();
-	AAbilitySystemTestPawn *DestActor = World->SpawnActor<AAbilitySystemTestPawn>();
-
-	UProperty *StackingProperty = FindFieldChecked<UProperty>(UAbilitySystemTestAttributeSet::StaticClass(), GET_MEMBER_NAME_CHECKED(UAbilitySystemTestAttributeSet, StackingAttribute1));
-
-	UAbilitySystemComponent * SourceComponent = SourceActor->GetAbilitySystemComponent();
-	UAbilitySystemComponent * DestComponent = DestActor->GetAbilitySystemComponent();
-
-	{
-		ABILITY_LOG_SCOPE(TEXT("Apply StackingEffect"));
-
-		UGameplayEffect * BaseStackedEffect = Cast<UGameplayEffect>(StaticConstructObject(UGameplayEffect::StaticClass(), GetTransientPackage(), FName(TEXT("StackingEffect1"))));
-		BaseStackedEffect->Modifiers.SetNum(1);
-		BaseStackedEffect->Modifiers[0].Magnitude.SetValue(StackingValue);
-		BaseStackedEffect->Modifiers[0].ModifierType = EGameplayMod::Attribute;
-		BaseStackedEffect->Modifiers[0].ModifierOp = EGameplayModOp::Additive;
-		BaseStackedEffect->Modifiers[0].Attribute.SetUProperty(StackingProperty);
-		BaseStackedEffect->Modifiers[0].OwnedTags.AddTag(IGameplayTagsModule::RequestGameplayTag("Stackable"));
-		BaseStackedEffect->Duration.Value = 2.f;
-		BaseStackedEffect->Period.Value = 1.f;
-		BaseStackedEffect->StackingPolicy = EGameplayEffectStackingPolicy::Callback;
-		BaseStackedEffect->StackingExtension = UGameplayEffectStackingExtension_CappedNumberTest::StaticClass();
-		BaseStackedEffect->StackedAttribName = FName(*StackingProperty->GetName());
-
-		BaseStackedEffect->ValidateGameplayEffect();
-
-		SourceComponent->ApplyGameplayEffectToTarget(BaseStackedEffect, DestComponent);
-	}
-
-	{
-		ABILITY_LOG_SCOPE(TEXT("Apply StackingEffect"));
-
-		UGameplayEffect * BaseStackedEffect = Cast<UGameplayEffect>(StaticConstructObject(UGameplayEffect::StaticClass(), GetTransientPackage(), FName(TEXT("StackingEffect2"))));
-		BaseStackedEffect->Modifiers.SetNum(1);
-		BaseStackedEffect->Modifiers[0].Magnitude.SetValue(StackingValue);
-		BaseStackedEffect->Modifiers[0].ModifierType = EGameplayMod::Attribute;
-		BaseStackedEffect->Modifiers[0].ModifierOp = EGameplayModOp::Additive;
-		BaseStackedEffect->Modifiers[0].Attribute.SetUProperty(StackingProperty);
-		BaseStackedEffect->Modifiers[0].OwnedTags.AddTag(IGameplayTagsModule::RequestGameplayTag("Stackable"));
-		BaseStackedEffect->Duration.Value = 2.f;
-		BaseStackedEffect->Period.Value = 1.f;
-		BaseStackedEffect->StackingPolicy = EGameplayEffectStackingPolicy::Callback;
-		BaseStackedEffect->StackingExtension = UGameplayEffectStackingExtension_CappedNumberTest::StaticClass();
-		BaseStackedEffect->StackedAttribName = FName(*StackingProperty->GetName());
-
-		BaseStackedEffect->ValidateGameplayEffect();
-
-		SourceComponent->ApplyGameplayEffectToTarget(BaseStackedEffect, DestComponent);
-	}
-
-	{
-		ABILITY_LOG_SCOPE(TEXT("Apply StackingEffect"));
-
-		UGameplayEffect * BaseStackedEffect = Cast<UGameplayEffect>(StaticConstructObject(UGameplayEffect::StaticClass(), GetTransientPackage(), FName(TEXT("StackingEffect3"))));
-		BaseStackedEffect->Modifiers.SetNum(1);
-		BaseStackedEffect->Modifiers[0].Magnitude.SetValue(StackingValue);
-		BaseStackedEffect->Modifiers[0].ModifierType = EGameplayMod::Attribute;
-		BaseStackedEffect->Modifiers[0].ModifierOp = EGameplayModOp::Additive;
-		BaseStackedEffect->Modifiers[0].Attribute.SetUProperty(StackingProperty);
-		BaseStackedEffect->Modifiers[0].OwnedTags.AddTag(IGameplayTagsModule::RequestGameplayTag("Stackable"));
-		BaseStackedEffect->Duration.Value = 2.f;
-		BaseStackedEffect->Period.Value = 1.f;
-		BaseStackedEffect->StackingPolicy = EGameplayEffectStackingPolicy::Callback;
-		BaseStackedEffect->StackingExtension = UGameplayEffectStackingExtension_CappedNumberTest::StaticClass();
-		BaseStackedEffect->StackedAttribName = FName(*StackingProperty->GetName());
-
-		BaseStackedEffect->ValidateGameplayEffect();
-
-		SourceComponent->ApplyGameplayEffectToTarget(BaseStackedEffect, DestComponent);
-
-		GameplayTest_TickWorld(World, 0.0001f); // Move our Effects from the pending stack to the active stack inside the timer manager, this starts the clock for effect execution
-		GameplayTest_TickWorld(World, 0.1f); // offset the current time from the start of the period to avoid floating point issues causing the tests to fail
-	}
-
-	{
-		// Tick once
-		GameplayTest_TickWorld(World, 1.f);
-
-		float ActualValue = DestComponent->GetSet<UAbilitySystemTestAttributeSet>()->StackingAttribute1;
-		float ExpectedValue = StackingValue * 4.f;	// 2 for effects being applied at the start of the tick, 2 for effects being applied at the end of this tick
-		Test->TestTrue(SKILL_TEST_TEXT("Stacking Attribute: Actual: %.2f == Exected: %.2f", ActualValue, ExpectedValue), ActualValue == ExpectedValue);
-	}
-
-	{
-		ABILITY_LOG_SCOPE(TEXT("Apply StackingEffect"));
-
-		UGameplayEffect * BaseStackedEffect = Cast<UGameplayEffect>(StaticConstructObject(UGameplayEffect::StaticClass(), GetTransientPackage(), FName(TEXT("StackingEffect4"))));
-		BaseStackedEffect->Modifiers.SetNum(1);
-		BaseStackedEffect->Modifiers[0].Magnitude.SetValue(StackingValue);
-		BaseStackedEffect->Modifiers[0].ModifierType = EGameplayMod::Attribute;
-		BaseStackedEffect->Modifiers[0].ModifierOp = EGameplayModOp::Additive;
-		BaseStackedEffect->Modifiers[0].Attribute.SetUProperty(StackingProperty);
-		BaseStackedEffect->Modifiers[0].OwnedTags.AddTag(IGameplayTagsModule::RequestGameplayTag("Stackable"));
-		BaseStackedEffect->Duration.Value = 2.f;
-		BaseStackedEffect->Period.Value = 1.f;
-		BaseStackedEffect->StackingPolicy = EGameplayEffectStackingPolicy::Callback;
-		BaseStackedEffect->StackingExtension = UGameplayEffectStackingExtension_CappedNumberTest::StaticClass();
-		BaseStackedEffect->StackedAttribName = FName(*StackingProperty->GetName());
-
-		BaseStackedEffect->ValidateGameplayEffect();
-
-		SourceComponent->ApplyGameplayEffectToTarget(BaseStackedEffect, DestComponent);
-
-		GameplayTest_TickWorld(World, 0.0001f); // Move our Effects from the pending stack to the active stack inside the timer manager, this starts the clock for effect execution
-		GameplayTest_TickWorld(World, 0.1f); // offset the current time from the start of the period to avoid floating point issues causing the tests to fail
-	}
-
-	{
-		// Tick once
-		GameplayTest_TickWorld(World, 1.f);
-
-		float ActualValue = DestComponent->GetSet<UAbilitySystemTestAttributeSet>()->StackingAttribute1;
-		float ExpectedValue = StackingValue * 6.f;	// 2 for effects being applied, 2 for tick, 2 for tick
-		Test->TestTrue(SKILL_TEST_TEXT("Stacking Attribute: Actual: %.2f == Exected: %.2f", ActualValue, ExpectedValue), ActualValue == ExpectedValue);
-	}
-
-	{
-		// Tick once
-		GameplayTest_TickWorld(World, 1.f);
-
-		float ActualValue = DestComponent->GetSet<UAbilitySystemTestAttributeSet>()->StackingAttribute1;
-		float ExpectedValue = StackingValue * 8.f;	// the last GE should have refreshed the timer so we should have 2 for GEs applied and 3 * 2 for ticks
-		Test->TestTrue(SKILL_TEST_TEXT("Stacking Attribute: Actual: %.2f == Exected: %.2f", ActualValue, ExpectedValue), ActualValue == ExpectedValue);
-	}
-
-	{
-		// Tick once
-		GameplayTest_TickWorld(World, 1.f);
-
-		float ActualValue = DestComponent->GetSet<UAbilitySystemTestAttributeSet>()->StackingAttribute1;
-		float ExpectedValue = StackingValue * 8.f;	// the effects should have timed out so we should have 2 for GEs applied and 3 * 2 for ticks
-		Test->TestTrue(SKILL_TEST_TEXT("Stacking Attribute: Actual: %.2f == Exected: %.2f", ActualValue, ExpectedValue), ActualValue == ExpectedValue);
-	}
-
-	ClearGlobalCurveTable();
-	ClearGlobalDataTable();
-
-	World->EditorDestroyActor(SourceActor, false);
-	World->EditorDestroyActor(DestActor, false);
-
-	return true;
-}
-
-// test the custom rule that stacks based on diminishing returns
-bool GameplayEffectsTest_StackingCustomDiminishingReturns(UWorld *World, FAutomationTestBase * Test)
-{
-	const float StackingValue = 1.f;
-
-	SetGlobalCurveTable();
-	SetGlobalDataTable();
-
-	AAbilitySystemTestPawn *SourceActor = World->SpawnActor<AAbilitySystemTestPawn>();
-	AAbilitySystemTestPawn *DestActor = World->SpawnActor<AAbilitySystemTestPawn>();
-
-	UProperty *StackingProperty = FindFieldChecked<UProperty>(UAbilitySystemTestAttributeSet::StaticClass(), GET_MEMBER_NAME_CHECKED(UAbilitySystemTestAttributeSet, StackingAttribute1));
-
-	UAbilitySystemComponent * SourceComponent = SourceActor->GetAbilitySystemComponent();
-	UAbilitySystemComponent * DestComponent = DestActor->GetAbilitySystemComponent();
-
-	{
-		ABILITY_LOG_SCOPE(TEXT("Apply StackingEffect"));
-
-		UGameplayEffect * BaseStackedEffect = Cast<UGameplayEffect>(StaticConstructObject(UGameplayEffect::StaticClass(), GetTransientPackage(), FName(TEXT("StackingEffect1"))));
-		BaseStackedEffect->Modifiers.SetNum(1);
-		BaseStackedEffect->Modifiers[0].Magnitude.SetValue(StackingValue);
-		BaseStackedEffect->Modifiers[0].ModifierType = EGameplayMod::Attribute;
-		BaseStackedEffect->Modifiers[0].ModifierOp = EGameplayModOp::Additive;
-		BaseStackedEffect->Modifiers[0].Attribute.SetUProperty(StackingProperty);
-		BaseStackedEffect->Modifiers[0].OwnedTags.AddTag(IGameplayTagsModule::RequestGameplayTag("Stackable"));
-		BaseStackedEffect->Duration.Value = UGameplayEffect::INFINITE_DURATION;
-		BaseStackedEffect->Period.Value = 1.f;
-		BaseStackedEffect->StackingPolicy = EGameplayEffectStackingPolicy::Callback;
-		BaseStackedEffect->StackingExtension = UGameplayEffectStackingExtension_DiminishingReturnsTest::StaticClass();
-		BaseStackedEffect->StackedAttribName = FName(*StackingProperty->GetName());
-
-		BaseStackedEffect->ValidateGameplayEffect();
-
-		SourceComponent->ApplyGameplayEffectToTarget(BaseStackedEffect, DestComponent);
-
-		GameplayTest_TickWorld(World, 0.0001f); // Move our Effects from the pending stack to the active stack inside the timer manager, this starts the clock for effect execution
-		GameplayTest_TickWorld(World, 0.1f); // Offset the current time from the start of the period so that floating point issues don't affect the test
-	}
-
-	{
-		// Tick once
-		GameplayTest_TickWorld(World, 1.f);
-
-		float ActualValue = DestComponent->GetSet<UAbilitySystemTestAttributeSet>()->StackingAttribute1;
-		float ExpectedValue = StackingValue * (5.f * 2);	// first application gets five times the result and will be applied at the start and end of the tick
-		Test->TestTrue(SKILL_TEST_TEXT("Stacking Attribute: Actual: %.2f == Exected: %.2f", ActualValue, ExpectedValue), ActualValue == ExpectedValue);
-	}
-
-	{
-		ABILITY_LOG_SCOPE(TEXT("Apply StackingEffect"));
-
-		UGameplayEffect * BaseStackedEffect = Cast<UGameplayEffect>(StaticConstructObject(UGameplayEffect::StaticClass(), GetTransientPackage(), FName(TEXT("StackingEffect2"))));
-		BaseStackedEffect->Modifiers.SetNum(1);
-		BaseStackedEffect->Modifiers[0].Magnitude.SetValue(StackingValue);
-		BaseStackedEffect->Modifiers[0].ModifierType = EGameplayMod::Attribute;
-		BaseStackedEffect->Modifiers[0].ModifierOp = EGameplayModOp::Additive;
-		BaseStackedEffect->Modifiers[0].Attribute.SetUProperty(StackingProperty);
-		BaseStackedEffect->Modifiers[0].OwnedTags.AddTag(IGameplayTagsModule::RequestGameplayTag("Stackable"));
-		BaseStackedEffect->Duration.Value = UGameplayEffect::INFINITE_DURATION;
-		BaseStackedEffect->Period.Value = 1.f;
-		BaseStackedEffect->StackingPolicy = EGameplayEffectStackingPolicy::Callback;
-		BaseStackedEffect->StackingExtension = UGameplayEffectStackingExtension_DiminishingReturnsTest::StaticClass();
-		BaseStackedEffect->StackedAttribName = FName(*StackingProperty->GetName());
-
-		BaseStackedEffect->ValidateGameplayEffect();
-
-		SourceComponent->ApplyGameplayEffectToTarget(BaseStackedEffect, DestComponent);
-
-		GameplayTest_TickWorld(World, 0.0001f); // Move our Effects from the pending stack to the active stack inside the timer manager, this starts the clock for effect execution
-		GameplayTest_TickWorld(World, 0.1f); // Offset the current time from the start of the period so that floating point issues don't affect the test
-	}
-
-	{
-		// Tick once
-		GameplayTest_TickWorld(World, 1.f);
-
-		float ActualValue = DestComponent->GetSet<UAbilitySystemTestAttributeSet>()->StackingAttribute1;
-		float ExpectedValue = StackingValue * ((5.f * 2) + 7.f);	// second application gets seven times the result
-		Test->TestTrue(SKILL_TEST_TEXT("Stacking Attribute: Actual: %.2f == Exected: %.2f", ActualValue, ExpectedValue), ActualValue == ExpectedValue);
-	}
-
-	{
-		ABILITY_LOG_SCOPE(TEXT("Apply StackingEffect"));
-
-		UGameplayEffect * BaseStackedEffect = Cast<UGameplayEffect>(StaticConstructObject(UGameplayEffect::StaticClass(), GetTransientPackage(), FName(TEXT("StackingEffect3"))));
-		BaseStackedEffect->Modifiers.SetNum(1);
-		BaseStackedEffect->Modifiers[0].Magnitude.SetValue(StackingValue);
-		BaseStackedEffect->Modifiers[0].ModifierType = EGameplayMod::Attribute;
-		BaseStackedEffect->Modifiers[0].ModifierOp = EGameplayModOp::Additive;
-		BaseStackedEffect->Modifiers[0].Attribute.SetUProperty(StackingProperty);
-		BaseStackedEffect->Modifiers[0].OwnedTags.AddTag(IGameplayTagsModule::RequestGameplayTag("Stackable"));
-		BaseStackedEffect->Duration.Value = UGameplayEffect::INFINITE_DURATION;
-		BaseStackedEffect->Period.Value = 1.f;
-		BaseStackedEffect->StackingPolicy = EGameplayEffectStackingPolicy::Callback;
-		BaseStackedEffect->StackingExtension = UGameplayEffectStackingExtension_DiminishingReturnsTest::StaticClass();
-		BaseStackedEffect->StackedAttribName = FName(*StackingProperty->GetName());
-
-		BaseStackedEffect->ValidateGameplayEffect();
-
-		SourceComponent->ApplyGameplayEffectToTarget(BaseStackedEffect, DestComponent);
-
-		GameplayTest_TickWorld(World, 0.0001f); // Move our Effects from the pending stack to the active stack inside the timer manager, this starts the clock for effect execution
-		GameplayTest_TickWorld(World, 0.1f); // Offset the current time from the start of the period so that floating point issues don't affect the test
-	}
-
-	{
-		// Tick once
-		GameplayTest_TickWorld(World, 1.f);
-
-		float ActualValue = DestComponent->GetSet<UAbilitySystemTestAttributeSet>()->StackingAttribute1;
-		float ExpectedValue = StackingValue * ((5.f * 2) + 7.f + 8.f);	// third application gets eight times the result
-		Test->TestTrue(SKILL_TEST_TEXT("Stacking Attribute: Actual: %.2f == Exected: %.2f", ActualValue, ExpectedValue), ActualValue == ExpectedValue);
-	}
-
-	{
-		ABILITY_LOG_SCOPE(TEXT("Apply StackingEffect"));
-
-		UGameplayEffect * BaseStackedEffect = Cast<UGameplayEffect>(StaticConstructObject(UGameplayEffect::StaticClass(), GetTransientPackage(), FName(TEXT("StackingEffect4"))));
-		BaseStackedEffect->Modifiers.SetNum(1);
-		BaseStackedEffect->Modifiers[0].Magnitude.SetValue(StackingValue);
-		BaseStackedEffect->Modifiers[0].ModifierType = EGameplayMod::Attribute;
-		BaseStackedEffect->Modifiers[0].ModifierOp = EGameplayModOp::Additive;
-		BaseStackedEffect->Modifiers[0].Attribute.SetUProperty(StackingProperty);
-		BaseStackedEffect->Modifiers[0].OwnedTags.AddTag(IGameplayTagsModule::RequestGameplayTag("Stackable"));
-		BaseStackedEffect->Duration.Value = UGameplayEffect::INFINITE_DURATION;
-		BaseStackedEffect->Period.Value = 1.f;
-		BaseStackedEffect->StackingPolicy = EGameplayEffectStackingPolicy::Callback;
-		BaseStackedEffect->StackingExtension = UGameplayEffectStackingExtension_DiminishingReturnsTest::StaticClass();
-		BaseStackedEffect->StackedAttribName = FName(*StackingProperty->GetName());
-
-		BaseStackedEffect->ValidateGameplayEffect();
-
-		SourceComponent->ApplyGameplayEffectToTarget(BaseStackedEffect, DestComponent);
-
-		GameplayTest_TickWorld(World, 0.0001f); // Move our Effects from the pending stack to the active stack inside the timer manager, this starts the clock for effect execution
-		GameplayTest_TickWorld(World, 0.1f); // Offset the current time from the start of the period so that floating point issues don't affect the test
-	}
-
-	{
-		// Tick once
-		GameplayTest_TickWorld(World, 1.f);
-
-		float ActualValue = DestComponent->GetSet<UAbilitySystemTestAttributeSet>()->StackingAttribute1;
-		float ExpectedValue = StackingValue * ((5.f * 2) + 7.f + 8.f + 9.f);	// fourth application gets nine times the result
-		Test->TestTrue(SKILL_TEST_TEXT("Stacking Attribute: Actual: %.2f == Exected: %.2f", ActualValue, ExpectedValue), ActualValue == ExpectedValue);
-	}
-
-	ClearGlobalCurveTable();
-	ClearGlobalDataTable();
-
-	World->EditorDestroyActor(SourceActor, false);
-	World->EditorDestroyActor(DestActor, false);
-
-	return true;
-}
-
-// Gameplay effects that change the same attribute but have different stacking rules shouldn't interfere with each other
-bool GameplayEffectsTest_StackingDifferentRules(UWorld *World, FAutomationTestBase * Test)
-{
-	const float StackingValue = 5.f;
-
-	SetGlobalCurveTable();
-	SetGlobalDataTable();
-
-	AAbilitySystemTestPawn *SourceActor = World->SpawnActor<AAbilitySystemTestPawn>();
-	AAbilitySystemTestPawn *DestActor = World->SpawnActor<AAbilitySystemTestPawn>();
-
-	UProperty *StackingProperty = FindFieldChecked<UProperty>(UAbilitySystemTestAttributeSet::StaticClass(), GET_MEMBER_NAME_CHECKED(UAbilitySystemTestAttributeSet, StackingAttribute1));
-
-	UAbilitySystemComponent * SourceComponent = SourceActor->GetAbilitySystemComponent();
-	UAbilitySystemComponent * DestComponent = DestActor->GetAbilitySystemComponent();
-
-	{
-		ABILITY_LOG_SCOPE(TEXT("Apply StackingEffect"));
-
-		UGameplayEffect * BaseStackedEffect = Cast<UGameplayEffect>(StaticConstructObject(UGameplayEffect::StaticClass(), GetTransientPackage(), FName(TEXT("StackingEffect1"))));
-		BaseStackedEffect->Modifiers.SetNum(1);
-		BaseStackedEffect->Modifiers[0].Magnitude.SetValue(StackingValue);
-		BaseStackedEffect->Modifiers[0].ModifierType = EGameplayMod::Attribute;
-		BaseStackedEffect->Modifiers[0].ModifierOp = EGameplayModOp::Additive;
-		BaseStackedEffect->Modifiers[0].Attribute.SetUProperty(StackingProperty);
-		BaseStackedEffect->Duration.Value = UGameplayEffect::INFINITE_DURATION;
-		BaseStackedEffect->Period.Value = 1.f;
-		BaseStackedEffect->StackingPolicy = EGameplayEffectStackingPolicy::Highest;
-		BaseStackedEffect->StackedAttribName = FName(*StackingProperty->GetName());
-
-		BaseStackedEffect->ValidateGameplayEffect();
-
-		SourceComponent->ApplyGameplayEffectToTarget(BaseStackedEffect, DestComponent);
-	}
-
-	{
-		ABILITY_LOG_SCOPE(TEXT("Apply StackingEffect"));
-
-		UGameplayEffect * BaseStackedEffect = Cast<UGameplayEffect>(StaticConstructObject(UGameplayEffect::StaticClass(), GetTransientPackage(), FName(TEXT("StackingEffect2"))));
-		BaseStackedEffect->Modifiers.SetNum(1);
-		BaseStackedEffect->Modifiers[0].Magnitude.SetValue(StackingValue * 2.f);
-		BaseStackedEffect->Modifiers[0].ModifierType = EGameplayMod::Attribute;
-		BaseStackedEffect->Modifiers[0].ModifierOp = EGameplayModOp::Additive;
-		BaseStackedEffect->Modifiers[0].Attribute.SetUProperty(StackingProperty);
-		BaseStackedEffect->Duration.Value = UGameplayEffect::INFINITE_DURATION;
-		BaseStackedEffect->Period.Value = 1.f;
-		BaseStackedEffect->StackingPolicy = EGameplayEffectStackingPolicy::Lowest;
-		BaseStackedEffect->StackedAttribName = FName(*StackingProperty->GetName());
-
-		BaseStackedEffect->ValidateGameplayEffect();
-
-		SourceComponent->ApplyGameplayEffectToTarget(BaseStackedEffect, DestComponent);
-
-		GameplayTest_TickWorld(World, 0.0001f); // Move our Effects from the pending stack to the active stack inside the timer manager, this starts the clock for effect execution
-		GameplayTest_TickWorld(World, 0.1f); // Offset the current time from the start of the period so that floating point issues don't affect the test
-	}
-
-	{
-		// Tick once
-		GameplayTest_TickWorld(World, 1.f);
-
-		float ActualValue = DestComponent->GetSet<UAbilitySystemTestAttributeSet>()->StackingAttribute1;
-		float ExpectedValue = StackingValue * 6.f;	// 1 for first effect, 2 for the second effect, 3 during tick
-		Test->TestTrue(SKILL_TEST_TEXT("Stacking Attribute: Actual: %.2f == Exected: %.2f", ActualValue, ExpectedValue), ActualValue == ExpectedValue);
-	}
-
-	ClearGlobalCurveTable();
-	ClearGlobalDataTable();
-
-	World->EditorDestroyActor(SourceActor, false);
-	World->EditorDestroyActor(DestActor, false);
-
-	return true;
-}
-
-// different attributes with the same stacking rule shouldn't interfere with each other
-bool GameplayEffectsTest_StackingDifferentAttributes(UWorld *World, FAutomationTestBase * Test)
-{
-	const float StackingValue = 5.f;
-
-	SetGlobalCurveTable();
-	SetGlobalDataTable();
-
-	AAbilitySystemTestPawn *SourceActor = World->SpawnActor<AAbilitySystemTestPawn>();
-	AAbilitySystemTestPawn *DestActor = World->SpawnActor<AAbilitySystemTestPawn>();
-
-	UProperty *StackingProperty1 = FindFieldChecked<UProperty>(UAbilitySystemTestAttributeSet::StaticClass(), GET_MEMBER_NAME_CHECKED(UAbilitySystemTestAttributeSet, StackingAttribute1));
-	UProperty *StackingProperty2 = FindFieldChecked<UProperty>(UAbilitySystemTestAttributeSet::StaticClass(), GET_MEMBER_NAME_CHECKED(UAbilitySystemTestAttributeSet, StackingAttribute2));
-
-	UAbilitySystemComponent * SourceComponent = SourceActor->GetAbilitySystemComponent();
-	UAbilitySystemComponent * DestComponent = DestActor->GetAbilitySystemComponent();
-
-	{
-		ABILITY_LOG_SCOPE(TEXT("Apply StackingEffect"));
-
-		UGameplayEffect * BaseStackedEffect = Cast<UGameplayEffect>(StaticConstructObject(UGameplayEffect::StaticClass(), GetTransientPackage(), FName(TEXT("StackingEffect1"))));
-		BaseStackedEffect->Modifiers.SetNum(1);
-		BaseStackedEffect->Modifiers[0].Magnitude.SetValue(StackingValue);
-		BaseStackedEffect->Modifiers[0].ModifierType = EGameplayMod::Attribute;
-		BaseStackedEffect->Modifiers[0].ModifierOp = EGameplayModOp::Additive;
-		BaseStackedEffect->Modifiers[0].Attribute.SetUProperty(StackingProperty1);
-		BaseStackedEffect->Duration.Value = UGameplayEffect::INFINITE_DURATION;
-		BaseStackedEffect->Period.Value = 1.f;
-		BaseStackedEffect->StackingPolicy = EGameplayEffectStackingPolicy::Highest;
-		BaseStackedEffect->StackedAttribName = FName(*StackingProperty1->GetName());
-
-		BaseStackedEffect->ValidateGameplayEffect();
-
-		SourceComponent->ApplyGameplayEffectToTarget(BaseStackedEffect, DestComponent);
-	}
-
-	{
-		ABILITY_LOG_SCOPE(TEXT("Apply StackingEffect"));
-
-		UGameplayEffect * BaseStackedEffect = Cast<UGameplayEffect>(StaticConstructObject(UGameplayEffect::StaticClass(), GetTransientPackage(), FName(TEXT("StackingEffect2"))));
-		BaseStackedEffect->Modifiers.SetNum(1);
-		BaseStackedEffect->Modifiers[0].Magnitude.SetValue(StackingValue * 2.f);
-		BaseStackedEffect->Modifiers[0].ModifierType = EGameplayMod::Attribute;
-		BaseStackedEffect->Modifiers[0].ModifierOp = EGameplayModOp::Additive;
-		BaseStackedEffect->Modifiers[0].Attribute.SetUProperty(StackingProperty2);
-		BaseStackedEffect->Duration.Value = UGameplayEffect::INFINITE_DURATION;
-		BaseStackedEffect->Period.Value = 1.f;
-		BaseStackedEffect->StackingPolicy = EGameplayEffectStackingPolicy::Highest;
-		BaseStackedEffect->StackedAttribName = FName(*StackingProperty2->GetName());
-
-		BaseStackedEffect->ValidateGameplayEffect();
-
-		SourceComponent->ApplyGameplayEffectToTarget(BaseStackedEffect, DestComponent);
-
-		GameplayTest_TickWorld(World, 0.0001f); // Move our Effects from the pending stack to the active stack inside the timer manager, this starts the clock for effect execution
-		GameplayTest_TickWorld(World, 0.1f); // Offset the current time from the start of the period so that floating point issues don't affect the test
-	}
-
-	{
-		// Tick once
-		GameplayTest_TickWorld(World, 1.f);
-
-		float ActualValue = DestComponent->GetSet<UAbilitySystemTestAttributeSet>()->StackingAttribute1;
-		float ExpectedValue = StackingValue * 2.f;	// 1 for first effect, 1 for first effect during tick
-		Test->TestTrue(SKILL_TEST_TEXT("Stacking Attribute 1: %.2f == Exected: %.2f", ActualValue, ExpectedValue), ActualValue == ExpectedValue);
-		ActualValue = DestComponent->GetSet<UAbilitySystemTestAttributeSet>()->StackingAttribute2;
-		ExpectedValue = StackingValue * 4.f; // 2 for the second effect, 2 for the second effect during tick
-		Test->TestTrue(SKILL_TEST_TEXT("Stacking Attribute 2: Actual: %.2f == Exected: %.2f", ActualValue, ExpectedValue), ActualValue == ExpectedValue);
-	}
-
-	ClearGlobalCurveTable();
-	ClearGlobalDataTable();
-
-	World->EditorDestroyActor(SourceActor, false);
-	World->EditorDestroyActor(DestActor, false);
-
-	return true;
-}
-
-// test to make sure that two different custom rules don't interfere with each other
-bool GameplayEffectsTest_StackingCustomTwoRules(UWorld *World, FAutomationTestBase * Test)
-{
-	const float StackingValue = 1.f;
-
-	SetGlobalCurveTable();
-	SetGlobalDataTable();
-
-	AAbilitySystemTestPawn *SourceActor = World->SpawnActor<AAbilitySystemTestPawn>();
-	AAbilitySystemTestPawn *DestActor = World->SpawnActor<AAbilitySystemTestPawn>();
-
-	UProperty *StackingProperty = FindFieldChecked<UProperty>(UAbilitySystemTestAttributeSet::StaticClass(), GET_MEMBER_NAME_CHECKED(UAbilitySystemTestAttributeSet, StackingAttribute1));
-
-	UAbilitySystemComponent * SourceComponent = SourceActor->GetAbilitySystemComponent();
-	UAbilitySystemComponent * DestComponent = DestActor->GetAbilitySystemComponent();
-
-	{
-		ABILITY_LOG_SCOPE(TEXT("Apply StackingEffect"));
-
-		UGameplayEffect * BaseStackedEffect = Cast<UGameplayEffect>(StaticConstructObject(UGameplayEffect::StaticClass(), GetTransientPackage(), FName(TEXT("StackingEffect1"))));
-		BaseStackedEffect->Modifiers.SetNum(1);
-		BaseStackedEffect->Modifiers[0].Magnitude.SetValue(StackingValue);
-		BaseStackedEffect->Modifiers[0].ModifierType = EGameplayMod::Attribute;
-		BaseStackedEffect->Modifiers[0].ModifierOp = EGameplayModOp::Additive;
-		BaseStackedEffect->Modifiers[0].Attribute.SetUProperty(StackingProperty);
-		BaseStackedEffect->Modifiers[0].OwnedTags.AddTag(IGameplayTagsModule::RequestGameplayTag("Stackable"));
-		BaseStackedEffect->Duration.Value = UGameplayEffect::INFINITE_DURATION;
-		BaseStackedEffect->Period.Value = 1.f;
-		BaseStackedEffect->StackingPolicy = EGameplayEffectStackingPolicy::Callback;
-		BaseStackedEffect->StackingExtension = UGameplayEffectStackingExtension_CappedNumberTest::StaticClass();
-		BaseStackedEffect->StackedAttribName = FName(*StackingProperty->GetName());
-
-		BaseStackedEffect->ValidateGameplayEffect();
-
-		SourceComponent->ApplyGameplayEffectToTarget(BaseStackedEffect, DestComponent);
-	}
-
-	{
-		ABILITY_LOG_SCOPE(TEXT("Apply StackingEffect"));
-
-		UGameplayEffect * BaseStackedEffect = Cast<UGameplayEffect>(StaticConstructObject(UGameplayEffect::StaticClass(), GetTransientPackage(), FName(TEXT("StackingEffect2"))));
-		BaseStackedEffect->Modifiers.SetNum(1);
-		BaseStackedEffect->Modifiers[0].Magnitude.SetValue(StackingValue);
-		BaseStackedEffect->Modifiers[0].ModifierType = EGameplayMod::Attribute;
-		BaseStackedEffect->Modifiers[0].ModifierOp = EGameplayModOp::Additive;
-		BaseStackedEffect->Modifiers[0].Attribute.SetUProperty(StackingProperty);
-		BaseStackedEffect->Modifiers[0].OwnedTags.AddTag(IGameplayTagsModule::RequestGameplayTag("Stackable"));
-		BaseStackedEffect->Duration.Value = UGameplayEffect::INFINITE_DURATION;
-		BaseStackedEffect->Period.Value = 1.f;
-		BaseStackedEffect->StackingPolicy = EGameplayEffectStackingPolicy::Callback;
-		BaseStackedEffect->StackingExtension = UGameplayEffectStackingExtension_DiminishingReturnsTest::StaticClass();
-		BaseStackedEffect->StackedAttribName = FName(*StackingProperty->GetName());
-
-		BaseStackedEffect->ValidateGameplayEffect();
-
-		SourceComponent->ApplyGameplayEffectToTarget(BaseStackedEffect, DestComponent);
-	}
-
-	{
-		ABILITY_LOG_SCOPE(TEXT("Apply StackingEffect"));
-
-		UGameplayEffect * BaseStackedEffect = Cast<UGameplayEffect>(StaticConstructObject(UGameplayEffect::StaticClass(), GetTransientPackage(), FName(TEXT("StackingEffect3"))));
-		BaseStackedEffect->Modifiers.SetNum(1);
-		BaseStackedEffect->Modifiers[0].Magnitude.SetValue(StackingValue);
-		BaseStackedEffect->Modifiers[0].ModifierType = EGameplayMod::Attribute;
-		BaseStackedEffect->Modifiers[0].ModifierOp = EGameplayModOp::Additive;
-		BaseStackedEffect->Modifiers[0].Attribute.SetUProperty(StackingProperty);
-		BaseStackedEffect->Modifiers[0].OwnedTags.AddTag(IGameplayTagsModule::RequestGameplayTag("Stackable"));
-		BaseStackedEffect->Duration.Value = UGameplayEffect::INFINITE_DURATION;
-		BaseStackedEffect->Period.Value = 1.f;
-		BaseStackedEffect->StackingPolicy = EGameplayEffectStackingPolicy::Callback;
-		BaseStackedEffect->StackingExtension = UGameplayEffectStackingExtension_DiminishingReturnsTest::StaticClass();
-		BaseStackedEffect->StackedAttribName = FName(*StackingProperty->GetName());
-
-		BaseStackedEffect->ValidateGameplayEffect();
-
-		SourceComponent->ApplyGameplayEffectToTarget(BaseStackedEffect, DestComponent);
-
-		GameplayTest_TickWorld(World, 0.0001f); // Move our Effects from the pending stack to the active stack inside the timer manager, this starts the clock for effect execution
-		GameplayTest_TickWorld(World, 0.1f); // Offset the current time from the start of the period so that floating point issues don't affect the test
-	}
-
-	{
-		// Tick once
-		GameplayTest_TickWorld(World, 1.f);
-
-		float ActualValue = DestComponent->GetSet<UAbilitySystemTestAttributeSet>()->StackingAttribute1;
-		float ExpectedValue = StackingValue * ((1 + 7) * 2);	// the capped stacking rule will apply one, the diminishing returns rule will apply seven, both rules are applied twice
-		Test->TestTrue(SKILL_TEST_TEXT("Stacking Attribute: Actual: %.2f == Exected: %.2f", ActualValue, ExpectedValue), ActualValue == ExpectedValue);
-	}
-
-	ClearGlobalCurveTable();
-	ClearGlobalDataTable();
-
-	World->EditorDestroyActor(SourceActor, false);
-	World->EditorDestroyActor(DestActor, false);
-
-	return true;
-}
-
-// test to make sure that a custom rules applied to different attributes doesn't interfere with itself
-bool GameplayEffectsTest_StackingCustomTwoAttributes(UWorld *World, FAutomationTestBase * Test)
-{
-	const float StackingValue = 1.f;
-
-	SetGlobalCurveTable();
-	SetGlobalDataTable();
-
-	AAbilitySystemTestPawn *SourceActor = World->SpawnActor<AAbilitySystemTestPawn>();
-	AAbilitySystemTestPawn *DestActor = World->SpawnActor<AAbilitySystemTestPawn>();
-
-	UProperty *StackingProperty1 = FindFieldChecked<UProperty>(UAbilitySystemTestAttributeSet::StaticClass(), GET_MEMBER_NAME_CHECKED(UAbilitySystemTestAttributeSet, StackingAttribute1));
-	UProperty *StackingProperty2 = FindFieldChecked<UProperty>(UAbilitySystemTestAttributeSet::StaticClass(), GET_MEMBER_NAME_CHECKED(UAbilitySystemTestAttributeSet, StackingAttribute2));
-
-	UAbilitySystemComponent * SourceComponent = SourceActor->GetAbilitySystemComponent();
-	UAbilitySystemComponent * DestComponent = DestActor->GetAbilitySystemComponent();
-
-	{
-		ABILITY_LOG_SCOPE(TEXT("Apply StackingEffect"));
-
-		UGameplayEffect * BaseStackedEffect = Cast<UGameplayEffect>(StaticConstructObject(UGameplayEffect::StaticClass(), GetTransientPackage(), FName(TEXT("StackingEffect1"))));
-		BaseStackedEffect->Modifiers.SetNum(1);
-		BaseStackedEffect->Modifiers[0].Magnitude.SetValue(StackingValue);
-		BaseStackedEffect->Modifiers[0].ModifierType = EGameplayMod::Attribute;
-		BaseStackedEffect->Modifiers[0].ModifierOp = EGameplayModOp::Additive;
-		BaseStackedEffect->Modifiers[0].Attribute.SetUProperty(StackingProperty1);
-		BaseStackedEffect->Modifiers[0].OwnedTags.AddTag(IGameplayTagsModule::RequestGameplayTag("Stackable"));
-		BaseStackedEffect->Duration.Value = UGameplayEffect::INFINITE_DURATION;
-		BaseStackedEffect->Period.Value = 1.f;
-		BaseStackedEffect->StackingPolicy = EGameplayEffectStackingPolicy::Callback;
-		BaseStackedEffect->StackingExtension = UGameplayEffectStackingExtension_CappedNumberTest::StaticClass();
-		BaseStackedEffect->StackedAttribName = FName(*StackingProperty1->GetName());
-
-		BaseStackedEffect->ValidateGameplayEffect();
-
-		SourceComponent->ApplyGameplayEffectToTarget(BaseStackedEffect, DestComponent);
-	}
-
-	{
-		ABILITY_LOG_SCOPE(TEXT("Apply StackingEffect"));
-
-		UGameplayEffect * BaseStackedEffect = Cast<UGameplayEffect>(StaticConstructObject(UGameplayEffect::StaticClass(), GetTransientPackage(), FName(TEXT("StackingEffect2"))));
-		BaseStackedEffect->Modifiers.SetNum(1);
-		BaseStackedEffect->Modifiers[0].Magnitude.SetValue(StackingValue);
-		BaseStackedEffect->Modifiers[0].ModifierType = EGameplayMod::Attribute;
-		BaseStackedEffect->Modifiers[0].ModifierOp = EGameplayModOp::Additive;
-		BaseStackedEffect->Modifiers[0].Attribute.SetUProperty(StackingProperty1);
-		BaseStackedEffect->Modifiers[0].OwnedTags.AddTag(IGameplayTagsModule::RequestGameplayTag("Stackable"));
-		BaseStackedEffect->Duration.Value = UGameplayEffect::INFINITE_DURATION;
-		BaseStackedEffect->Period.Value = 1.f;
-		BaseStackedEffect->StackingPolicy = EGameplayEffectStackingPolicy::Callback;
-		BaseStackedEffect->StackingExtension = UGameplayEffectStackingExtension_CappedNumberTest::StaticClass();
-		BaseStackedEffect->StackedAttribName = FName(*StackingProperty1->GetName());
-
-		BaseStackedEffect->ValidateGameplayEffect();
-
-		SourceComponent->ApplyGameplayEffectToTarget(BaseStackedEffect, DestComponent);
-	}
-
-	{
-		ABILITY_LOG_SCOPE(TEXT("Apply StackingEffect"));
-
-		UGameplayEffect * BaseStackedEffect = Cast<UGameplayEffect>(StaticConstructObject(UGameplayEffect::StaticClass(), GetTransientPackage(), FName(TEXT("StackingEffect3"))));
-		BaseStackedEffect->Modifiers.SetNum(1);
-		BaseStackedEffect->Modifiers[0].Magnitude.SetValue(StackingValue);
-		BaseStackedEffect->Modifiers[0].ModifierType = EGameplayMod::Attribute;
-		BaseStackedEffect->Modifiers[0].ModifierOp = EGameplayModOp::Additive;
-		BaseStackedEffect->Modifiers[0].Attribute.SetUProperty(StackingProperty2);
-		BaseStackedEffect->Modifiers[0].OwnedTags.AddTag(IGameplayTagsModule::RequestGameplayTag("Stackable"));
-		BaseStackedEffect->Duration.Value = UGameplayEffect::INFINITE_DURATION;
-		BaseStackedEffect->Period.Value = 1.f;
-		BaseStackedEffect->StackingPolicy = EGameplayEffectStackingPolicy::Callback;
-		BaseStackedEffect->StackingExtension = UGameplayEffectStackingExtension_CappedNumberTest::StaticClass();
-		BaseStackedEffect->StackedAttribName = FName(*StackingProperty2->GetName());
-
-		BaseStackedEffect->ValidateGameplayEffect();
-
-		SourceComponent->ApplyGameplayEffectToTarget(BaseStackedEffect, DestComponent);
-
-		GameplayTest_TickWorld(World, 0.0001f); // Move our Effects from the pending stack to the active stack inside the timer manager, this starts the clock for effect execution
-		GameplayTest_TickWorld(World, 0.1f); // Offset the current time from the start of the period so that floating point issues don't affect the test
-	}
-
-	{
-		// Tick once
-		GameplayTest_TickWorld(World, 1.f);
-
-		float ActualValue = DestComponent->GetSet<UAbilitySystemTestAttributeSet>()->StackingAttribute1;
-		float ExpectedValue = StackingValue * (2 * 2);	// two effects should be applied to the first attribute, both apply twice
-		Test->TestTrue(SKILL_TEST_TEXT("Stacking Attribute: Actual: %.2f == Exected: %.2f", ActualValue, ExpectedValue), ActualValue == ExpectedValue);
-
-		ActualValue = DestComponent->GetSet<UAbilitySystemTestAttributeSet>()->StackingAttribute2;
-		ExpectedValue = StackingValue * (1 * 2);	// one effect should be applied twice to the second attribute
-		Test->TestTrue(SKILL_TEST_TEXT("Stacking Attribute: Actual: %.2f == Exected: %.2f", ActualValue, ExpectedValue), ActualValue == ExpectedValue);
-	}
-
-	ClearGlobalCurveTable();
-	ClearGlobalDataTable();
-
-	World->EditorDestroyActor(SourceActor, false);
-	World->EditorDestroyActor(DestActor, false);
-
-	return true;
-}
-
-// Test to make sure that removing a modifier causes stacks to be recalculated correctly
-bool GameplayEffectsTest_StackingRemovingModifiers(UWorld *World, FAutomationTestBase * Test)
-{
-	const float StackingValue = 5.f;
-
-	SetGlobalCurveTable();
-	SetGlobalDataTable();
-
-	AAbilitySystemTestPawn *SourceActor = World->SpawnActor<AAbilitySystemTestPawn>();
-	AAbilitySystemTestPawn *DestActor = World->SpawnActor<AAbilitySystemTestPawn>();
-
-	UProperty *StackingProperty = FindFieldChecked<UProperty>(UAbilitySystemTestAttributeSet::StaticClass(), GET_MEMBER_NAME_CHECKED(UAbilitySystemTestAttributeSet, StackingAttribute1));
-
-	UAbilitySystemComponent * SourceComponent = SourceActor->GetAbilitySystemComponent();
-	UAbilitySystemComponent * DestComponent = DestActor->GetAbilitySystemComponent();
-
-	// Setup a GE to modify IncomingGEs
-	{
-		ABILITY_LOG_SCOPE(TEXT("Apply modifier to incoming, tagged GEs"));
-
-		UGameplayEffect* BaseModEffect = Cast<UGameplayEffect>(StaticConstructObject(UGameplayEffect::StaticClass(), GetTransientPackage(), FName(TEXT("ProtectBuff"))));
-		BaseModEffect->Modifiers.SetNum(1);
-		BaseModEffect->Modifiers[0].Magnitude.SetValue(4.f);
-		BaseModEffect->Modifiers[0].ModifierType = EGameplayMod::IncomingGE;
-		BaseModEffect->Modifiers[0].ModifierOp = EGameplayModOp::Multiplicitive;
-		BaseModEffect->Modifiers[0].Attribute.SetUProperty(StackingProperty);
-		BaseModEffect->Duration.SetValue(1.f);
-		BaseModEffect->CopyPolicy = EGameplayEffectCopyPolicy::AlwaysLink;
-		BaseModEffect->GameplayEffectRequiredTags.AddTag(IGameplayTagsModule::RequestGameplayTag(FName(TEXT("Stack"))));
-
-		// Apply to self
-		DestComponent->ApplyGameplayEffectToTarget(BaseModEffect, DestComponent, 1.f);
-	}
-
-	{
-		ABILITY_LOG_SCOPE(TEXT("Apply StackingEffect"));
-
-		UGameplayEffect * BaseStackedEffect = Cast<UGameplayEffect>(StaticConstructObject(UGameplayEffect::StaticClass(), GetTransientPackage(), FName(TEXT("StackingEffect1"))));
-		BaseStackedEffect->Modifiers.SetNum(1);
-		BaseStackedEffect->Modifiers[0].Magnitude.SetValue(StackingValue);
-		BaseStackedEffect->Modifiers[0].ModifierType = EGameplayMod::Attribute;
-		BaseStackedEffect->Modifiers[0].ModifierOp = EGameplayModOp::Additive;
-		BaseStackedEffect->Modifiers[0].Attribute.SetUProperty(StackingProperty);
-		BaseStackedEffect->Duration.Value = UGameplayEffect::INFINITE_DURATION;
-		BaseStackedEffect->Period.Value = 1.f;
-		BaseStackedEffect->StackingPolicy = EGameplayEffectStackingPolicy::Highest;
-		BaseStackedEffect->GameplayEffectTags.AddTag(IGameplayTagsModule::RequestGameplayTag(FName(TEXT("Stack"))));
-		BaseStackedEffect->StackedAttribName = FName(*StackingProperty->GetName());
-
-		BaseStackedEffect->ValidateGameplayEffect();
-
-		SourceComponent->ApplyGameplayEffectToTarget(BaseStackedEffect, DestComponent);
-	}
-
-	{
-		ABILITY_LOG_SCOPE(TEXT("Apply StackingEffect"));
-
-		UGameplayEffect * BaseStackedEffect = Cast<UGameplayEffect>(StaticConstructObject(UGameplayEffect::StaticClass(), GetTransientPackage(), FName(TEXT("StackingEffect2"))));
-		BaseStackedEffect->Modifiers.SetNum(1);
-		BaseStackedEffect->Modifiers[0].Magnitude.SetValue(StackingValue * 2.f);
-		BaseStackedEffect->Modifiers[0].ModifierType = EGameplayMod::Attribute;
-		BaseStackedEffect->Modifiers[0].ModifierOp = EGameplayModOp::Additive;
-		BaseStackedEffect->Modifiers[0].Attribute.SetUProperty(StackingProperty);
-		BaseStackedEffect->Duration.Value = UGameplayEffect::INFINITE_DURATION;
-		BaseStackedEffect->Period.Value = 1.f;
-		BaseStackedEffect->StackingPolicy = EGameplayEffectStackingPolicy::Highest;
-		BaseStackedEffect->StackedAttribName = FName(*StackingProperty->GetName());
-
-		BaseStackedEffect->ValidateGameplayEffect();
-
-		SourceComponent->ApplyGameplayEffectToTarget(BaseStackedEffect, DestComponent);
-
-		GameplayTest_TickWorld(World, 0.0001f); // Move our Effects from the pending stack to the active stack inside the timer manager, this starts the clock for effect execution
-		GameplayTest_TickWorld(World, 0.1f); // Offset the current time from the start of the period so that floating point issues don't affect the test
-	}
-
-	{
-		// Tick once
-		GameplayTest_TickWorld(World, 1.f);
-
-		float ActualValue = DestComponent->GetSet<UAbilitySystemTestAttributeSet>()->StackingAttribute1;
-		float ExpectedValue = StackingValue * 8.f;	// 4 for the first GE, 4 for first GE during tick
-		Test->TestTrue(SKILL_TEST_TEXT("Stacking Attribute: Actual: %.2f == Exected: %.2f", ActualValue, ExpectedValue), ActualValue == ExpectedValue);
-
-		// At this point the modifier should be removed and the second GE should be the best match for the stacking rule
-		// Tick once
-		GameplayTest_TickWorld(World, 1.f);
-
-		ActualValue = DestComponent->GetSet<UAbilitySystemTestAttributeSet>()->StackingAttribute1;
-		ExpectedValue = StackingValue * 10.f;	// 4 for the first GE, 4 for first GE during first tick, 2 for the second GE during second tick
-		Test->TestTrue(SKILL_TEST_TEXT("Stacking Attribute: Actual: %.2f == Exected: %.2f", ActualValue, ExpectedValue), ActualValue == ExpectedValue);
-	}
-
-	ClearGlobalCurveTable();
-	ClearGlobalDataTable();
-
-	World->EditorDestroyActor(SourceActor, false);
-	World->EditorDestroyActor(DestActor, false);
-
-	return true;
-}
-
-// test to make sure that adding a modifier causes stacks to be recalculated correctly
-bool GameplayEffectsTest_StackingAddingModifiers(UWorld *World, FAutomationTestBase * Test)
-{
-	const float StackingValue = 5.f;
-
-	SetGlobalCurveTable();
-	SetGlobalDataTable();
-
-	AAbilitySystemTestPawn *SourceActor = World->SpawnActor<AAbilitySystemTestPawn>();
-	AAbilitySystemTestPawn *DestActor = World->SpawnActor<AAbilitySystemTestPawn>();
-
-	UProperty *StackingProperty = FindFieldChecked<UProperty>(UAbilitySystemTestAttributeSet::StaticClass(), GET_MEMBER_NAME_CHECKED(UAbilitySystemTestAttributeSet, StackingAttribute1));
-
-	UAbilitySystemComponent * SourceComponent = SourceActor->GetAbilitySystemComponent();
-	UAbilitySystemComponent * DestComponent = DestActor->GetAbilitySystemComponent();
-
-	{
-		ABILITY_LOG_SCOPE(TEXT("Apply StackingEffect"));
-
-		UGameplayEffect * BaseStackedEffect = Cast<UGameplayEffect>(StaticConstructObject(UGameplayEffect::StaticClass(), GetTransientPackage(), FName(TEXT("StackingEffect1"))));
-		BaseStackedEffect->Modifiers.SetNum(1);
-		BaseStackedEffect->Modifiers[0].Magnitude.SetValue(StackingValue);
-		BaseStackedEffect->Modifiers[0].ModifierType = EGameplayMod::Attribute;
-		BaseStackedEffect->Modifiers[0].ModifierOp = EGameplayModOp::Additive;
-		BaseStackedEffect->Modifiers[0].Attribute.SetUProperty(StackingProperty);
-		BaseStackedEffect->Duration.Value = UGameplayEffect::INFINITE_DURATION;
-		BaseStackedEffect->Period.Value = 1.f;
-		BaseStackedEffect->StackingPolicy = EGameplayEffectStackingPolicy::Highest;
-		BaseStackedEffect->GameplayEffectTags.AddTag(IGameplayTagsModule::RequestGameplayTag(FName(TEXT("Stack"))));
-		BaseStackedEffect->StackedAttribName = FName(*StackingProperty->GetName());
-
-		BaseStackedEffect->ValidateGameplayEffect();
-
-		SourceComponent->ApplyGameplayEffectToTarget(BaseStackedEffect, DestComponent);
-	}
-
-	{
-		ABILITY_LOG_SCOPE(TEXT("Apply StackingEffect"));
-
-		UGameplayEffect * BaseStackedEffect = Cast<UGameplayEffect>(StaticConstructObject(UGameplayEffect::StaticClass(), GetTransientPackage(), FName(TEXT("StackingEffect2"))));
-		BaseStackedEffect->Modifiers.SetNum(1);
-		BaseStackedEffect->Modifiers[0].Magnitude.SetValue(StackingValue * 2.f);
-		BaseStackedEffect->Modifiers[0].ModifierType = EGameplayMod::Attribute;
-		BaseStackedEffect->Modifiers[0].ModifierOp = EGameplayModOp::Additive;
-		BaseStackedEffect->Modifiers[0].Attribute.SetUProperty(StackingProperty);
-		BaseStackedEffect->Duration.Value = UGameplayEffect::INFINITE_DURATION;
-		BaseStackedEffect->Period.Value = 1.f;
-		BaseStackedEffect->StackingPolicy = EGameplayEffectStackingPolicy::Highest;
-		BaseStackedEffect->StackedAttribName = FName(*StackingProperty->GetName());
-
-		BaseStackedEffect->ValidateGameplayEffect();
-
-		SourceComponent->ApplyGameplayEffectToTarget(BaseStackedEffect, DestComponent);
-
-		GameplayTest_TickWorld(World, 0.0001f); // Move our Effects from the pending stack to the active stack inside the timer manager, this starts the clock for effect execution
-		GameplayTest_TickWorld(World, 0.1f); // Offset the current time from the start of the period so that floating point issues don't affect the test
-	}
-
-	{
-		// Tick once
-		GameplayTest_TickWorld(World, 1.f);
-
-		float ActualValue = DestComponent->GetSet<UAbilitySystemTestAttributeSet>()->StackingAttribute1;
-		float ExpectedValue = StackingValue * 4.f;	// 4 for the second GE executing twice
-		Test->TestTrue(SKILL_TEST_TEXT("Stacking Attribute: Actual: %.2f == Exected: %.2f", ActualValue, ExpectedValue), ActualValue == ExpectedValue);
-	}
-
-	// Setup a GE to modify IncomingGEs
-	{
-		ABILITY_LOG_SCOPE(TEXT("Apply modifier to incoming, tagged GEs"));
-
-		UGameplayEffect* BaseModEffect = Cast<UGameplayEffect>(StaticConstructObject(UGameplayEffect::StaticClass(), GetTransientPackage(), FName(TEXT("ProtectBuff"))));
-		BaseModEffect->Modifiers.SetNum(1);
-		BaseModEffect->Modifiers[0].Magnitude.SetValue(3.f);
-		BaseModEffect->Modifiers[0].ModifierType = EGameplayMod::ActiveGE;
-		BaseModEffect->Modifiers[0].ModifierOp = EGameplayModOp::Multiplicitive;
-		BaseModEffect->Modifiers[0].Attribute.SetUProperty(StackingProperty);
-		BaseModEffect->Duration.SetValue(UGameplayEffect::INFINITE_DURATION);
-		BaseModEffect->CopyPolicy = EGameplayEffectCopyPolicy::AlwaysLink;
-		BaseModEffect->GameplayEffectRequiredTags.AddTag(IGameplayTagsModule::RequestGameplayTag(FName(TEXT("Stack"))));
-
-		// Apply to self
-		DestComponent->ApplyGameplayEffectToTarget(BaseModEffect, DestComponent, 1.f);
-
-		GameplayTest_TickWorld(World, 0.0001f); // Move our Effects from the pending stack to the active stack inside the timer manager, this starts the clock for effect execution
-		GameplayTest_TickWorld(World, 0.1f); // Offset the current time from the start of the period so that floating point issues don't affect the test
-	}
-
-	{
-		// At this point the modifier should be removed and the second GE should be the best match for the stacking rule
-		// Tick once
-		GameplayTest_TickWorld(World, 1.f);
-
-		float ActualValue = DestComponent->GetSet<UAbilitySystemTestAttributeSet>()->StackingAttribute1;
-		float ExpectedValue = StackingValue * 7.f;	// 2 for the second GE, 2 for second GE during tick, 3 for the modified first GE during second tick
-		Test->TestTrue(SKILL_TEST_TEXT("Stacking Attribute: Actual: %.2f == Exected: %.2f", ActualValue, ExpectedValue), ActualValue == ExpectedValue);
-	}
-
-	ClearGlobalCurveTable();
-	ClearGlobalDataTable();
-
-	World->EditorDestroyActor(SourceActor, false);
-	World->EditorDestroyActor(DestActor, false);
-
-	return true;
-}
-
-
-bool GameplayEffectsTest_ImmunityIncoming(UWorld *World, FAutomationTestBase * Test)
-{
-	const float StartHealth = 100.f;
-	const float DamageValue = 5.f;
-
-	AAbilitySystemTestPawn *SourceActor = World->SpawnActor<AAbilitySystemTestPawn>();
-	AAbilitySystemTestPawn *DestActor = World->SpawnActor<AAbilitySystemTestPawn>();
-
-	UProperty *DamageProperty = FindFieldChecked<UProperty>(UAbilitySystemTestAttributeSet::StaticClass(), GET_MEMBER_NAME_CHECKED(UAbilitySystemTestAttributeSet, Damage));
-
-	UAbilitySystemComponent * SourceComponent = SourceActor->GetAbilitySystemComponent();
-	UAbilitySystemComponent * DestComponent = DestActor->GetAbilitySystemComponent();
-	SourceComponent->GetSet<UAbilitySystemTestAttributeSet>()->Health = StartHealth;
-	DestComponent->GetSet<UAbilitySystemTestAttributeSet>()->Health = StartHealth;
-
-	// Setup a GE to provide immunity from incoming GEs
-	{
-		ABILITY_LOG_SCOPE(TEXT("Apply ImmunityBuff"))
-
-		UGameplayEffect * BaseImmunityEffect = Cast<UGameplayEffect>(StaticConstructObject(UGameplayEffect::StaticClass(), GetTransientPackage(), FName(TEXT("ImmunityBuff"))));
-		BaseImmunityEffect->Duration.SetValue(UGameplayEffect::INFINITE_DURATION);
-		BaseImmunityEffect->GameplayEffectRequiredTags.AddTag(IGameplayTagsModule::RequestGameplayTag(FName(TEXT("Immune"))));
-		BaseImmunityEffect->AppliesImmunityTo = EGameplayImmunity::IncomingGE;
-
-		// Apply to self
-		DestComponent->ApplyGameplayEffectToTarget(BaseImmunityEffect, DestComponent, 1.f);
-	}
-
-	// Apply Damage
-	{
-		ABILITY_LOG_SCOPE(TEXT("Apply InstantDamage"))
-
-		UGameplayEffect * BaseDmgEffect = Cast<UGameplayEffect>(StaticConstructObject(UGameplayEffect::StaticClass(), GetTransientPackage(), FName(TEXT("Damage"))));
-		BaseDmgEffect->Modifiers.SetNum(1);
-		BaseDmgEffect->Modifiers[0].Magnitude.SetValue(DamageValue);
-		BaseDmgEffect->Modifiers[0].ModifierType = EGameplayMod::Attribute;
-		BaseDmgEffect->Modifiers[0].ModifierOp = EGameplayModOp::Additive;
-		BaseDmgEffect->Modifiers[0].Attribute.SetUProperty(DamageProperty);
-		BaseDmgEffect->Modifiers[0].OwnedTags.AddTag(IGameplayTagsModule::RequestGameplayTag(FName(TEXT("Damage.Basic"))));
-		BaseDmgEffect->Duration.Value = UGameplayEffect::INSTANT_APPLICATION;
-		BaseDmgEffect->GameplayEffectTags.AddTag(IGameplayTagsModule::RequestGameplayTag(FName(TEXT("Immune"))));
-
-		// Apply to target
-		SourceComponent->ApplyGameplayEffectToTarget(BaseDmgEffect, DestComponent, 1.f);
-
-		float ExpectedValue = StartHealth;
-		float ActualValue = DestComponent->GetSet<UAbilitySystemTestAttributeSet>()->Health;
-
-		Test->TestTrue(SKILL_TEST_TEXT("Buff Instant Damage Applied"), (ActualValue == ExpectedValue));
-		ABILITY_LOG(Log, TEXT("Final Health: %.2f"), DestComponent->GetSet<UAbilitySystemTestAttributeSet>()->Health);
-	}
-
-	World->EditorDestroyActor(SourceActor, false);
-	World->EditorDestroyActor(DestActor, false);
-
-	return true;
-}
-
-bool GameplayEffectsTest_ImmunityOutgoing(UWorld *World, FAutomationTestBase * Test)
-{
-	const float StartHealth = 100.f;
-	const float DamageValue = 5.f;
-
-	AAbilitySystemTestPawn *SourceActor = World->SpawnActor<AAbilitySystemTestPawn>();
-	AAbilitySystemTestPawn *DestActor = World->SpawnActor<AAbilitySystemTestPawn>();
-
-	UProperty *DamageProperty = FindFieldChecked<UProperty>(UAbilitySystemTestAttributeSet::StaticClass(), GET_MEMBER_NAME_CHECKED(UAbilitySystemTestAttributeSet, Damage));
-
-	UAbilitySystemComponent * SourceComponent = SourceActor->GetAbilitySystemComponent();
-	UAbilitySystemComponent * DestComponent = DestActor->GetAbilitySystemComponent();
-	SourceComponent->GetSet<UAbilitySystemTestAttributeSet>()->Health = StartHealth;
-	DestComponent->GetSet<UAbilitySystemTestAttributeSet>()->Health = StartHealth;
-
-	// Setup a GE to provide immunity from outgoing GEs
-	{
-		ABILITY_LOG_SCOPE(TEXT("Apply ImmunityBuff"))
-
-		UGameplayEffect * BaseImmunityEffect = Cast<UGameplayEffect>(StaticConstructObject(UGameplayEffect::StaticClass(), GetTransientPackage(), FName(TEXT("ImmunityBuff"))));
-		BaseImmunityEffect->Duration.SetValue(UGameplayEffect::INFINITE_DURATION);
-		BaseImmunityEffect->GameplayEffectRequiredTags.AddTag(IGameplayTagsModule::RequestGameplayTag(FName(TEXT("Immune"))));
-		BaseImmunityEffect->AppliesImmunityTo = EGameplayImmunity::OutgoingGE;
-
-		// Apply to self
-		SourceComponent->ApplyGameplayEffectToTarget(BaseImmunityEffect, SourceComponent, 1.f);
-	}
-
-	// Apply Damage
-	{
-		ABILITY_LOG_SCOPE(TEXT("Apply InstantDamage"))
-
-		UGameplayEffect * BaseDmgEffect = Cast<UGameplayEffect>(StaticConstructObject(UGameplayEffect::StaticClass(), GetTransientPackage(), FName(TEXT("Damage"))));
-		BaseDmgEffect->Modifiers.SetNum(1);
-		BaseDmgEffect->Modifiers[0].Magnitude.SetValue(DamageValue);
-		BaseDmgEffect->Modifiers[0].ModifierType = EGameplayMod::Attribute;
-		BaseDmgEffect->Modifiers[0].ModifierOp = EGameplayModOp::Additive;
-		BaseDmgEffect->Modifiers[0].Attribute.SetUProperty(DamageProperty);
-		BaseDmgEffect->Modifiers[0].OwnedTags.AddTag(IGameplayTagsModule::RequestGameplayTag(FName(TEXT("Damage.Basic"))));
-		BaseDmgEffect->Duration.Value = UGameplayEffect::INSTANT_APPLICATION;
-		BaseDmgEffect->GameplayEffectTags.AddTag(IGameplayTagsModule::RequestGameplayTag(FName(TEXT("Immune"))));
-
-		// Apply to target
-		SourceComponent->ApplyGameplayEffectToTarget(BaseDmgEffect, DestComponent, 1.f);
-
-		float ExpectedValue = StartHealth;
-		float ActualValue = DestComponent->GetSet<UAbilitySystemTestAttributeSet>()->Health;
-
-		Test->TestTrue(SKILL_TEST_TEXT("Buff Instant Damage Applied"), (ActualValue == ExpectedValue));
-		ABILITY_LOG(Log, TEXT("Final Health: %.2f"), DestComponent->GetSet<UAbilitySystemTestAttributeSet>()->Health);
-	}
-
-	World->EditorDestroyActor(SourceActor, false);
-	World->EditorDestroyActor(DestActor, false);
-
-	return true;
-}
-
-// A buff passes its tags to the base modifier. The base modifier can now be blocked by immunity
-// This would be bad to do in practice. We now have a gameplay effect that hangs around but doesn't do anything
-bool GameplayEffectsTest_ImmunityMod(UWorld *World, FAutomationTestBase * Test)
-{
-	const float StartHealth = 100.f;
-	const float DamageValue = 5.f;
-	const float BonusDamageMultiplier = 2.f;
-
-	AAbilitySystemTestPawn *SourceActor = World->SpawnActor<AAbilitySystemTestPawn>();
-	AAbilitySystemTestPawn *DestActor = World->SpawnActor<AAbilitySystemTestPawn>();
-
-	UProperty *DamageProperty = FindFieldChecked<UProperty>(UAbilitySystemTestAttributeSet::StaticClass(), GET_MEMBER_NAME_CHECKED(UAbilitySystemTestAttributeSet, Damage));
-
-	UAbilitySystemComponent * SourceComponent = SourceActor->GetAbilitySystemComponent();
-	UAbilitySystemComponent * DestComponent = DestActor->GetAbilitySystemComponent();
-	SourceComponent->GetSet<UAbilitySystemTestAttributeSet>()->Health = StartHealth;
-	DestComponent->GetSet<UAbilitySystemTestAttributeSet>()->Health = StartHealth;
-
-	// Setup a GE to provide immunity to buffed damage
-	{
-		ABILITY_LOG_SCOPE(TEXT("Apply Immunity"))
-
-		UGameplayEffect * BaseImmunityEffect = Cast<UGameplayEffect>(StaticConstructObject(UGameplayEffect::StaticClass(), GetTransientPackage(), FName(TEXT("ImmunityBuff"))));
-		BaseImmunityEffect->Duration.SetValue(UGameplayEffect::INFINITE_DURATION);
-		BaseImmunityEffect->GameplayEffectRequiredTags.AddTag(IGameplayTagsModule::RequestGameplayTag(FName(TEXT("Damage.Buff"))));
-		BaseImmunityEffect->AppliesImmunityTo = EGameplayImmunity::IncomingGE;
-
-		// Apply to self
-		DestComponent->ApplyGameplayEffectToTarget(BaseImmunityEffect, DestComponent, 1.f);
-	}
-
-	// Setup a GE to modify OutgoingGEs
-	{
-		ABILITY_LOG_SCOPE(TEXT("Apply DamageBuff"))
-
-		UGameplayEffect * BaseDmgEffect = Cast<UGameplayEffect>(StaticConstructObject(UGameplayEffect::StaticClass(), GetTransientPackage(), FName(TEXT("DamageBuff"))));
-		BaseDmgEffect->Modifiers.SetNum(1);
-		BaseDmgEffect->Modifiers[0].Magnitude.SetValue(BonusDamageMultiplier);
-		BaseDmgEffect->Modifiers[0].ModifierType = EGameplayMod::OutgoingGE;
-		BaseDmgEffect->Modifiers[0].ModifierOp = EGameplayModOp::Multiplicitive;
-		BaseDmgEffect->Modifiers[0].Attribute.SetUProperty(DamageProperty);
-		BaseDmgEffect->Modifiers[0].OwnedTags.AddTag(IGameplayTagsModule::RequestGameplayTag(FName(TEXT("Damage.Buff"))));
-		BaseDmgEffect->Duration.SetValue(UGameplayEffect::INFINITE_DURATION);
-
-		// Apply to self
-		SourceComponent->ApplyGameplayEffectToTarget(BaseDmgEffect, SourceComponent, 1.f);
-	}
-
-	// Apply Damage
-	{
-		ABILITY_LOG_SCOPE(TEXT("Apply InstantDamage"))
-
-		UGameplayEffect * BaseDmgEffect = Cast<UGameplayEffect>(StaticConstructObject(UGameplayEffect::StaticClass(), GetTransientPackage(), FName(TEXT("Damage"))));
-		BaseDmgEffect->Modifiers.SetNum(1);
-		BaseDmgEffect->Modifiers[0].Magnitude.SetValue(DamageValue);
-		BaseDmgEffect->Modifiers[0].ModifierType = EGameplayMod::Attribute;
-		BaseDmgEffect->Modifiers[0].ModifierOp = EGameplayModOp::Additive;
-		BaseDmgEffect->Modifiers[0].Attribute.SetUProperty(DamageProperty);
-		BaseDmgEffect->Modifiers[0].OwnedTags.AddTag(IGameplayTagsModule::RequestGameplayTag(FName(TEXT("Damage.Basic"))));
-		BaseDmgEffect->Duration.Value = UGameplayEffect::INSTANT_APPLICATION;
-
-		// Apply to target
-		SourceComponent->ApplyGameplayEffectToTarget(BaseDmgEffect, DestComponent, 1.f);
-
-		float ExpectedValue = StartHealth;
-		float ActualValue = DestComponent->GetSet<UAbilitySystemTestAttributeSet>()->Health;
-
-		Test->TestTrue(SKILL_TEST_TEXT("Buff Instant Damage Applied. %.2f == %.2f", ActualValue, ExpectedValue), (ActualValue == ExpectedValue));
-		ABILITY_LOG(Log, TEXT("Final Health: %.2f"), DestComponent->GetSet<UAbilitySystemTestAttributeSet>()->Health);
-	}
-
-	World->EditorDestroyActor(SourceActor, false);
-	World->EditorDestroyActor(DestActor, false);
-
-	return true;
-}
-
-// A buff is added in the form of a new gameplay effect being attached to the base gameplay effect.
-// The base effect applies but immunity stops the buff.
-bool GameplayEffectsTest_ImmunityBlockedBuff(UWorld *World, FAutomationTestBase * Test)
-{
-	const float StartHealth = 100.f;
-	const float DamageValue = 5.f;
-
-	AAbilitySystemTestPawn *SourceActor = World->SpawnActor<AAbilitySystemTestPawn>();
-	AAbilitySystemTestPawn *DestActor = World->SpawnActor<AAbilitySystemTestPawn>();
-
-	UProperty *DamageProperty = FindFieldChecked<UProperty>(UAbilitySystemTestAttributeSet::StaticClass(), GET_MEMBER_NAME_CHECKED(UAbilitySystemTestAttributeSet, Damage));
-
-	UAbilitySystemComponent * SourceComponent = SourceActor->GetAbilitySystemComponent();
-	UAbilitySystemComponent * DestComponent = DestActor->GetAbilitySystemComponent();
-	SourceComponent->GetSet<UAbilitySystemTestAttributeSet>()->Health = StartHealth;
-	DestComponent->GetSet<UAbilitySystemTestAttributeSet>()->Health = StartHealth;
-
-	// Setup a GE to provide immunity to buffed damage
-	{
-		ABILITY_LOG_SCOPE(TEXT("Apply Immunity"))
-
-		UGameplayEffect * BaseImmunityEffect = Cast<UGameplayEffect>(StaticConstructObject(UGameplayEffect::StaticClass(), GetTransientPackage(), FName(TEXT("ImmunityBuff"))));
-		BaseImmunityEffect->Duration.SetValue(UGameplayEffect::INFINITE_DURATION);
-		BaseImmunityEffect->GameplayEffectRequiredTags.AddTag(IGameplayTagsModule::RequestGameplayTag(FName(TEXT("Damage.Buff"))));
-		BaseImmunityEffect->AppliesImmunityTo = EGameplayImmunity::IncomingGE;
-
-		// Apply to self
-		DestComponent->ApplyGameplayEffectToTarget(BaseImmunityEffect, DestComponent, 1.f);
-	}
-
-	// Apply base damage and a buff that causes extra damage
-	{
-		ABILITY_LOG_SCOPE(TEXT("Apply DamageBuff and InstantDamage"))
-
-		UGameplayEffect * BuffEffect = Cast<UGameplayEffect>(StaticConstructObject(UGameplayEffect::StaticClass(), GetTransientPackage(), FName(TEXT("DamageBuff"))));
-		BuffEffect->Modifiers.SetNum(1);
-		BuffEffect->Modifiers[0].Magnitude.SetValue(DamageValue);
-		BuffEffect->Modifiers[0].ModifierType = EGameplayMod::Attribute;
-		BuffEffect->Modifiers[0].ModifierOp = EGameplayModOp::Additive;
-		BuffEffect->Modifiers[0].Attribute.SetUProperty(DamageProperty);
-		BuffEffect->GameplayEffectTags.AddTag(IGameplayTagsModule::RequestGameplayTag(FName(TEXT("Damage.Buff"))));
-		BuffEffect->Duration.SetValue(UGameplayEffect::INSTANT_APPLICATION);
-
-		UGameplayEffect * BaseDmgEffect = Cast<UGameplayEffect>(StaticConstructObject(UGameplayEffect::StaticClass(), GetTransientPackage(), FName(TEXT("Damage"))));
-		BaseDmgEffect->Modifiers.SetNum(1);
-		BaseDmgEffect->Modifiers[0].Magnitude.SetValue(DamageValue);
-		BaseDmgEffect->Modifiers[0].ModifierType = EGameplayMod::Attribute;
-		BaseDmgEffect->Modifiers[0].ModifierOp = EGameplayModOp::Additive;
-		BaseDmgEffect->Modifiers[0].Attribute.SetUProperty(DamageProperty);
-		BaseDmgEffect->GameplayEffectTags.AddTag(IGameplayTagsModule::RequestGameplayTag(FName(TEXT("Damage.Basic"))));
-		BaseDmgEffect->Duration.Value = UGameplayEffect::INSTANT_APPLICATION;
-		BaseDmgEffect->TargetEffects.Add(BuffEffect);
-
-		// Apply to target
-		SourceComponent->ApplyGameplayEffectToTarget(BaseDmgEffect, DestComponent, 1.f);
-
-		float ExpectedValue = (StartHealth - DamageValue);
-		float ActualValue = DestComponent->GetSet<UAbilitySystemTestAttributeSet>()->Health;
-
-		Test->TestTrue(SKILL_TEST_TEXT("Buff Instant Damage Applied. ActualValue: %.2f. ExpectedValue: %.2f.", ActualValue, ExpectedValue), (ActualValue == ExpectedValue));
-		ABILITY_LOG(Log, TEXT("Final Health: %.2f"), DestComponent->GetSet<UAbilitySystemTestAttributeSet>()->Health);
-	}
-
-	World->EditorDestroyActor(SourceActor, false);
-	World->EditorDestroyActor(DestActor, false);
-
-	return true;
-}
-
-// A buff is added in the form of a new gameplay effect being attached to the base gameplay effect.
-// The base effect is stopped by immunity, the buff should also be stopped despite not matching the immunity tags
-bool GameplayEffectsTest_ImmunityBlockedBaseAndBuff(UWorld *World, FAutomationTestBase * Test)
-{
-	const float StartHealth = 100.f;
-	const float DamageValue = 5.f;
-
-	AAbilitySystemTestPawn *SourceActor = World->SpawnActor<AAbilitySystemTestPawn>();
-	AAbilitySystemTestPawn *DestActor = World->SpawnActor<AAbilitySystemTestPawn>();
-
-	UProperty *DamageProperty = FindFieldChecked<UProperty>(UAbilitySystemTestAttributeSet::StaticClass(), GET_MEMBER_NAME_CHECKED(UAbilitySystemTestAttributeSet, Damage));
-
-	UAbilitySystemComponent * SourceComponent = SourceActor->GetAbilitySystemComponent();
-	UAbilitySystemComponent * DestComponent = DestActor->GetAbilitySystemComponent();
-	SourceComponent->GetSet<UAbilitySystemTestAttributeSet>()->Health = StartHealth;
-	DestComponent->GetSet<UAbilitySystemTestAttributeSet>()->Health = StartHealth;
-
-	// Setup a GE to provide immunity to buffed damage
-	{
-		ABILITY_LOG_SCOPE(TEXT("Apply Immunity"))
-
-		UGameplayEffect * BaseImmunityEffect = Cast<UGameplayEffect>(StaticConstructObject(UGameplayEffect::StaticClass(), GetTransientPackage(), FName(TEXT("ImmunityBuff"))));
-		BaseImmunityEffect->Duration.SetValue(UGameplayEffect::INFINITE_DURATION);
-		BaseImmunityEffect->GameplayEffectRequiredTags.AddTag(IGameplayTagsModule::RequestGameplayTag(FName(TEXT("Damage.Basic"))));
-		BaseImmunityEffect->AppliesImmunityTo = EGameplayImmunity::IncomingGE;
-
-		// Apply to self
-		DestComponent->ApplyGameplayEffectToTarget(BaseImmunityEffect, DestComponent, 1.f);
-	}
-
-	// Apply base damage and a buff that causes extra damage
-	{
-		ABILITY_LOG_SCOPE(TEXT("Apply DamageBuff and InstantDamage"))
-
-		UGameplayEffect * BuffEffect = Cast<UGameplayEffect>(StaticConstructObject(UGameplayEffect::StaticClass(), GetTransientPackage(), FName(TEXT("DamageBuff"))));
-		BuffEffect->Modifiers.SetNum(1);
-		BuffEffect->Modifiers[0].Magnitude.SetValue(DamageValue);
-		BuffEffect->Modifiers[0].ModifierType = EGameplayMod::Attribute;
-		BuffEffect->Modifiers[0].ModifierOp = EGameplayModOp::Additive;
-		BuffEffect->Modifiers[0].Attribute.SetUProperty(DamageProperty);
-		BuffEffect->GameplayEffectTags.AddTag(IGameplayTagsModule::RequestGameplayTag(FName(TEXT("Damage.Buff"))));
-		BuffEffect->Duration.SetValue(UGameplayEffect::INSTANT_APPLICATION);
-
-		UGameplayEffect * BaseDmgEffect = Cast<UGameplayEffect>(StaticConstructObject(UGameplayEffect::StaticClass(), GetTransientPackage(), FName(TEXT("Damage"))));
-		BaseDmgEffect->Modifiers.SetNum(1);
-		BaseDmgEffect->Modifiers[0].Magnitude.SetValue(DamageValue);
-		BaseDmgEffect->Modifiers[0].ModifierType = EGameplayMod::Attribute;
-		BaseDmgEffect->Modifiers[0].ModifierOp = EGameplayModOp::Additive;
-		BaseDmgEffect->Modifiers[0].Attribute.SetUProperty(DamageProperty);
-		BaseDmgEffect->GameplayEffectTags.AddTag(IGameplayTagsModule::RequestGameplayTag(FName(TEXT("Damage.Basic"))));
-		BaseDmgEffect->Duration.Value = UGameplayEffect::INSTANT_APPLICATION;
-		BaseDmgEffect->TargetEffects.Add(BuffEffect);
-
-		// Apply to target
-		SourceComponent->ApplyGameplayEffectToTarget(BaseDmgEffect, DestComponent, 1.f);
-
-		float ExpectedValue = StartHealth;
-		float ActualValue = DestComponent->GetSet<UAbilitySystemTestAttributeSet>()->Health;
-
-		Test->TestTrue(SKILL_TEST_TEXT("Buff Instant Damage Applied"), (ActualValue == ExpectedValue));
-		ABILITY_LOG(Log, TEXT("Final Health: %.2f"), DestComponent->GetSet<UAbilitySystemTestAttributeSet>()->Health);
-	}
-
-	World->EditorDestroyActor(SourceActor, false);
-	World->EditorDestroyActor(DestActor, false);
-
-	return true;
-}
-
-// Removes an active gameplay effect from the target
-bool GameplayEffectsTest_ImmunityActiveGE(UWorld *World, FAutomationTestBase * Test)
-{
-	const float StartHealth = 100.f;
-	const float DamageValue = 5.f;
-	const float DamageProtectionDivisor = 2.f;
-
-	AAbilitySystemTestPawn *SourceActor = World->SpawnActor<AAbilitySystemTestPawn>();
-	AAbilitySystemTestPawn *DestActor = World->SpawnActor<AAbilitySystemTestPawn>();
-
-	UProperty *DamageProperty = FindFieldChecked<UProperty>(UAbilitySystemTestAttributeSet::StaticClass(), GET_MEMBER_NAME_CHECKED(UAbilitySystemTestAttributeSet, Damage));
-
-	UAbilitySystemComponent * SourceComponent = SourceActor->GetAbilitySystemComponent();
-	UAbilitySystemComponent * DestComponent = DestActor->GetAbilitySystemComponent();
-	SourceComponent->GetSet<UAbilitySystemTestAttributeSet>()->Health = StartHealth;
-	DestComponent->GetSet<UAbilitySystemTestAttributeSet>()->Health = StartHealth;
-
-	// Apply a buff to reduce incoming damage on DestComponent
-	{
-		UGameplayEffect * BuffEffect = Cast<UGameplayEffect>(StaticConstructObject(UGameplayEffect::StaticClass(), GetTransientPackage(), FName(TEXT("DamageBuff"))));
-		BuffEffect->Modifiers.SetNum(1);
-		BuffEffect->Modifiers[0].Magnitude.SetValue(DamageProtectionDivisor);
-		BuffEffect->Modifiers[0].ModifierType = EGameplayMod::IncomingGE;
-		BuffEffect->Modifiers[0].ModifierOp = EGameplayModOp::Division;
-		BuffEffect->Modifiers[0].Attribute.SetUProperty(DamageProperty);
-		BuffEffect->GameplayEffectTags.AddTag(IGameplayTagsModule::RequestGameplayTag(FName(TEXT("Damage.Reduce"))));
-		BuffEffect->Duration.SetValue(UGameplayEffect::INFINITE_DURATION);
-
-		// Apply to self
-		DestComponent->ApplyGameplayEffectToTarget(BuffEffect, DestComponent, 1.f);
-	}
-
-	// Apply Damage to verify the buff is working
-	{
-		ABILITY_LOG_SCOPE(TEXT("Apply InstantDamage"))
-
-		UGameplayEffect * BaseDmgEffect = Cast<UGameplayEffect>(StaticConstructObject(UGameplayEffect::StaticClass(), GetTransientPackage(), FName(TEXT("Damage"))));
-		BaseDmgEffect->Modifiers.SetNum(1);
-		BaseDmgEffect->Modifiers[0].Magnitude.SetValue(DamageValue);
-		BaseDmgEffect->Modifiers[0].ModifierType = EGameplayMod::Attribute;
-		BaseDmgEffect->Modifiers[0].ModifierOp = EGameplayModOp::Additive;
-		BaseDmgEffect->Modifiers[0].Attribute.SetUProperty(DamageProperty);
-		BaseDmgEffect->Modifiers[0].OwnedTags.AddTag(IGameplayTagsModule::RequestGameplayTag(FName(TEXT("Damage.Basic"))));
-		BaseDmgEffect->Duration.Value = UGameplayEffect::INSTANT_APPLICATION;
-
-		// Apply to target
-		SourceComponent->ApplyGameplayEffectToTarget(BaseDmgEffect, DestComponent, 1.f);
-
-		float ExpectedValue = StartHealth - (DamageValue / DamageProtectionDivisor);
-		float ActualValue = DestComponent->GetSet<UAbilitySystemTestAttributeSet>()->Health;
-
-		Test->TestTrue(SKILL_TEST_TEXT("Buff Instant Damage Applied"), (ActualValue == ExpectedValue));
-		ABILITY_LOG(Log, TEXT("Final Health: %.2f"), DestComponent->GetSet<UAbilitySystemTestAttributeSet>()->Health);
-	}
-
-	// Setup a GE to remove the buff
-	{
-		ABILITY_LOG_SCOPE(TEXT("Apply Immunity"))
-
-		UGameplayEffect * BaseImmunityEffect = Cast<UGameplayEffect>(StaticConstructObject(UGameplayEffect::StaticClass(), GetTransientPackage(), FName(TEXT("ImmunityBuff"))));
-		BaseImmunityEffect->Duration.SetValue(UGameplayEffect::INSTANT_APPLICATION);
-		BaseImmunityEffect->GameplayEffectRequiredTags.AddTag(IGameplayTagsModule::RequestGameplayTag(FName(TEXT("Damage.Reduce"))));
-		BaseImmunityEffect->AppliesImmunityTo = EGameplayImmunity::ActiveGE;
-
-		// Apply
-		SourceComponent->ApplyGameplayEffectToTarget(BaseImmunityEffect, DestComponent, 1.f);
-	}
-
-	// Apply Damage
-	{
-		ABILITY_LOG_SCOPE(TEXT("Apply InstantDamage"))
-
-		UGameplayEffect * BaseDmgEffect = Cast<UGameplayEffect>(StaticConstructObject(UGameplayEffect::StaticClass(), GetTransientPackage(), FName(TEXT("Damage"))));
-		BaseDmgEffect->Modifiers.SetNum(1);
-		BaseDmgEffect->Modifiers[0].Magnitude.SetValue(DamageValue);
-		BaseDmgEffect->Modifiers[0].ModifierType = EGameplayMod::Attribute;
-		BaseDmgEffect->Modifiers[0].ModifierOp = EGameplayModOp::Additive;
-		BaseDmgEffect->Modifiers[0].Attribute.SetUProperty(DamageProperty);
-		BaseDmgEffect->Modifiers[0].OwnedTags.AddTag(IGameplayTagsModule::RequestGameplayTag(FName(TEXT("Damage.Basic"))));
-		BaseDmgEffect->Duration.Value = UGameplayEffect::INSTANT_APPLICATION;
-
-		// Apply to target
-		SourceComponent->ApplyGameplayEffectToTarget(BaseDmgEffect, DestComponent, 1.f);
-
-		float ExpectedValue = StartHealth - (DamageValue / DamageProtectionDivisor) - DamageValue;
-		float ActualValue = DestComponent->GetSet<UAbilitySystemTestAttributeSet>()->Health;
-
-		Test->TestTrue(SKILL_TEST_TEXT("Buff Instant Damage Applied"), (ActualValue == ExpectedValue));
-		ABILITY_LOG(Log, TEXT("Final Health: %.2f"), DestComponent->GetSet<UAbilitySystemTestAttributeSet>()->Health);
-	}
-
-	World->EditorDestroyActor(SourceActor, false);
-	World->EditorDestroyActor(DestActor, false);
-
-	return true;
-}
-
-// removes a modifier on an active gameplay effect from the target
-bool GameplayEffectsTest_ImmunityActiveMod(UWorld *World, FAutomationTestBase * Test)
-{
-	const float StartHealth = 100.f;
-	const float DamageValue = 5.f;
-	const float DamageProtectionDivisor = 2.f;
-
-	AAbilitySystemTestPawn *SourceActor = World->SpawnActor<AAbilitySystemTestPawn>();
-	AAbilitySystemTestPawn *DestActor = World->SpawnActor<AAbilitySystemTestPawn>();
-
-	UProperty *DamageProperty = FindFieldChecked<UProperty>(UAbilitySystemTestAttributeSet::StaticClass(), GET_MEMBER_NAME_CHECKED(UAbilitySystemTestAttributeSet, Damage));
-
-	UAbilitySystemComponent* SourceComponent = SourceActor->GetAbilitySystemComponent();
-	UAbilitySystemComponent* DestComponent = DestActor->GetAbilitySystemComponent();
-	SourceComponent->GetSet<UAbilitySystemTestAttributeSet>()->Health = StartHealth;
-	DestComponent->GetSet<UAbilitySystemTestAttributeSet>()->Health = StartHealth;
-
-	// Apply a buff to reduce incoming damage on DestComponent
-	{
-		UGameplayEffect * BuffEffect = Cast<UGameplayEffect>(StaticConstructObject(UGameplayEffect::StaticClass(), GetTransientPackage(), FName(TEXT("DamageBuff"))));
-		BuffEffect->Modifiers.SetNum(1);
-		BuffEffect->Modifiers[0].Magnitude.SetValue(DamageProtectionDivisor);
-		BuffEffect->Modifiers[0].ModifierType = EGameplayMod::IncomingGE;
-		BuffEffect->Modifiers[0].ModifierOp = EGameplayModOp::Division;
-		BuffEffect->Modifiers[0].Attribute.SetUProperty(DamageProperty);
-		BuffEffect->Modifiers[0].OwnedTags.AddTag(IGameplayTagsModule::RequestGameplayTag(FName(TEXT("Damage.Reduce"))));
-		BuffEffect->Duration.SetValue(UGameplayEffect::INFINITE_DURATION);
-
-		// Apply to self
-		DestComponent->ApplyGameplayEffectToTarget(BuffEffect, DestComponent, 1.f);
-	}
-
-	// Apply Damage to verify the buff is working
-	{
-		ABILITY_LOG_SCOPE(TEXT("Apply InstantDamage"))
-
-		UGameplayEffect * BaseDmgEffect = Cast<UGameplayEffect>(StaticConstructObject(UGameplayEffect::StaticClass(), GetTransientPackage(), FName(TEXT("Damage"))));
-		BaseDmgEffect->Modifiers.SetNum(1);
-		BaseDmgEffect->Modifiers[0].Magnitude.SetValue(DamageValue);
-		BaseDmgEffect->Modifiers[0].ModifierType = EGameplayMod::Attribute;
-		BaseDmgEffect->Modifiers[0].ModifierOp = EGameplayModOp::Additive;
-		BaseDmgEffect->Modifiers[0].Attribute.SetUProperty(DamageProperty);
-		BaseDmgEffect->Modifiers[0].OwnedTags.AddTag(IGameplayTagsModule::RequestGameplayTag(FName(TEXT("Damage.Basic"))));
-		BaseDmgEffect->Duration.Value = UGameplayEffect::INSTANT_APPLICATION;
-
-		// Apply to target
-		SourceComponent->ApplyGameplayEffectToTarget(BaseDmgEffect, DestComponent, 1.f);
-
-		float ExpectedValue = StartHealth - (DamageValue / DamageProtectionDivisor);
-		float ActualValue = DestComponent->GetSet<UAbilitySystemTestAttributeSet>()->Health;
-
-		Test->TestTrue(SKILL_TEST_TEXT("Buff Instant Damage Applied"), (ActualValue == ExpectedValue));
-		ABILITY_LOG(Log, TEXT("Final Health: %.2f"), DestComponent->GetSet<UAbilitySystemTestAttributeSet>()->Health);
-	}
-
-	// Setup a GE to remove the buff
-	{
-		ABILITY_LOG_SCOPE(TEXT("Apply Immunity"))
-
-		UGameplayEffect * BaseImmunityEffect = Cast<UGameplayEffect>(StaticConstructObject(UGameplayEffect::StaticClass(), GetTransientPackage(), FName(TEXT("ImmunityBuff"))));
-		BaseImmunityEffect->Duration.SetValue(UGameplayEffect::INSTANT_APPLICATION);
-		BaseImmunityEffect->GameplayEffectRequiredTags.AddTag(IGameplayTagsModule::RequestGameplayTag(FName(TEXT("Damage.Reduce"))));
-		BaseImmunityEffect->AppliesImmunityTo = EGameplayImmunity::ActiveGE;
-
-		// Apply
-		SourceComponent->ApplyGameplayEffectToTarget(BaseImmunityEffect, DestComponent, 1.f);
-	}
-
-	// Apply Damage
-	{
-		ABILITY_LOG_SCOPE(TEXT("Apply InstantDamage"))
-
-		UGameplayEffect * BaseDmgEffect = Cast<UGameplayEffect>(StaticConstructObject(UGameplayEffect::StaticClass(), GetTransientPackage(), FName(TEXT("Damage"))));
-		BaseDmgEffect->Modifiers.SetNum(1);
-		BaseDmgEffect->Modifiers[0].Magnitude.SetValue(DamageValue);
-		BaseDmgEffect->Modifiers[0].ModifierType = EGameplayMod::Attribute;
-		BaseDmgEffect->Modifiers[0].ModifierOp = EGameplayModOp::Additive;
-		BaseDmgEffect->Modifiers[0].Attribute.SetUProperty(DamageProperty);
-		BaseDmgEffect->Modifiers[0].OwnedTags.AddTag(IGameplayTagsModule::RequestGameplayTag(FName(TEXT("Damage.Basic"))));
-		BaseDmgEffect->Duration.Value = UGameplayEffect::INSTANT_APPLICATION;
-
-		// Apply to target
-		SourceComponent->ApplyGameplayEffectToTarget(BaseDmgEffect, DestComponent, 1.f);
-
-		float ExpectedValue = StartHealth - (DamageValue / DamageProtectionDivisor) - DamageValue;
-		float ActualValue = DestComponent->GetSet<UAbilitySystemTestAttributeSet>()->Health;
-
-		Test->TestTrue(SKILL_TEST_TEXT("Buff Instant Damage Applied"), (ActualValue == ExpectedValue));
-		ABILITY_LOG(Log, TEXT("Final Health: %.2f"), DestComponent->GetSet<UAbilitySystemTestAttributeSet>()->Health);
-	}
-
-	World->EditorDestroyActor(SourceActor, false);
-	World->EditorDestroyActor(DestActor, false);
-
-	return true;
-}
-
-
-bool GameplayEffectsTest_ChanceToApplyToTarget(UWorld *World, FAutomationTestBase * Test)
-{
-	const float StartHealth = 100.f;
-	const float DamageValue = 5.f;
-
-	AAbilitySystemTestPawn *SourceActor = World->SpawnActor<AAbilitySystemTestPawn>();
-	AAbilitySystemTestPawn *DestActor = World->SpawnActor<AAbilitySystemTestPawn>();
-
-	UProperty *DamageProperty = FindFieldChecked<UProperty>(UAbilitySystemTestAttributeSet::StaticClass(), GET_MEMBER_NAME_CHECKED(UAbilitySystemTestAttributeSet, Damage));
-
-	UAbilitySystemComponent * SourceComponent = SourceActor->GetAbilitySystemComponent();
-	UAbilitySystemComponent * DestComponent = DestActor->GetAbilitySystemComponent();
-
-	SourceComponent->GetSet<UAbilitySystemTestAttributeSet>()->Health = StartHealth;
-	DestComponent->GetSet<UAbilitySystemTestAttributeSet>()->Health = StartHealth;
-
-	{
-		ABILITY_LOG_SCOPE(TEXT("Apply Damage mod"))
-
-		UGameplayEffect * BaseDmgEffect = Cast<UGameplayEffect>(StaticConstructObject(UGameplayEffect::StaticClass(), GetTransientPackage(), FName(TEXT("Damage"))));
-		BaseDmgEffect->Modifiers.SetNum(1);
-		BaseDmgEffect->Modifiers[0].Magnitude.SetValue(DamageValue);
-		BaseDmgEffect->Modifiers[0].ModifierType = EGameplayMod::Attribute;
-		BaseDmgEffect->Modifiers[0].ModifierOp = EGameplayModOp::Additive;
-		BaseDmgEffect->Modifiers[0].Attribute.SetUProperty(DamageProperty);
-		BaseDmgEffect->Modifiers[0].OwnedTags.AddTag(IGameplayTagsModule::RequestGameplayTag(FName(TEXT("Damage.Basic"))));
-		BaseDmgEffect->Duration.SetValue(UGameplayEffect::INSTANT_APPLICATION);
-		BaseDmgEffect->ChanceToApplyToTarget.SetValue(1.f);
-
-		// Apply to target
-		SourceComponent->ApplyGameplayEffectToTarget(BaseDmgEffect, DestComponent, 1.f);
-
-		float ExpectedValue = (StartHealth - DamageValue);
-		float ActualValue = DestComponent->GetSet<UAbilitySystemTestAttributeSet>()->Health;
-
-		Test->TestTrue(SKILL_TEST_TEXT("Damage Applied. %.2f == %.2f", ActualValue, ExpectedValue), ActualValue == ExpectedValue);
-	}
-
-	{
-		ABILITY_LOG_SCOPE(TEXT("Apply Damage mod"))
-
-		UGameplayEffect * BaseDmgEffect = Cast<UGameplayEffect>(StaticConstructObject(UGameplayEffect::StaticClass(), GetTransientPackage(), FName(TEXT("Damage"))));
-		BaseDmgEffect->Modifiers.SetNum(1);
-		BaseDmgEffect->Modifiers[0].Magnitude.SetValue(DamageValue);
-		BaseDmgEffect->Modifiers[0].ModifierType = EGameplayMod::Attribute;
-		BaseDmgEffect->Modifiers[0].ModifierOp = EGameplayModOp::Additive;
-		BaseDmgEffect->Modifiers[0].Attribute.SetUProperty(DamageProperty);
-		BaseDmgEffect->Modifiers[0].OwnedTags.AddTag(IGameplayTagsModule::RequestGameplayTag(FName(TEXT("Damage.Basic"))));
-		BaseDmgEffect->Duration.SetValue(UGameplayEffect::INSTANT_APPLICATION);
-		BaseDmgEffect->ChanceToApplyToTarget.SetValue(0.f);
-
-		// Apply to target
-		FActiveGameplayEffectHandle Handle = SourceComponent->ApplyGameplayEffectToTarget(BaseDmgEffect, DestComponent, 1.f);
-
-		float ExpectedValue = (StartHealth - DamageValue);
-		float ActualValue = DestComponent->GetSet<UAbilitySystemTestAttributeSet>()->Health;
-
-		Test->TestTrue(SKILL_TEST_TEXT("Damage Applied. %.2f == %.2f", ActualValue, ExpectedValue), ActualValue == ExpectedValue);
-		Test->TestFalse(SKILL_TEST_TEXT("Effect applied to target when chance was 0.f"), Handle.IsValid());
-	}
-
-	World->EditorDestroyActor(SourceActor, false);
-	World->EditorDestroyActor(DestActor, false);
-
-	return true;
-}
-
-bool GameplayEffectsTest_ChanceToExecuteOnActiveGEMod(UWorld *World, FAutomationTestBase * Test)
-{
-	const float StartHealth = 100.f;
-	const float DamageValue = 5.f;
-	const float BonusDamageMultiplier = 2.f;
-	const float ExtraDamageMultiplier = 1.f;
-
-	AAbilitySystemTestPawn *SourceActor = World->SpawnActor<AAbilitySystemTestPawn>();
-	AAbilitySystemTestPawn *DestActor = World->SpawnActor<AAbilitySystemTestPawn>();
-
-	UProperty *DamageProperty = FindFieldChecked<UProperty>(UAbilitySystemTestAttributeSet::StaticClass(), GET_MEMBER_NAME_CHECKED(UAbilitySystemTestAttributeSet, Damage));
-
-	UAbilitySystemComponent * SourceComponent = SourceActor->GetAbilitySystemComponent();
-	UAbilitySystemComponent * DestComponent = DestActor->GetAbilitySystemComponent();
-
-	SourceComponent->GetSet<UAbilitySystemTestAttributeSet>()->Health = StartHealth;
-	DestComponent->GetSet<UAbilitySystemTestAttributeSet>()->Health = StartHealth;
-
-	// Setup a GE to modify OutgoingGEs
-	{
-		ABILITY_LOG_SCOPE(TEXT("Apply DamageBuff"))
-
-		UGameplayEffect * BuffDmgEffect = Cast<UGameplayEffect>(StaticConstructObject(UGameplayEffect::StaticClass(), GetTransientPackage(), FName(TEXT("DamageBuff"))));
-		BuffDmgEffect->Modifiers.SetNum(1);
-		BuffDmgEffect->Modifiers[0].Magnitude.SetValue(BonusDamageMultiplier);
-		BuffDmgEffect->Modifiers[0].ModifierType = EGameplayMod::OutgoingGE;
-		BuffDmgEffect->Modifiers[0].ModifierOp = EGameplayModOp::Multiplicitive;
-		BuffDmgEffect->Modifiers[0].Attribute.SetUProperty(DamageProperty);
-		BuffDmgEffect->Modifiers[0].OwnedTags.AddTag(IGameplayTagsModule::RequestGameplayTag(FName(TEXT("Damage.Buff"))));
-		BuffDmgEffect->Modifiers[0].RequiredTags.AddTag(IGameplayTagsModule::RequestGameplayTag(FName(TEXT("Damage.Basic"))));
-		BuffDmgEffect->Duration.SetValue(UGameplayEffect::INFINITE_DURATION);
-
-		// Apply to self
-		SourceComponent->ApplyGameplayEffectToTarget(BuffDmgEffect, SourceComponent, 1.f);
-	}
-
-	// attempt to modify the buff but fail because of a 0.f chance to apply
-	{
-		ABILITY_LOG_SCOPE(TEXT("Fail to modify DamageBuff"))
-			
-		UGameplayEffect * ModBuffEffect = Cast<UGameplayEffect>(StaticConstructObject(UGameplayEffect::StaticClass(), GetTransientPackage(), FName(TEXT("ModDamageBuff"))));
-		ModBuffEffect->Modifiers.SetNum(1);
-		ModBuffEffect->Modifiers[0].Magnitude.SetValue(ExtraDamageMultiplier);
-		ModBuffEffect->Modifiers[0].ModifierType = EGameplayMod::ActiveGE;
-		ModBuffEffect->Modifiers[0].ModifierOp = EGameplayModOp::Additive;
-		ModBuffEffect->Modifiers[0].Attribute.SetUProperty(DamageProperty);
-		ModBuffEffect->Duration.SetValue(UGameplayEffect::INSTANT_APPLICATION);
-		ModBuffEffect->ChanceToExecuteOnGameplayEffect.SetValue(0.f);
-
-		// Apply to self
-		SourceComponent->ApplyGameplayEffectToTarget(ModBuffEffect, SourceComponent, 1.f);
-	}
-
-	{
-		ABILITY_LOG_SCOPE(TEXT("Apply Damage mod"))
-
-		UGameplayEffect * BaseDmgEffect = Cast<UGameplayEffect>(StaticConstructObject(UGameplayEffect::StaticClass(), GetTransientPackage(), FName(TEXT("Damage"))));
-		BaseDmgEffect->Modifiers.SetNum(1);
-		BaseDmgEffect->Modifiers[0].Magnitude.SetValue(DamageValue);
-		BaseDmgEffect->Modifiers[0].ModifierType = EGameplayMod::Attribute;
-		BaseDmgEffect->Modifiers[0].ModifierOp = EGameplayModOp::Additive;
-		BaseDmgEffect->Modifiers[0].Attribute.SetUProperty(DamageProperty);
-		BaseDmgEffect->Modifiers[0].OwnedTags.AddTag(IGameplayTagsModule::RequestGameplayTag(FName(TEXT("Damage.Basic"))));
-		BaseDmgEffect->Duration.SetValue(UGameplayEffect::INSTANT_APPLICATION);
-
-		// Apply to target
-		SourceComponent->ApplyGameplayEffectToTarget(BaseDmgEffect, DestComponent, 1.f);
-
-		float ExpectedValue = StartHealth - (DamageValue * BonusDamageMultiplier);
-		float ActualValue = DestComponent->GetSet<UAbilitySystemTestAttributeSet>()->Health;
-
-		Test->TestTrue(SKILL_TEST_TEXT("Damage Applied. %.2f == %.2f", ActualValue, ExpectedValue), ActualValue == ExpectedValue);
-	}
-
-	// modify the buff
-	{
-		ABILITY_LOG_SCOPE(TEXT("Modify DamageBuff"))
-
-		UGameplayEffect * ModBuffEffect = Cast<UGameplayEffect>(StaticConstructObject(UGameplayEffect::StaticClass(), GetTransientPackage(), FName(TEXT("ModDamageBuff"))));
-		ModBuffEffect->Modifiers.SetNum(1);
-		ModBuffEffect->Modifiers[0].Magnitude.SetValue(ExtraDamageMultiplier);
-		ModBuffEffect->Modifiers[0].ModifierType = EGameplayMod::ActiveGE;
-		ModBuffEffect->Modifiers[0].ModifierOp = EGameplayModOp::Additive;
-		ModBuffEffect->Modifiers[0].Attribute.SetUProperty(DamageProperty);
-		ModBuffEffect->Duration.SetValue(UGameplayEffect::INSTANT_APPLICATION);
-		ModBuffEffect->ChanceToExecuteOnGameplayEffect.SetValue(1.f);
-
-		// Apply to self
-		SourceComponent->ApplyGameplayEffectToTarget(ModBuffEffect, SourceComponent, 1.f);
-	}
-
-	{
-		ABILITY_LOG_SCOPE(TEXT("Apply Damage mod"))
-
-		UGameplayEffect * BaseDmgEffect = Cast<UGameplayEffect>(StaticConstructObject(UGameplayEffect::StaticClass(), GetTransientPackage(), FName(TEXT("Damage"))));
-		BaseDmgEffect->Modifiers.SetNum(1);
-		BaseDmgEffect->Modifiers[0].Magnitude.SetValue(DamageValue);
-		BaseDmgEffect->Modifiers[0].ModifierType = EGameplayMod::Attribute;
-		BaseDmgEffect->Modifiers[0].ModifierOp = EGameplayModOp::Additive;
-		BaseDmgEffect->Modifiers[0].Attribute.SetUProperty(DamageProperty);
-		BaseDmgEffect->Modifiers[0].OwnedTags.AddTag(IGameplayTagsModule::RequestGameplayTag(FName(TEXT("Damage.Basic"))));
-		BaseDmgEffect->Duration.SetValue(UGameplayEffect::INSTANT_APPLICATION);
-
-		// Apply to target
-		SourceComponent->ApplyGameplayEffectToTarget(BaseDmgEffect, DestComponent, 1.f);
-
-		float ExpectedValue = StartHealth - (DamageValue * BonusDamageMultiplier) - (DamageValue * (BonusDamageMultiplier + ExtraDamageMultiplier));
-		float ActualValue = DestComponent->GetSet<UAbilitySystemTestAttributeSet>()->Health;
-
-		Test->TestTrue(SKILL_TEST_TEXT("Damage Applied. %.2f == %.2f", ActualValue, ExpectedValue), ActualValue == ExpectedValue);
-	}
-
-	World->EditorDestroyActor(SourceActor, false);
-	World->EditorDestroyActor(DestActor, false);
-
-	return true;
-}
-
-bool GameplayEffectsTest_ChanceToExecuteOnActiveGEImmunity(UWorld *World, FAutomationTestBase * Test)
-{
-	const float StartHealth = 100.f;
-	const float DamageValue = 5.f;
-	const float BonusDamageMultiplier = 2.f;
-
-	AAbilitySystemTestPawn *SourceActor = World->SpawnActor<AAbilitySystemTestPawn>();
-	AAbilitySystemTestPawn *DestActor = World->SpawnActor<AAbilitySystemTestPawn>();
-
-	UProperty *DamageProperty = FindFieldChecked<UProperty>(UAbilitySystemTestAttributeSet::StaticClass(), GET_MEMBER_NAME_CHECKED(UAbilitySystemTestAttributeSet, Damage));
-
-	UAbilitySystemComponent * SourceComponent = SourceActor->GetAbilitySystemComponent();
-	UAbilitySystemComponent * DestComponent = DestActor->GetAbilitySystemComponent();
-
-	SourceComponent->GetSet<UAbilitySystemTestAttributeSet>()->Health = StartHealth;
-	DestComponent->GetSet<UAbilitySystemTestAttributeSet>()->Health = StartHealth;
-
-	// Setup a GE to modify OutgoingGEs
-	{
-		ABILITY_LOG_SCOPE(TEXT("Apply DamageBuff"))
-
-		UGameplayEffect * BuffDmgEffect = Cast<UGameplayEffect>(StaticConstructObject(UGameplayEffect::StaticClass(), GetTransientPackage(), FName(TEXT("DamageBuff"))));
-		BuffDmgEffect->Modifiers.SetNum(1);
-		BuffDmgEffect->Modifiers[0].Magnitude.SetValue(BonusDamageMultiplier);
-		BuffDmgEffect->Modifiers[0].ModifierType = EGameplayMod::OutgoingGE;
-		BuffDmgEffect->Modifiers[0].ModifierOp = EGameplayModOp::Multiplicitive;
-		BuffDmgEffect->Modifiers[0].Attribute.SetUProperty(DamageProperty);
-		BuffDmgEffect->Modifiers[0].OwnedTags.AddTag(IGameplayTagsModule::RequestGameplayTag(FName(TEXT("Damage.Buff"))));
-		BuffDmgEffect->Modifiers[0].RequiredTags.AddTag(IGameplayTagsModule::RequestGameplayTag(FName(TEXT("Damage.Basic"))));
-		BuffDmgEffect->Duration.SetValue(UGameplayEffect::INFINITE_DURATION);
-
-		// Apply to self
-		SourceComponent->ApplyGameplayEffectToTarget(BuffDmgEffect, SourceComponent, 1.f);
-	}
-
-	// attempt to remove the buff but fail because of a 0.f chance to apply
-	{
-		ABILITY_LOG_SCOPE(TEXT("Fail to remove DamageBuff"))
-
-		UGameplayEffect * ModBuffEffect = Cast<UGameplayEffect>(StaticConstructObject(UGameplayEffect::StaticClass(), GetTransientPackage(), FName(TEXT("ModDamageBuff"))));
-		ModBuffEffect->GameplayEffectRequiredTags.AddTag(IGameplayTagsModule::RequestGameplayTag(FName(TEXT("Damage.Buff"))));
-		ModBuffEffect->AppliesImmunityTo = EGameplayImmunity::ActiveGE;
-		ModBuffEffect->Duration.SetValue(UGameplayEffect::INSTANT_APPLICATION);
-		ModBuffEffect->ChanceToExecuteOnGameplayEffect.SetValue(0.f);
-
-		// Apply to self
-		SourceComponent->ApplyGameplayEffectToTarget(ModBuffEffect, SourceComponent, 1.f);
-	}
-
-	{
-		ABILITY_LOG_SCOPE(TEXT("Apply Damage mod"))
-
-		UGameplayEffect * BaseDmgEffect = Cast<UGameplayEffect>(StaticConstructObject(UGameplayEffect::StaticClass(), GetTransientPackage(), FName(TEXT("Damage"))));
-		BaseDmgEffect->Modifiers.SetNum(1);
-		BaseDmgEffect->Modifiers[0].Magnitude.SetValue(DamageValue);
-		BaseDmgEffect->Modifiers[0].ModifierType = EGameplayMod::Attribute;
-		BaseDmgEffect->Modifiers[0].ModifierOp = EGameplayModOp::Additive;
-		BaseDmgEffect->Modifiers[0].Attribute.SetUProperty(DamageProperty);
-		BaseDmgEffect->Modifiers[0].OwnedTags.AddTag(IGameplayTagsModule::RequestGameplayTag(FName(TEXT("Damage.Basic"))));
-		BaseDmgEffect->Duration.SetValue(UGameplayEffect::INSTANT_APPLICATION);
-
-		// Apply to target
-		SourceComponent->ApplyGameplayEffectToTarget(BaseDmgEffect, DestComponent, 1.f);
-
-		float ExpectedValue = StartHealth - (DamageValue * BonusDamageMultiplier);
-		float ActualValue = DestComponent->GetSet<UAbilitySystemTestAttributeSet>()->Health;
-
-		Test->TestTrue(SKILL_TEST_TEXT("Damage Applied. %.2f == %.2f", ActualValue, ExpectedValue), ActualValue == ExpectedValue);
-	}
-
-	// modify the buff
-	{
-		ABILITY_LOG_SCOPE(TEXT("Remove DamageBuff"))
-
-		UGameplayEffect * ModBuffEffect = Cast<UGameplayEffect>(StaticConstructObject(UGameplayEffect::StaticClass(), GetTransientPackage(), FName(TEXT("ModDamageBuff"))));
-		ModBuffEffect->GameplayEffectRequiredTags.AddTag(IGameplayTagsModule::RequestGameplayTag(FName(TEXT("Damage.Buff"))));
-		ModBuffEffect->AppliesImmunityTo = EGameplayImmunity::ActiveGE;
-		ModBuffEffect->Duration.SetValue(UGameplayEffect::INSTANT_APPLICATION);
-		ModBuffEffect->ChanceToExecuteOnGameplayEffect.SetValue(1.f);
-
-		// Apply to self
-		SourceComponent->ApplyGameplayEffectToTarget(ModBuffEffect, SourceComponent, 1.f);
-	}
-
-	{
-		ABILITY_LOG_SCOPE(TEXT("Apply Damage mod"))
-
-		UGameplayEffect * BaseDmgEffect = Cast<UGameplayEffect>(StaticConstructObject(UGameplayEffect::StaticClass(), GetTransientPackage(), FName(TEXT("Damage"))));
-		BaseDmgEffect->Modifiers.SetNum(1);
-		BaseDmgEffect->Modifiers[0].Magnitude.SetValue(DamageValue);
-		BaseDmgEffect->Modifiers[0].ModifierType = EGameplayMod::Attribute;
-		BaseDmgEffect->Modifiers[0].ModifierOp = EGameplayModOp::Additive;
-		BaseDmgEffect->Modifiers[0].Attribute.SetUProperty(DamageProperty);
-		BaseDmgEffect->Modifiers[0].OwnedTags.AddTag(IGameplayTagsModule::RequestGameplayTag(FName(TEXT("Damage.Basic"))));
-		BaseDmgEffect->Duration.SetValue(UGameplayEffect::INSTANT_APPLICATION);
-
-		// Apply to target
-		SourceComponent->ApplyGameplayEffectToTarget(BaseDmgEffect, DestComponent, 1.f);
-
-		float ExpectedValue = StartHealth - (DamageValue * BonusDamageMultiplier) - (DamageValue);
-		float ActualValue = DestComponent->GetSet<UAbilitySystemTestAttributeSet>()->Health;
-
-		Test->TestTrue(SKILL_TEST_TEXT("Damage Applied. %.2f == %.2f", ActualValue, ExpectedValue), ActualValue == ExpectedValue);
-	}
-
-	World->EditorDestroyActor(SourceActor, false);
-	World->EditorDestroyActor(DestActor, false);
-
-	return true;
-}
-
-bool GameplayEffectsTest_ChanceToExecuteOnOutgoingGEMod(UWorld *World, FAutomationTestBase * Test)
-{
-	const float StartHealth = 100.f;
-	const float DamageValue = 5.f;
-	const float BonusDamageMultiplier = 2.f;
-
-	AAbilitySystemTestPawn *SourceActor = World->SpawnActor<AAbilitySystemTestPawn>();
-	AAbilitySystemTestPawn *DestActor = World->SpawnActor<AAbilitySystemTestPawn>();
-
-	UProperty *DamageProperty = FindFieldChecked<UProperty>(UAbilitySystemTestAttributeSet::StaticClass(), GET_MEMBER_NAME_CHECKED(UAbilitySystemTestAttributeSet, Damage));
-	
-	UAbilitySystemComponent * SourceComponent = SourceActor->GetAbilitySystemComponent();
-	UAbilitySystemComponent * DestComponent = DestActor->GetAbilitySystemComponent();
-	SourceComponent->GetSet<UAbilitySystemTestAttributeSet>()->Health = StartHealth;
-	DestComponent->GetSet<UAbilitySystemTestAttributeSet>()->Health = StartHealth;
-
-	// Setup a GE to modify OutgoingGEs
-	// this GE won't do anything because the chance to execute is zero
-	{
-		ABILITY_LOG_SCOPE(TEXT("Apply useless DamageBuff"))
-
-		UGameplayEffect * BuffDmgEffect = Cast<UGameplayEffect>(StaticConstructObject(UGameplayEffect::StaticClass(), GetTransientPackage(), FName(TEXT("DamageBuff"))));
-		BuffDmgEffect->Modifiers.SetNum(1);
-		BuffDmgEffect->Modifiers[0].Magnitude.SetValue(BonusDamageMultiplier);
-		BuffDmgEffect->Modifiers[0].ModifierType = EGameplayMod::OutgoingGE;
-		BuffDmgEffect->Modifiers[0].ModifierOp = EGameplayModOp::Multiplicitive;
-		BuffDmgEffect->Modifiers[0].Attribute.SetUProperty(DamageProperty);
-		BuffDmgEffect->Modifiers[0].OwnedTags.AddTag(IGameplayTagsModule::RequestGameplayTag(FName(TEXT("Damage.Buff"))));
-		BuffDmgEffect->Modifiers[0].RequiredTags.AddTag(IGameplayTagsModule::RequestGameplayTag(FName(TEXT("Damage.Basic"))));
-		BuffDmgEffect->Duration.SetValue(UGameplayEffect::INFINITE_DURATION);
-		BuffDmgEffect->ChanceToExecuteOnGameplayEffect.SetValue(0.f);
-
-		// Apply to self
-		SourceComponent->ApplyGameplayEffectToTarget(BuffDmgEffect, SourceComponent, 1.f);
-	}
-
-	{
-		ABILITY_LOG_SCOPE(TEXT("Apply Damage mod"))
-
-		UGameplayEffect * BaseDmgEffect = Cast<UGameplayEffect>(StaticConstructObject(UGameplayEffect::StaticClass(), GetTransientPackage(), FName(TEXT("Damage"))));
-		BaseDmgEffect->Modifiers.SetNum(1);
-		BaseDmgEffect->Modifiers[0].Magnitude.SetValue(DamageValue);
-		BaseDmgEffect->Modifiers[0].ModifierType = EGameplayMod::Attribute;
-		BaseDmgEffect->Modifiers[0].ModifierOp = EGameplayModOp::Additive;
-		BaseDmgEffect->Modifiers[0].Attribute.SetUProperty(DamageProperty);
-		BaseDmgEffect->Modifiers[0].OwnedTags.AddTag(IGameplayTagsModule::RequestGameplayTag(FName(TEXT("Damage.Basic"))));
-		BaseDmgEffect->Duration.SetValue(UGameplayEffect::INSTANT_APPLICATION);
-
-		// Apply to target
-		SourceComponent->ApplyGameplayEffectToTarget(BaseDmgEffect, DestComponent, 1.f);
-
-		float ExpectedValue = StartHealth - DamageValue;
-		float ActualValue = DestComponent->GetSet<UAbilitySystemTestAttributeSet>()->Health;
-
-		Test->TestTrue(SKILL_TEST_TEXT("Damage Applied. %.2f == %.2f", ActualValue, ExpectedValue), ActualValue == ExpectedValue);
-	}
-
-	// Setup a GE to modify OutgoingGEs
-	// this GE will always execute work because the chance to execute is one
-	{
-		ABILITY_LOG_SCOPE(TEXT("Apply DamageBuff"))
-
-		UGameplayEffect * BuffDmgEffect = Cast<UGameplayEffect>(StaticConstructObject(UGameplayEffect::StaticClass(), GetTransientPackage(), FName(TEXT("DamageBuff"))));
-		BuffDmgEffect->Modifiers.SetNum(1);
-		BuffDmgEffect->Modifiers[0].Magnitude.SetValue(BonusDamageMultiplier);
-		BuffDmgEffect->Modifiers[0].ModifierType = EGameplayMod::OutgoingGE;
-		BuffDmgEffect->Modifiers[0].ModifierOp = EGameplayModOp::Multiplicitive;
-		BuffDmgEffect->Modifiers[0].Attribute.SetUProperty(DamageProperty);
-		BuffDmgEffect->Modifiers[0].OwnedTags.AddTag(IGameplayTagsModule::RequestGameplayTag(FName(TEXT("Damage.Buff"))));
-		BuffDmgEffect->Modifiers[0].RequiredTags.AddTag(IGameplayTagsModule::RequestGameplayTag(FName(TEXT("Damage.Basic"))));
-		BuffDmgEffect->Duration.SetValue(UGameplayEffect::INFINITE_DURATION);
-		BuffDmgEffect->ChanceToExecuteOnGameplayEffect.SetValue(1.f);
-
-		// Apply to self
-		SourceComponent->ApplyGameplayEffectToTarget(BuffDmgEffect, SourceComponent, 1.f);
-	}
-
-	{
-		ABILITY_LOG_SCOPE(TEXT("Apply Damage mod"))
-
-		UGameplayEffect * BaseDmgEffect = Cast<UGameplayEffect>(StaticConstructObject(UGameplayEffect::StaticClass(), GetTransientPackage(), FName(TEXT("Damage"))));
-		BaseDmgEffect->Modifiers.SetNum(1);
-		BaseDmgEffect->Modifiers[0].Magnitude.SetValue(DamageValue);
-		BaseDmgEffect->Modifiers[0].ModifierType = EGameplayMod::Attribute;
-		BaseDmgEffect->Modifiers[0].ModifierOp = EGameplayModOp::Additive;
-		BaseDmgEffect->Modifiers[0].Attribute.SetUProperty(DamageProperty);
-		BaseDmgEffect->Modifiers[0].OwnedTags.AddTag(IGameplayTagsModule::RequestGameplayTag(FName(TEXT("Damage.Basic"))));
-		BaseDmgEffect->Duration.SetValue(UGameplayEffect::INSTANT_APPLICATION);
-
-		// Apply to target
-		SourceComponent->ApplyGameplayEffectToTarget(BaseDmgEffect, DestComponent, 1.f);
-
-		float ExpectedValue = StartHealth - DamageValue - (DamageValue * BonusDamageMultiplier);
-		float ActualValue = DestComponent->GetSet<UAbilitySystemTestAttributeSet>()->Health;
-
-		Test->TestTrue(SKILL_TEST_TEXT("Damage Applied. %.2f == %.2f", ActualValue, ExpectedValue), ActualValue == ExpectedValue);
-	}
-
-	World->EditorDestroyActor(SourceActor, false);
-	World->EditorDestroyActor(DestActor, false);
-
-	return true;
-}
-
-bool GameplayEffectsTest_ChanceToExecuteOnOutgoingGEImmunity(UWorld *World, FAutomationTestBase * Test)
-{
-	const float StartHealth = 100.f;
-	const float DamageValue = 5.f;
-
-	AAbilitySystemTestPawn *SourceActor = World->SpawnActor<AAbilitySystemTestPawn>();
-	AAbilitySystemTestPawn *DestActor = World->SpawnActor<AAbilitySystemTestPawn>();
-
-	UProperty *DamageProperty = FindFieldChecked<UProperty>(UAbilitySystemTestAttributeSet::StaticClass(), GET_MEMBER_NAME_CHECKED(UAbilitySystemTestAttributeSet, Damage));
-
-	UAbilitySystemComponent * SourceComponent = SourceActor->GetAbilitySystemComponent();
-	UAbilitySystemComponent * DestComponent = DestActor->GetAbilitySystemComponent();
-	SourceComponent->GetSet<UAbilitySystemTestAttributeSet>()->Health = StartHealth;
-	DestComponent->GetSet<UAbilitySystemTestAttributeSet>()->Health = StartHealth;
-
-
-	// attempt to prevent outgoing damage but fail because of a 0.f chance to apply
-	{
-		ABILITY_LOG_SCOPE(TEXT("Fail to remove DamageBuff"))
-
-		UGameplayEffect * ModBuffEffect = Cast<UGameplayEffect>(StaticConstructObject(UGameplayEffect::StaticClass(), GetTransientPackage(), FName(TEXT("ModDamageBuff"))));
-		ModBuffEffect->GameplayEffectRequiredTags.AddTag(IGameplayTagsModule::RequestGameplayTag(FName(TEXT("Damage.Basic"))));
-		ModBuffEffect->AppliesImmunityTo = EGameplayImmunity::OutgoingGE;
-		ModBuffEffect->Duration.SetValue(UGameplayEffect::INFINITE_DURATION);
-		ModBuffEffect->ChanceToExecuteOnGameplayEffect.SetValue(0.f);
-
-		// Apply to self
-		SourceComponent->ApplyGameplayEffectToTarget(ModBuffEffect, SourceComponent, 1.f);
-	}
-
-	{
-		ABILITY_LOG_SCOPE(TEXT("Apply Damage mod"))
-
-		UGameplayEffect * BaseDmgEffect = Cast<UGameplayEffect>(StaticConstructObject(UGameplayEffect::StaticClass(), GetTransientPackage(), FName(TEXT("Damage"))));
-		BaseDmgEffect->Modifiers.SetNum(1);
-		BaseDmgEffect->Modifiers[0].Magnitude.SetValue(DamageValue);
-		BaseDmgEffect->Modifiers[0].ModifierType = EGameplayMod::Attribute;
-		BaseDmgEffect->Modifiers[0].ModifierOp = EGameplayModOp::Additive;
-		BaseDmgEffect->Modifiers[0].Attribute.SetUProperty(DamageProperty);
-		BaseDmgEffect->Modifiers[0].OwnedTags.AddTag(IGameplayTagsModule::RequestGameplayTag(FName(TEXT("Damage.Basic"))));
-		BaseDmgEffect->Duration.SetValue(UGameplayEffect::INSTANT_APPLICATION);
-
-		// Apply to target
-		SourceComponent->ApplyGameplayEffectToTarget(BaseDmgEffect, DestComponent, 1.f);
-
-		float ExpectedValue = StartHealth - DamageValue;
-		float ActualValue = DestComponent->GetSet<UAbilitySystemTestAttributeSet>()->Health;
-
-		Test->TestTrue(SKILL_TEST_TEXT("Damage Applied. %.2f == %.2f", ActualValue, ExpectedValue), ActualValue == ExpectedValue);
-	}
-
-	// prevent outgoing damage buff
-	{
-		ABILITY_LOG_SCOPE(TEXT("Fail to remove DamageBuff"))
-
-		UGameplayEffect * ModBuffEffect = Cast<UGameplayEffect>(StaticConstructObject(UGameplayEffect::StaticClass(), GetTransientPackage(), FName(TEXT("ModDamageBuff"))));
-		ModBuffEffect->GameplayEffectRequiredTags.AddTag(IGameplayTagsModule::RequestGameplayTag(FName(TEXT("Damage.Basic"))));
-		ModBuffEffect->AppliesImmunityTo = EGameplayImmunity::OutgoingGE;
-		ModBuffEffect->Duration.SetValue(UGameplayEffect::INFINITE_DURATION);
-		ModBuffEffect->ChanceToExecuteOnGameplayEffect.SetValue(1.f);
-
-		// Apply to self
-		SourceComponent->ApplyGameplayEffectToTarget(ModBuffEffect, SourceComponent, 1.f);
-	}
-
-	{
-		ABILITY_LOG_SCOPE(TEXT("Apply Damage mod"))
-
-		UGameplayEffect * BaseDmgEffect = Cast<UGameplayEffect>(StaticConstructObject(UGameplayEffect::StaticClass(), GetTransientPackage(), FName(TEXT("Damage"))));
-		BaseDmgEffect->Modifiers.SetNum(1);
-		BaseDmgEffect->Modifiers[0].Magnitude.SetValue(DamageValue);
-		BaseDmgEffect->Modifiers[0].ModifierType = EGameplayMod::Attribute;
-		BaseDmgEffect->Modifiers[0].ModifierOp = EGameplayModOp::Additive;
-		BaseDmgEffect->Modifiers[0].Attribute.SetUProperty(DamageProperty);
-		BaseDmgEffect->Modifiers[0].OwnedTags.AddTag(IGameplayTagsModule::RequestGameplayTag(FName(TEXT("Damage.Basic"))));
-		BaseDmgEffect->Duration.SetValue(UGameplayEffect::INSTANT_APPLICATION);
-
-		// Apply to target
-		SourceComponent->ApplyGameplayEffectToTarget(BaseDmgEffect, DestComponent, 1.f);
-
-		float ExpectedValue = StartHealth - DamageValue;
-		float ActualValue = DestComponent->GetSet<UAbilitySystemTestAttributeSet>()->Health;
-
-		Test->TestTrue(SKILL_TEST_TEXT("Damage Applied. %.2f == %.2f", ActualValue, ExpectedValue), ActualValue == ExpectedValue);
-	}
-
-	World->EditorDestroyActor(SourceActor, false);
-	World->EditorDestroyActor(DestActor, false);
-
-	return true;
-}
-
-bool GameplayEffectsTest_ChanceToExecuteOnIncomingGEMod(UWorld *World, FAutomationTestBase * Test)
+static UDataTable* CreateGameplayDataTable()
 {
-	const float StartHealth = 100.f;
-	const float DamageValue = 5.f;
-	const float BonusDamageMultiplier = 2.f;
-
-	AAbilitySystemTestPawn *SourceActor = World->SpawnActor<AAbilitySystemTestPawn>();
-	AAbilitySystemTestPawn *DestActor = World->SpawnActor<AAbilitySystemTestPawn>();
-
-	UProperty *DamageProperty = FindFieldChecked<UProperty>(UAbilitySystemTestAttributeSet::StaticClass(), GET_MEMBER_NAME_CHECKED(UAbilitySystemTestAttributeSet, Damage));
-
-	UAbilitySystemComponent * SourceComponent = SourceActor->GetAbilitySystemComponent();
-	UAbilitySystemComponent * DestComponent = DestActor->GetAbilitySystemComponent();
-	SourceComponent->GetSet<UAbilitySystemTestAttributeSet>()->Health = StartHealth;
-	DestComponent->GetSet<UAbilitySystemTestAttributeSet>()->Health = StartHealth;
-
-	// Setup a GE to modify OutgoingGEs
-	// this GE won't do anything because the chance to execute is zero
-	{
-		ABILITY_LOG_SCOPE(TEXT("Apply useless DamageBuff"))
-
-		UGameplayEffect * BuffDmgEffect = Cast<UGameplayEffect>(StaticConstructObject(UGameplayEffect::StaticClass(), GetTransientPackage(), FName(TEXT("DamageBuff"))));
-		BuffDmgEffect->Modifiers.SetNum(1);
-		BuffDmgEffect->Modifiers[0].Magnitude.SetValue(BonusDamageMultiplier);
-		BuffDmgEffect->Modifiers[0].ModifierType = EGameplayMod::IncomingGE;
-		BuffDmgEffect->Modifiers[0].ModifierOp = EGameplayModOp::Multiplicitive;
-		BuffDmgEffect->Modifiers[0].Attribute.SetUProperty(DamageProperty);
-		BuffDmgEffect->Modifiers[0].OwnedTags.AddTag(IGameplayTagsModule::RequestGameplayTag(FName(TEXT("Damage.Buff"))));
-		BuffDmgEffect->Modifiers[0].RequiredTags.AddTag(IGameplayTagsModule::RequestGameplayTag(FName(TEXT("Damage.Basic"))));
-		BuffDmgEffect->Duration.SetValue(UGameplayEffect::INFINITE_DURATION);
-		BuffDmgEffect->ChanceToExecuteOnGameplayEffect.SetValue(0.f);
-
-		// Apply to self
-		SourceComponent->ApplyGameplayEffectToTarget(BuffDmgEffect, DestComponent, 1.f);
-	}
-
-	{
-		ABILITY_LOG_SCOPE(TEXT("Apply Damage mod"))
-
-		UGameplayEffect * BaseDmgEffect = Cast<UGameplayEffect>(StaticConstructObject(UGameplayEffect::StaticClass(), GetTransientPackage(), FName(TEXT("Damage"))));
-		BaseDmgEffect->Modifiers.SetNum(1);
-		BaseDmgEffect->Modifiers[0].Magnitude.SetValue(DamageValue);
-		BaseDmgEffect->Modifiers[0].ModifierType = EGameplayMod::Attribute;
-		BaseDmgEffect->Modifiers[0].ModifierOp = EGameplayModOp::Additive;
-		BaseDmgEffect->Modifiers[0].Attribute.SetUProperty(DamageProperty);
-		BaseDmgEffect->Modifiers[0].OwnedTags.AddTag(IGameplayTagsModule::RequestGameplayTag(FName(TEXT("Damage.Basic"))));
-		BaseDmgEffect->Duration.SetValue(UGameplayEffect::INSTANT_APPLICATION);
-
-		// Apply to target
-		SourceComponent->ApplyGameplayEffectToTarget(BaseDmgEffect, DestComponent, 1.f);
-
-		float ExpectedValue = StartHealth - DamageValue;
-		float ActualValue = DestComponent->GetSet<UAbilitySystemTestAttributeSet>()->Health;
-
-		Test->TestTrue(SKILL_TEST_TEXT("Damage Applied. %.2f == %.2f", ActualValue, ExpectedValue), ActualValue == ExpectedValue);
-	}
-
-	// Setup a GE to modify OutgoingGEs
-	// this GE will always execute work because the chance to execute is one
-	{
-		ABILITY_LOG_SCOPE(TEXT("Apply DamageBuff"))
-
-		UGameplayEffect * BuffDmgEffect = Cast<UGameplayEffect>(StaticConstructObject(UGameplayEffect::StaticClass(), GetTransientPackage(), FName(TEXT("DamageBuff"))));
-		BuffDmgEffect->Modifiers.SetNum(1);
-		BuffDmgEffect->Modifiers[0].Magnitude.SetValue(BonusDamageMultiplier);
-		BuffDmgEffect->Modifiers[0].ModifierType = EGameplayMod::IncomingGE;
-		BuffDmgEffect->Modifiers[0].ModifierOp = EGameplayModOp::Multiplicitive;
-		BuffDmgEffect->Modifiers[0].Attribute.SetUProperty(DamageProperty);
-		BuffDmgEffect->Modifiers[0].OwnedTags.AddTag(IGameplayTagsModule::RequestGameplayTag(FName(TEXT("Damage.Buff"))));
-		BuffDmgEffect->Modifiers[0].RequiredTags.AddTag(IGameplayTagsModule::RequestGameplayTag(FName(TEXT("Damage.Basic"))));
-		BuffDmgEffect->Duration.SetValue(UGameplayEffect::INFINITE_DURATION);
-		BuffDmgEffect->ChanceToExecuteOnGameplayEffect.SetValue(1.f);
+	FString CSV(TEXT(",Tag,CategoryText,"));
+	CSV.Append(TEXT("\r\n0,Damage"));
+	CSV.Append(TEXT("\r\n1,Damage.Basic"));
+	CSV.Append(TEXT("\r\n2,Damage.Type1"));
+	CSV.Append(TEXT("\r\n3,Damage.Type2"));
+	CSV.Append(TEXT("\r\n4,Damage.Reduce"));
+	CSV.Append(TEXT("\r\n5,Damage.Buffable"));
+	CSV.Append(TEXT("\r\n6,Damage.Buff"));
+	CSV.Append(TEXT("\r\n7,Damage.Physical"));
+	CSV.Append(TEXT("\r\n8,Damage.Fire"));
+	CSV.Append(TEXT("\r\n9,Damage.Buffed.FireBuff"));
+	CSV.Append(TEXT("\r\n10,Damage.Mitigated.Armor"));
+	CSV.Append(TEXT("\r\n11,Lifesteal"));
+	CSV.Append(TEXT("\r\n12,Shield"));
+	CSV.Append(TEXT("\r\n13,Buff"));
+	CSV.Append(TEXT("\r\n14,Immune"));
+	CSV.Append(TEXT("\r\n15,FireDamage"));
+	CSV.Append(TEXT("\r\n16,ShieldAbsorb"));
+	CSV.Append(TEXT("\r\n17,Stackable"));
+	CSV.Append(TEXT("\r\n18,Stack"));
+	CSV.Append(TEXT("\r\n19,Stack.CappedNumber"));
+	CSV.Append(TEXT("\r\n20,Stack.DiminishingReturns"));
+	CSV.Append(TEXT("\r\n21,Protect.Damage"));
+	CSV.Append(TEXT("\r\n22,SpellDmg.Buff"));
+	CSV.Append(TEXT("\r\n23,GameplayCue.Burning"));
 
-		// Apply to self
-		SourceComponent->ApplyGameplayEffectToTarget(BuffDmgEffect, DestComponent, 1.f);
-	}
+	UDataTable * DataTable = Cast<UDataTable>(StaticConstructObject(UDataTable::StaticClass(), GetTransientPackage(), FName(TEXT("TempDataTable"))));
+	DataTable->RowStruct = FGameplayTagTableRow::StaticStruct();
+	DataTable->CreateTableFromCSVString(CSV);
 
+	FGameplayTagTableRow * Row = (FGameplayTagTableRow*)DataTable->RowMap["0"];
+	if (Row)
 	{
-		ABILITY_LOG_SCOPE(TEXT("Apply Damage mod"))
-
-		UGameplayEffect * BaseDmgEffect = Cast<UGameplayEffect>(StaticConstructObject(UGameplayEffect::StaticClass(), GetTransientPackage(), FName(TEXT("Damage"))));
-		BaseDmgEffect->Modifiers.SetNum(1);
-		BaseDmgEffect->Modifiers[0].Magnitude.SetValue(DamageValue);
-		BaseDmgEffect->Modifiers[0].ModifierType = EGameplayMod::Attribute;
-		BaseDmgEffect->Modifiers[0].ModifierOp = EGameplayModOp::Additive;
-		BaseDmgEffect->Modifiers[0].Attribute.SetUProperty(DamageProperty);
-		BaseDmgEffect->Modifiers[0].OwnedTags.AddTag(IGameplayTagsModule::RequestGameplayTag(FName(TEXT("Damage.Basic"))));
-		BaseDmgEffect->Duration.SetValue(UGameplayEffect::INSTANT_APPLICATION);
-
-		// Apply to target
-		SourceComponent->ApplyGameplayEffectToTarget(BaseDmgEffect, DestComponent, 1.f);
-
-		float ExpectedValue = StartHealth - DamageValue - (DamageValue * BonusDamageMultiplier);
-		float ActualValue = DestComponent->GetSet<UAbilitySystemTestAttributeSet>()->Health;
-
-		Test->TestTrue(SKILL_TEST_TEXT("Damage Applied. %.2f == %.2f", ActualValue, ExpectedValue), ActualValue == ExpectedValue);
+		check(Row->Tag == TEXT("Damage"));
 	}
-
-	World->EditorDestroyActor(SourceActor, false);
-	World->EditorDestroyActor(DestActor, false);
-
-	return true;
+	return DataTable;
 }
 
-bool GameplayEffectsTest_ChanceToExecuteOnIncomingGEImmunity(UWorld *World, FAutomationTestBase * Test)
+bool GameplayEffectsTest_InstantDamage(UWorld *World, FAutomationTestBase * Test)
 {
 	const float StartHealth = 100.f;
-	const float DamageValue = 5.f;
+	const float DamageValue = -5.f;
 
 	AAbilitySystemTestPawn *SourceActor = World->SpawnActor<AAbilitySystemTestPawn>();
 	AAbilitySystemTestPawn *DestActor = World->SpawnActor<AAbilitySystemTestPawn>();
 
-	UProperty *DamageProperty = FindFieldChecked<UProperty>(UAbilitySystemTestAttributeSet::StaticClass(), GET_MEMBER_NAME_CHECKED(UAbilitySystemTestAttributeSet, Damage));
+	UProperty *HealthProperty = FindFieldChecked<UProperty>(UAbilitySystemTestAttributeSet::StaticClass(), GET_MEMBER_NAME_CHECKED(UAbilitySystemTestAttributeSet, Health));
 
 	UAbilitySystemComponent * SourceComponent = SourceActor->GetAbilitySystemComponent();
 	UAbilitySystemComponent * DestComponent = DestActor->GetAbilitySystemComponent();
 	SourceComponent->GetSet<UAbilitySystemTestAttributeSet>()->Health = StartHealth;
 	DestComponent->GetSet<UAbilitySystemTestAttributeSet>()->Health = StartHealth;
 
-
-	// attempt to prevent outgoing damage but fail because of a 0.f chance to apply
-	{
-		ABILITY_LOG_SCOPE(TEXT("Fail to remove DamageBuff"))
-
-		UGameplayEffect * ModBuffEffect = Cast<UGameplayEffect>(StaticConstructObject(UGameplayEffect::StaticClass(), GetTransientPackage(), FName(TEXT("ModDamageBuff"))));
-		ModBuffEffect->GameplayEffectRequiredTags.AddTag(IGameplayTagsModule::RequestGameplayTag(FName(TEXT("Damage.Basic"))));
-		ModBuffEffect->AppliesImmunityTo = EGameplayImmunity::IncomingGE;
-		ModBuffEffect->Duration.SetValue(UGameplayEffect::INFINITE_DURATION);
-		ModBuffEffect->ChanceToExecuteOnGameplayEffect.SetValue(0.f);
-
-		// Apply to self
-		SourceComponent->ApplyGameplayEffectToTarget(ModBuffEffect, DestComponent, 1.f);
-	}
-
-	{
-		ABILITY_LOG_SCOPE(TEXT("Apply Damage mod"))
-
-		UGameplayEffect * BaseDmgEffect = Cast<UGameplayEffect>(StaticConstructObject(UGameplayEffect::StaticClass(), GetTransientPackage(), FName(TEXT("Damage"))));
-		BaseDmgEffect->Modifiers.SetNum(1);
-		BaseDmgEffect->Modifiers[0].Magnitude.SetValue(DamageValue);
-		BaseDmgEffect->Modifiers[0].ModifierType = EGameplayMod::Attribute;
-		BaseDmgEffect->Modifiers[0].ModifierOp = EGameplayModOp::Additive;
-		BaseDmgEffect->Modifiers[0].Attribute.SetUProperty(DamageProperty);
-		BaseDmgEffect->Modifiers[0].OwnedTags.AddTag(IGameplayTagsModule::RequestGameplayTag(FName(TEXT("Damage.Basic"))));
-		BaseDmgEffect->Duration.SetValue(UGameplayEffect::INSTANT_APPLICATION);
-
-		// Apply to target
-		SourceComponent->ApplyGameplayEffectToTarget(BaseDmgEffect, DestComponent, 1.f);
-
-		float ExpectedValue = StartHealth - DamageValue;
-		float ActualValue = DestComponent->GetSet<UAbilitySystemTestAttributeSet>()->Health;
-
-		Test->TestTrue(SKILL_TEST_TEXT("Damage Applied. %.2f == %.2f", ActualValue, ExpectedValue), ActualValue == ExpectedValue);
-	}
-
-	// prevent outgoing damage buff
-	{
-		ABILITY_LOG_SCOPE(TEXT("Fail to remove DamageBuff"))
-
-		UGameplayEffect * ModBuffEffect = Cast<UGameplayEffect>(StaticConstructObject(UGameplayEffect::StaticClass(), GetTransientPackage(), FName(TEXT("ModDamageBuff"))));
-		ModBuffEffect->GameplayEffectRequiredTags.AddTag(IGameplayTagsModule::RequestGameplayTag(FName(TEXT("Damage.Basic"))));
-		ModBuffEffect->AppliesImmunityTo = EGameplayImmunity::IncomingGE;
-		ModBuffEffect->Duration.SetValue(UGameplayEffect::INFINITE_DURATION);
-		ModBuffEffect->ChanceToExecuteOnGameplayEffect.SetValue(1.f);
-
-		// Apply to self
-		SourceComponent->ApplyGameplayEffectToTarget(ModBuffEffect, DestComponent, 1.f);
-	}
-
 	{
-		ABILITY_LOG_SCOPE(TEXT("Apply Damage mod"))
+		ABILITY_LOG_SCOPE(TEXT("Apply InstantDamage"))
 
-		UGameplayEffect * BaseDmgEffect = Cast<UGameplayEffect>(StaticConstructObject(UGameplayEffect::StaticClass(), GetTransientPackage(), FName(TEXT("Damage"))));
+		UGameplayEffect * BaseDmgEffect = Cast<UGameplayEffect>(StaticConstructObject(UGameplayEffect::StaticClass(), GetTransientPackage(), FName(TEXT("BaseDmgEffect"))));
 		BaseDmgEffect->Modifiers.SetNum(1);
 		BaseDmgEffect->Modifiers[0].Magnitude.SetValue(DamageValue);
-		BaseDmgEffect->Modifiers[0].ModifierType = EGameplayMod::Attribute;
 		BaseDmgEffect->Modifiers[0].ModifierOp = EGameplayModOp::Additive;
-		BaseDmgEffect->Modifiers[0].Attribute.SetUProperty(DamageProperty);
-		BaseDmgEffect->Modifiers[0].OwnedTags.AddTag(IGameplayTagsModule::RequestGameplayTag(FName(TEXT("Damage.Basic"))));
-		BaseDmgEffect->Duration.SetValue(UGameplayEffect::INSTANT_APPLICATION);
+		BaseDmgEffect->Modifiers[0].Attribute.SetUProperty(HealthProperty);
+		BaseDmgEffect->Duration.Value = UGameplayEffect::INSTANT_APPLICATION;
 
-		// Apply to target
 		SourceComponent->ApplyGameplayEffectToTarget(BaseDmgEffect, DestComponent, 1.f);
-
-		float ExpectedValue = StartHealth - DamageValue;
-		float ActualValue = DestComponent->GetSet<UAbilitySystemTestAttributeSet>()->Health;
-
-		Test->TestTrue(SKILL_TEST_TEXT("Damage Applied. %.2f == %.2f", ActualValue, ExpectedValue), ActualValue == ExpectedValue);
-	}
-
-	World->EditorDestroyActor(SourceActor, false);
-	World->EditorDestroyActor(DestActor, false);
-
-	return true;
-}
-
-bool GameplayEffectsTest_ModifyChanceToApplyToTarget(UWorld *World, FAutomationTestBase * Test)
-{
-	const float StartHealth = 100.f;
-	const float DamageValue = 5.f;
-
-	AAbilitySystemTestPawn *SourceActor = World->SpawnActor<AAbilitySystemTestPawn>();
-	AAbilitySystemTestPawn *DestActor = World->SpawnActor<AAbilitySystemTestPawn>();
-
-	UProperty *DamageProperty = FindFieldChecked<UProperty>(UAbilitySystemTestAttributeSet::StaticClass(), GET_MEMBER_NAME_CHECKED(UAbilitySystemTestAttributeSet, Damage));
-
-	UAbilitySystemComponent * SourceComponent = SourceActor->GetAbilitySystemComponent();
-	UAbilitySystemComponent * DestComponent = DestActor->GetAbilitySystemComponent();
-	SourceComponent->GetSet<UAbilitySystemTestAttributeSet>()->Health = StartHealth;
-	DestComponent->GetSet<UAbilitySystemTestAttributeSet>()->Health = StartHealth;
-
-	{
-		ABILITY_LOG_SCOPE(TEXT("Apply boost to chance to apply"))
-
-		UGameplayEffect * BaseEffect = Cast<UGameplayEffect>(StaticConstructObject(UGameplayEffect::StaticClass(), GetTransientPackage(), FName(TEXT("ChanceToApply"))));
-		BaseEffect->Modifiers.SetNum(1);
-		BaseEffect->Modifiers[0].Magnitude.SetValue(1.f);
-		BaseEffect->Modifiers[0].ModifierType = EGameplayMod::OutgoingGE;
-		BaseEffect->Modifiers[0].ModifierOp = EGameplayModOp::Additive;
-		BaseEffect->Modifiers[0].Attribute.SetUProperty(DamageProperty);
-		BaseEffect->Modifiers[0].EffectType = EGameplayModEffect::ChanceApplyTarget;
-		BaseEffect->Duration.SetValue(UGameplayEffect::INFINITE_DURATION);
-
-		// Apply to self
-		SourceComponent->ApplyGameplayEffectToTarget(BaseEffect, SourceComponent, 0.f);
-	}
-
-	{
-		ABILITY_LOG_SCOPE(TEXT("Apply Damage mod"))
-
-		UGameplayEffect * BaseDmgEffect = Cast<UGameplayEffect>(StaticConstructObject(UGameplayEffect::StaticClass(), GetTransientPackage(), FName(TEXT("Damage"))));
-		BaseDmgEffect->Modifiers.SetNum(1);
-		BaseDmgEffect->Modifiers[0].Magnitude.SetValue(DamageValue);
-		BaseDmgEffect->Modifiers[0].ModifierType = EGameplayMod::Attribute;
-		BaseDmgEffect->Modifiers[0].ModifierOp = EGameplayModOp::Additive;
-		BaseDmgEffect->Modifiers[0].Attribute.SetUProperty(DamageProperty);
-		BaseDmgEffect->Modifiers[0].OwnedTags.AddTag(IGameplayTagsModule::RequestGameplayTag(FName(TEXT("Damage.Basic"))));
-		BaseDmgEffect->Duration.SetValue(UGameplayEffect::INSTANT_APPLICATION);
-		BaseDmgEffect->ChanceToApplyToTarget.SetValue(0.f);
-
-		// Apply to target
-		SourceComponent->ApplyGameplayEffectToTarget(BaseDmgEffect, DestComponent, 0.f);
-
-		float ExpectedValue = (StartHealth - DamageValue);
-		float ActualValue = DestComponent->GetSet<UAbilitySystemTestAttributeSet>()->Health;
-
-		Test->TestTrue(SKILL_TEST_TEXT("Damage Applied. %.2f == %.2f", ActualValue, ExpectedValue), ActualValue == ExpectedValue);
-	}
-
-	World->EditorDestroyActor(SourceActor, false);
-	World->EditorDestroyActor(DestActor, false);
-
-	return true;
-}
-
-bool GameplayEffectsTest_ModifyChanceToExecuteOnGE(UWorld *World, FAutomationTestBase * Test)
-{
-	const float StartHealth = 100.f;
-	const float DamageValue = 5.f;
-
-	AAbilitySystemTestPawn *SourceActor = World->SpawnActor<AAbilitySystemTestPawn>();
-	AAbilitySystemTestPawn *DestActor = World->SpawnActor<AAbilitySystemTestPawn>();
-
-	UProperty *DamageProperty = FindFieldChecked<UProperty>(UAbilitySystemTestAttributeSet::StaticClass(), GET_MEMBER_NAME_CHECKED(UAbilitySystemTestAttributeSet, Damage));
-
-	UAbilitySystemComponent * SourceComponent = SourceActor->GetAbilitySystemComponent();
-	UAbilitySystemComponent * DestComponent = DestActor->GetAbilitySystemComponent();
-	SourceComponent->GetSet<UAbilitySystemTestAttributeSet>()->Health = StartHealth;
-	DestComponent->GetSet<UAbilitySystemTestAttributeSet>()->Health = StartHealth;
-
-	{
-		ABILITY_LOG_SCOPE(TEXT("Apply Damage buff"))
-
-		UGameplayEffect * BuffEffect = Cast<UGameplayEffect>(StaticConstructObject(UGameplayEffect::StaticClass(), GetTransientPackage(), FName(TEXT("DamageBuff"))));
-		BuffEffect->Modifiers.SetNum(1);
-		BuffEffect->Modifiers[0].Magnitude.SetValue(2.f);
-		BuffEffect->Modifiers[0].ModifierType = EGameplayMod::OutgoingGE;
-		BuffEffect->Modifiers[0].ModifierOp = EGameplayModOp::Multiplicitive;
-		BuffEffect->Modifiers[0].EffectType = EGameplayModEffect::Magnitude;
-		BuffEffect->Modifiers[0].Attribute.SetUProperty(DamageProperty);
-		BuffEffect->Modifiers[0].RequiredTags.AddTag(IGameplayTagsModule::RequestGameplayTag(FName(TEXT("Damage.Basic"))));
-		BuffEffect->Duration.SetValue(UGameplayEffect::INFINITE_DURATION);
-		BuffEffect->GameplayEffectTags.AddTag(IGameplayTagsModule::RequestGameplayTag(FName(TEXT("Damage.Buff"))));
-		BuffEffect->ChanceToExecuteOnGameplayEffect.SetValue(0.f);
-
-		// Apply to self
-		SourceComponent->ApplyGameplayEffectToTarget(BuffEffect, SourceComponent, 0.f);
-	}
-
-	// verify that outgoing damage is unbuffed
-	{
-		ABILITY_LOG_SCOPE(TEXT("Apply Damage mod"))
-
-		UGameplayEffect * BaseDmgEffect = Cast<UGameplayEffect>(StaticConstructObject(UGameplayEffect::StaticClass(), GetTransientPackage(), FName(TEXT("Damage"))));
-		BaseDmgEffect->Modifiers.SetNum(1);
-		BaseDmgEffect->Modifiers[0].Magnitude.SetValue(DamageValue);
-		BaseDmgEffect->Modifiers[0].ModifierType = EGameplayMod::Attribute;
-		BaseDmgEffect->Modifiers[0].ModifierOp = EGameplayModOp::Additive;
-		BaseDmgEffect->Modifiers[0].Attribute.SetUProperty(DamageProperty);
-		BaseDmgEffect->Modifiers[0].OwnedTags.AddTag(IGameplayTagsModule::RequestGameplayTag(FName(TEXT("Damage.Basic"))));
-		BaseDmgEffect->Duration.SetValue(UGameplayEffect::INSTANT_APPLICATION);
-
-		// Apply to target
-		SourceComponent->ApplyGameplayEffectToTarget(BaseDmgEffect, DestComponent, 0.f);
-
-		float ExpectedValue = (StartHealth - DamageValue);
-		float ActualValue = DestComponent->GetSet<UAbilitySystemTestAttributeSet>()->Health;
-
-		Test->TestTrue(SKILL_TEST_TEXT("Damage Applied. %.2f == %.2f", ActualValue, ExpectedValue), ActualValue == ExpectedValue);
-	}
-
-	{
-		ABILITY_LOG_SCOPE(TEXT("Apply boost to chance to execute"))
-
-		UGameplayEffect * BaseEffect = Cast<UGameplayEffect>(StaticConstructObject(UGameplayEffect::StaticClass(), GetTransientPackage(), FName(TEXT("ChanceToExecute"))));
-		BaseEffect->Modifiers.SetNum(1);
-		BaseEffect->Modifiers[0].Magnitude.SetValue(1.f);
-		BaseEffect->Modifiers[0].ModifierType = EGameplayMod::ActiveGE;
-		BaseEffect->Modifiers[0].ModifierOp = EGameplayModOp::Additive;
-		BaseEffect->Modifiers[0].EffectType = EGameplayModEffect::ChanceExecuteEffect;
-		BaseEffect->Modifiers[0].ModifierOp = EGameplayModOp::Additive;
-		BaseEffect->Modifiers[0].Attribute.SetUProperty(DamageProperty);
-		BaseEffect->Duration.SetValue(UGameplayEffect::INFINITE_DURATION);
-
-		// Apply to self
-		SourceComponent->ApplyGameplayEffectToTarget(BaseEffect, SourceComponent, 0.f);
-	}
-
-	// verify that the buff applies now
-	{
-		ABILITY_LOG_SCOPE(TEXT("Apply Damage mod"))
-
-		UGameplayEffect * BaseDmgEffect = Cast<UGameplayEffect>(StaticConstructObject(UGameplayEffect::StaticClass(), GetTransientPackage(), FName(TEXT("Damage"))));
-		BaseDmgEffect->Modifiers.SetNum(1);
-		BaseDmgEffect->Modifiers[0].Magnitude.SetValue(DamageValue);
-		BaseDmgEffect->Modifiers[0].ModifierType = EGameplayMod::Attribute;
-		BaseDmgEffect->Modifiers[0].ModifierOp = EGameplayModOp::Additive;
-		BaseDmgEffect->Modifiers[0].Attribute.SetUProperty(DamageProperty);
-		BaseDmgEffect->Modifiers[0].OwnedTags.AddTag(IGameplayTagsModule::RequestGameplayTag(FName(TEXT("Damage.Basic"))));
-		BaseDmgEffect->Duration.SetValue(UGameplayEffect::INSTANT_APPLICATION);
-
-		// Apply to target
-		SourceComponent->ApplyGameplayEffectToTarget(BaseDmgEffect, DestComponent, 0.f);
-
-		float ExpectedValue = (StartHealth - DamageValue * 3.f);
-		float ActualValue = DestComponent->GetSet<UAbilitySystemTestAttributeSet>()->Health;
-
-		Test->TestTrue(SKILL_TEST_TEXT("Damage Applied. %.2f == %.2f", ActualValue, ExpectedValue), ActualValue == ExpectedValue);
+		
+		Test->TestTrue(SKILL_TEST_TEXT("Basic Instant Damage Applied"), (DestComponent->GetSet<UAbilitySystemTestAttributeSet>()->Health == (StartHealth + DamageValue)));
+		ABILITY_LOG(Log, TEXT("Final Health: %.2f"), DestComponent->GetSet<UAbilitySystemTestAttributeSet>()->Health);
 	}
 
 	World->EditorDestroyActor(SourceActor, false);
@@ -6080,118 +132,7 @@ bool FGameplayEffectsTest::RunTest( const FString& Parameters )
 	World->BeginPlay();
 	
 	GameplayEffectsTest_InstantDamage(World, this);
-	GameplayEffectsTest_InstantDamageRemap(World, this);
-	GameplayEffectsTest_InstantDamage_Buffed(World, this);
-
-	GameplayEffectsTest_DurationDamage(World, this);
-	GameplayEffectsTest_PeriodicDamage(World, this);
-
-	GameplayEffectsTest_TemporaryDamage(World, this);
-	GameplayEffectsTest_TemporaryDamageBuffed(World, this);
-	GameplayEffectsTest_TemporaryDamageTemporaryBuff(World, this);
-	GameplayEffectsTest_LinkedBuffDestroy(World, this);
-	GameplayEffectsTest_SnapshotBuffDestroy(World, this);
-	GameplayEffectsTest_DurationBuff(World, this);
-	
-	// Buffing Buffs
-	GameplayEffectsTest_DamageBuffBuff_Basic(World, this);
-	GameplayEffectsTest_DamageBuffBuff_FullLink(World, this);
-	GameplayEffectsTest_DamageBuffBuff_FullSnapshot(World, this);
-	GameplayEffectsTest_DamageBuffBuff_SnapshotLink(World, this);
-
-	// GameplayEffects that apply other GameplayEffects
-	GameplayEffectsTest_DamageAppliesBuff(World, this);
-	GameplayEffectsTest_BuffAppliesBuff(World, this);
-	GameplayEffectsTest_BuffIndirection(World, this);
-
-	// GameplayEffect extensions
-	GameplayEffectsTest_LifestealExtension(World, this);
-	
-	GameplayEffectsTest_ShieldExtension(World, this);
-	GameplayEffectsTest_ShieldExtensionMultiple(World, this);
-
-	// Scaling modifiers
-	GameplayEffectsTest_InstantDamage_ScalingExplicit(World, this);
-	GameplayEffectsTest_InstantDamage_ScalingGlobal(World, this);
-
-	GameplayEffectsTest_InstantDamage_ScalingProperty(World, this);
-	GameplayEffectsTest_InstantDamage_ScalingPropertyNested(World, this);
-
-// 	GameplayEffectsTest_DotDamage_ScalingProperty_Snapshot(World, this);
-// GameplayTest_TickWorld(World, SMALL_NUMBER);
-// 	GameplayEffectsTest_DotDamage_ScalingProperty_Dynamic(World, this);
-// GameplayTest_TickWorld(World, SMALL_NUMBER);
-
-	GameplayEffectsTest_InstantDamage_OverrideScaling(World, this);
-
-	GameplayTest_TickWorld(World, SMALL_NUMBER);
-
-	// Tagging tests
-	GameplayEffectsTest_InstantDamageRequiredTag(World, this);
-	GameplayEffectsTest_InstantDamageIgnoreTag(World, this);  // busted
 	
-	GameplayEffectsTest_InstantDamageModifierPassesTag(World, this);
-	GameplayEffectsTest_InstantDamageModifierTag(World, this);
-
-	GameplayEffectsTest_MetaAttributes(World, this);
-	GameplayEffectsTest_TagOrdering(World, this);
-
-	GameplayTest_TickWorld(World, SMALL_NUMBER);
-	
-	//
-	// Stacking GE tests
-	//
-
-	// basic rules
-	GameplayEffectsTest_StackingHighest(World, this);
-	GameplayEffectsTest_StackingLowest(World, this);
-	GameplayEffectsTest_StackingUnlimited(World, this);
-	GameplayEffectsTest_StackingRemoval(World, this);
-	GameplayEffectsTest_StackingNoStack(World, this);
-
-	// custom rules
-	GameplayEffectsTest_StackingCustomCapped(World, this);
-	GameplayEffectsTest_StackingCustomDiminishingReturns(World, this);
-
-	// interactions between different rules/attributes
-	GameplayEffectsTest_StackingDifferentRules(World, this);
-	GameplayEffectsTest_StackingDifferentAttributes(World, this);
-	GameplayEffectsTest_StackingCustomTwoRules(World, this);
-	GameplayEffectsTest_StackingCustomTwoAttributes(World, this);
-
-	// interactions between stacking and modifiers
-	GameplayEffectsTest_StackingRemovingModifiers(World, this);
-	GameplayEffectsTest_StackingAddingModifiers(World, this);
-
-	// Immunity
-	GameplayEffectsTest_ImmunityIncoming(World, this);
-	GameplayEffectsTest_ImmunityOutgoing(World, this);
-	GameplayEffectsTest_ImmunityActiveGE(World, this);
-	GameplayEffectsTest_ImmunityMod(World, this);
-	GameplayEffectsTest_ImmunityActiveMod(World, this);
-	GameplayEffectsTest_ImmunityBlockedBuff(World, this);
-	GameplayEffectsTest_ImmunityBlockedBaseAndBuff(World, this);
-
-	//
-	// Chance to apply or execute
-	//
-
-	// Chance to apply to target
-	GameplayEffectsTest_ChanceToApplyToTarget(World, this);
-
-	// Chance to apply to GEs
-	// We need to test active, incoming and outgoing GEs
-	// Chance to execute has a slightly different path for immunity than it does for modifying a GE so it needs to be tested separately
-	GameplayEffectsTest_ChanceToExecuteOnActiveGEMod(World, this);
-	GameplayEffectsTest_ChanceToExecuteOnActiveGEImmunity(World, this);
-	GameplayEffectsTest_ChanceToExecuteOnOutgoingGEMod(World, this);
-	GameplayEffectsTest_ChanceToExecuteOnOutgoingGEImmunity(World, this);
-	GameplayEffectsTest_ChanceToExecuteOnIncomingGEMod(World, this);
-	GameplayEffectsTest_ChanceToExecuteOnIncomingGEImmunity(World, this);
-
-	// Modifiers to Chance to apply and execute
-	GameplayEffectsTest_ModifyChanceToApplyToTarget(World, this);
-	GameplayEffectsTest_ModifyChanceToExecuteOnGE(World, this);
 
 	GEngine->DestroyWorldContext(World);
 	World->DestroyWorld(false);
