Gameplay Effect inital refactor

diff --git a/Engine/Source/Runtime/GameplayAbilities/Private/GameplayEffectTypes.cpp b/Engine/Source/Runtime/GameplayAbilities/Private/GameplayEffectTypes.cpp
--- a/Engine/Source/Runtime/GameplayAbilities/Private/GameplayEffectTypes.cpp
+++ b/Engine/Source/Runtime/GameplayAbilities/Private/GameplayEffectTypes.cpp
@@ -147,12 +147,6 @@ FString EGameplayModOpToString(int32 Type)
 	return e->GetEnum(Type).ToString();
 }
 
-FString EGameplayModToString(int32 Type)
-{
-	static UEnum *e = FindObject<UEnum>(ANY_PACKAGE, TEXT("EGameplayMod"));
-	return e->GetEnum(Type).ToString();
-}
-
 FString EGameplayModEffectToString(int32 Type)
 {
 	static UEnum *e = FindObject<UEnum>(ANY_PACKAGE, TEXT("EGameplayModEffect"));
@@ -232,17 +226,46 @@ bool FGameplayTagCountContainer::HasAnyMatchingGameplayTags(const FGameplayTagCo
 	return false;
 }
 
-void FGameplayTagCountContainer::UpdateTagMap(const FGameplayTag& Tag, int32 CountDelta)
+void FGameplayTagCountContainer::UpdateTagMap(const struct FGameplayTag& Tag, int32 CountDelta)
+{
+	if (TagContainerType == EGameplayTagMatchType::Explicit)
+	{
+		// Update count of BaseTag
+		UpdateTagMap_Internal(Tag, CountDelta);
+	}
+	else if (TagContainerType == EGameplayTagMatchType::IncludeParentTags)
+	{
+		// Update count of BaseTag and all of its parent tags
+		FGameplayTagContainer TagAndParentsContainer = IGameplayTagsModule::Get().GetGameplayTagsManager().RequestGameplayTagParents(Tag);
+		for (auto ParentTagIt = TagAndParentsContainer.CreateConstIterator(); ParentTagIt; ++ParentTagIt)
+		{
+			UpdateTagMap_Internal(*ParentTagIt, CountDelta);
+		}
+	}
+}
+
+void FGameplayTagCountContainer::UpdateTagMap(const FGameplayTagContainer& Container, int32 CountDelta)
+{
+	for (auto TagIt = Container.CreateConstIterator(); TagIt; ++TagIt)
+	{
+		const FGameplayTag& BaseTag = *TagIt;
+		UpdateTagMap(BaseTag, CountDelta);
+	}
+}
+
+void FGameplayTagCountContainer::UpdateTagMap_Internal(const FGameplayTag& Tag, int32 CountDelta)
 {
 	// Update count of Tag
 	int32& Count = GameplayTagCountMap.FindOrAdd(Tag);
 
 	bool WasZero = Count == 0;
 	Count = FMath::Max(Count + CountDelta, 0);
-		
+
 	// If we went from 0->1 or 1->0
 	if (WasZero || Count == 0)
 	{
+		OnAnyTagChangeDelegate.Broadcast(Tag, Count);
+
 		FOnGameplayEffectTagCountChanged *Delegate = GameplayTagEventMap.Find(Tag);
 		if (Delegate)
 		{
@@ -251,27 +274,38 @@ void FGameplayTagCountContainer::UpdateTagMap(const FGameplayTag& Tag, int32 Cou
 	}
 }
 
-void FGameplayTagCountContainer::UpdateTagMap(const FGameplayTagContainer& Container, int32 CountDelta)
+bool FGameplayTagRequirements::RequirementsMet(FGameplayTagContainer Container) const
 {
-	for (auto TagIt = Container.CreateConstIterator(); TagIt; ++TagIt)
+	bool HasRequired = Container.MatchesAll(RequireTags, true);
+	bool HasIgnored = Container.MatchesAny(IgnoreTags, false);
+
+	return HasRequired && !HasIgnored;
+}
+
+void FActiveGameplayEffectsContainer::PrintAllGameplayEffects() const
+{
+	ABILITY_LOG_SCOPE(TEXT("ActiveGameplayEffects. Num: %d"), GameplayEffects.Num());
+	for (const FActiveGameplayEffect& Effect : GameplayEffects)
 	{
-		const FGameplayTag& BaseTag = *TagIt;
-		if (TagContainerType == EGameplayTagMatchType::Explicit)
-		{
-			// Update count of BaseTag
-			UpdateTagMap(BaseTag, CountDelta);
-		}
-		else if (TagContainerType == EGameplayTagMatchType::IncludeParentTags)
-		{
-			// Update count of BaseTag and all of its parent tags
-			FGameplayTagContainer TagAndParentsContainer = IGameplayTagsModule::Get().GetGameplayTagsManager().RequestGameplayTagParents(BaseTag);
-			for (auto ParentTagIt = TagAndParentsContainer.CreateConstIterator(); ParentTagIt; ++ParentTagIt)
-			{
-				UpdateTagMap(*ParentTagIt, CountDelta);
-			}
-		}
+		Effect.PrintAll();
 	}
 }
 
+void FActiveGameplayEffect::PrintAll() const
+{
+	ABILITY_LOG(Log, TEXT("Handle: %s"), *Handle.ToString());
+	ABILITY_LOG(Log, TEXT("StartWorldTime: %.2f"), StartWorldTime);
+	Spec.PrintAll();
+}
+
+void FGameplayEffectSpec::PrintAll() const
+{
+	ABILITY_LOG_SCOPE(TEXT("GameplayEffectSpec"));
+	ABILITY_LOG(Log, TEXT("Def: %s"), *Def->GetName());
+
+	ABILITY_LOG(Log, TEXT("Duration: %.2f"), GetDuration());
 
+	ABILITY_LOG(Log, TEXT("Period: %.2f"), GetPeriod());
 
+	ABILITY_LOG(Log, TEXT("Modifiers:"));
+}
\ No newline at end of file
diff --git a/Engine/Source/Runtime/GameplayAbilities/Private/GameplayModCalculation.cpp b/Engine/Source/Runtime/GameplayAbilities/Private/GameplayModCalculation.cpp
--- a/Engine/Source/Runtime/GameplayAbilities/Private/GameplayModCalculation.cpp
+++ b/Engine/Source/Runtime/GameplayAbilities/Private/GameplayModCalculation.cpp
@@ -3,26 +3,13 @@
 #include "AbilitySystemPrivatePCH.h"
 #include "GameplayModCalculation.h"
 
-FGameplayModCalculationSpecHandle FGameplayModCalculationSpecHandle::GenerateNewHandle()
-{
-	static int32 GHandleID = 0;
-	FGameplayModCalculationSpecHandle NewHandle(GHandleID++);
-
-	return NewHandle;
-}
-
 UGameplayModCalculation::UGameplayModCalculation(const FObjectInitializer& ObjectInitializer)
-: Super(ObjectInitializer)
+	: Super(ObjectInitializer)
 {
 
 }
 
-FGameplayModCalculationSpecHandle UGameplayModCalculation::GetContextForSource(const FGameplayEffectSpec& Spec)
+float UGameplayModCalculation::CalculateMagnitude_Implementation(const FGameplayEffectSpec& Spec) const
 {
-	return FGameplayModCalculationSpecHandle::GenerateNewHandle();
+	return 0.f;
 }
-
-FGameplayModCalculationSpecHandle UGameplayModCalculation::GetContextForTarget(const FGameplayEffectSpec& Spec)
-{
-	return FGameplayModCalculationSpecHandle::GenerateNewHandle();
-}
\ No newline at end of file
diff --git a/Engine/Source/Runtime/GameplayAbilities/Public/Abilities/GameplayAbility.h b/Engine/Source/Runtime/GameplayAbilities/Public/Abilities/GameplayAbility.h
--- a/Engine/Source/Runtime/GameplayAbilities/Public/Abilities/GameplayAbility.h
+++ b/Engine/Source/Runtime/GameplayAbilities/Public/Abilities/GameplayAbility.h
@@ -243,7 +243,13 @@ protected:
 	/**
 	 * The main function that defines what an ability does.
 	 *  -Child classes will want to override this
-	 *  -This function must call CommitAbility()
+	 *  -This function graph should call CommitAbility
+	 *  -This function graph should call EndAbility
+	 *  
+	 *  Latent/async actions are ok in this graph. Note that Commit and EndAbility calling requirements speak to the K2_ActivateAbility graph. 
+	 *  In C++, the call to K2_ActivateAbility() may return without CommitAbility or EndAbility having been called. But it is expected that this
+	 *  will only occur when latent/async actions are pending. When K2_ActivateAbility logically finishes, then we will expect Commit/End to have been called.
+	 *  
 	 */
 	UFUNCTION(BlueprintImplementableEvent, Category = Ability, FriendlyName = "ActivateAbility")
 	virtual void K2_ActivateAbility();	
diff --git a/Engine/Source/Runtime/GameplayAbilities/Public/Abilities/GameplayAbilityTargetTypes.h b/Engine/Source/Runtime/GameplayAbilities/Public/Abilities/GameplayAbilityTargetTypes.h
--- a/Engine/Source/Runtime/GameplayAbilities/Public/Abilities/GameplayAbilityTargetTypes.h
+++ b/Engine/Source/Runtime/GameplayAbilities/Public/Abilities/GameplayAbilityTargetTypes.h
@@ -63,7 +63,7 @@ struct GAMEPLAYABILITIES_API FGameplayAbilityTargetData
 
 	virtual ~FGameplayAbilityTargetData() { }
 
-	TArray<FActiveGameplayEffectHandle> ApplyGameplayEffect(const UGameplayEffect* GameplayEffect, const FGameplayEffectContextHandle& InEffectContext, float Level, FModifierQualifier Qualifier = FModifierQualifier());
+	TArray<FActiveGameplayEffectHandle> ApplyGameplayEffect(const UGameplayEffect* GameplayEffect, const FGameplayEffectContextHandle& InEffectContext, float Level, FPredictionKey PredictionKey = FPredictionKey());
 
 	virtual void AddTargetDataToContext(FGameplayEffectContextHandle& Context);
 
diff --git a/Engine/Source/Runtime/GameplayAbilities/Public/Abilities/Tasks/AbilityTask_WaitAttributeChange.h b/Engine/Source/Runtime/GameplayAbilities/Public/Abilities/Tasks/AbilityTask_WaitAttributeChange.h
--- a/Engine/Source/Runtime/GameplayAbilities/Public/Abilities/Tasks/AbilityTask_WaitAttributeChange.h
+++ b/Engine/Source/Runtime/GameplayAbilities/Public/Abilities/Tasks/AbilityTask_WaitAttributeChange.h
@@ -31,7 +31,7 @@ class UAbilityTask_WaitAttributeChange : public UAbilityTask
 
 	/** Wait until an attribute changes. */
 	UFUNCTION(BlueprintCallable, Category="Ability|Tasks", meta = (HidePin = "WorldContextObject", DefaultToSelf = "WorldContextObject", BlueprintInternalUseOnly = "TRUE"))
-	static UAbilityTask_WaitAttributeChange* WaitForAttributeChange(UObject* WorldContextObject, FGameplayAttribute Attribute, FGameplayTag WithTag, FGameplayTag WithoutTag);
+	static UAbilityTask_WaitAttributeChange* WaitForAttributeChange(UObject* WorldContextObject, FGameplayAttribute Attribute, FGameplayTag WithSrcTag, FGameplayTag WithoutSrcTag);
 
 	FGameplayTag WithTag;
 	FGameplayTag WithoutTag;
diff --git a/Engine/Source/Runtime/GameplayAbilities/Public/AbilitySystemComponent.h b/Engine/Source/Runtime/GameplayAbilities/Public/AbilitySystemComponent.h
--- a/Engine/Source/Runtime/GameplayAbilities/Public/AbilitySystemComponent.h
+++ b/Engine/Source/Runtime/GameplayAbilities/Public/AbilitySystemComponent.h
@@ -59,6 +59,7 @@ class GAMEPLAYABILITIES_API UAbilitySystemComponent : public UActorComponent, pu
 	virtual ~UAbilitySystemComponent();
 
 	virtual void InitializeComponent() override;
+	virtual void UninitializeComponent() override;
 	virtual void TickComponent(float DeltaTime, enum ELevelTick TickType, FActorComponentTickFunction *ThisTickFunction) override;
 
 	/** Checks to see if we should be active (ticking). Called after something changes that would cause us to tick or not tick. */
@@ -169,8 +170,8 @@ class GAMEPLAYABILITIES_API UAbilitySystemComponent : public UActorComponent, pu
 	// --------------------------------------------
 	// Primary outward facing API for other systems:
 	// --------------------------------------------
-	FActiveGameplayEffectHandle ApplyGameplayEffectSpecToTarget(OUT FGameplayEffectSpec& GameplayEffect, UAbilitySystemComponent *Target, FModifierQualifier BaseQualifier = FModifierQualifier());
-	FActiveGameplayEffectHandle ApplyGameplayEffectSpecToSelf(OUT FGameplayEffectSpec& GameplayEffect, FModifierQualifier BaseQualifier = FModifierQualifier());
+	FActiveGameplayEffectHandle ApplyGameplayEffectSpecToTarget(OUT FGameplayEffectSpec& GameplayEffect, UAbilitySystemComponent *Target, FPredictionKey PredictionKey=FPredictionKey());
+	FActiveGameplayEffectHandle ApplyGameplayEffectSpecToSelf(OUT FGameplayEffectSpec& GameplayEffect, FPredictionKey PredictionKey = FPredictionKey());
 
 	UFUNCTION(BlueprintCallable, BlueprintAuthorityOnly, Category = GameplayEffects)
 	bool RemoveActiveGameplayEffect(FActiveGameplayEffectHandle Handle);
@@ -198,12 +199,26 @@ class GAMEPLAYABILITIES_API UAbilitySystemComponent : public UActorComponent, pu
 	UFUNCTION(BlueprintCallable, Category = GameplayEffects)
 	float GetGameplayEffectMagnitude(FActiveGameplayEffectHandle Handle, FGameplayAttribute Attribute) const;
 
-	UFUNCTION(BlueprintCallable, Category = GameplayEffects)
-	float GetGameplayEffectMagnitudeByTag(FActiveGameplayEffectHandle InHandle, const FGameplayTag& InTag) const;
-
 	UFUNCTION(BlueprintCallable, Category = GameplayEffects)
 	bool IsGameplayEffectActive(FActiveGameplayEffectHandle InHandle) const;
 
+
+	// New Stuff ---------------------------------- (Remove this comment!)
+
+
+	/**
+	 * Populate the specified capture spec with the data necessary to capture an attribute from the component
+	 * 
+	 * @param OutCaptureSpec	[OUT] Capture spec to populate with captured data
+	 */
+	void CaptureAttributeForGameplayEffect(OUT FGameplayEffectAttributeCaptureSpec& OutCaptureSpec);
+
+	// Delegates (these need to be at the UObject level so we can safetly bind, rather than binding to raw at the ActiveGameplayEffect/Container level which is unsafe if the AbilitySystemComponent were killed).
+
+	void OnAttributeAggregatorDirty(FAggregator* Aggregator, FGameplayAttribute Attribute);
+
+	void OnOngoingTagRequirementChange(FActiveGameplayEffectHandle ActiveHandle);
+
 	// --------------------------------------------
 	// Tags
 	// --------------------------------------------
@@ -215,9 +230,28 @@ class GAMEPLAYABILITIES_API UAbilitySystemComponent : public UActorComponent, pu
 
 	virtual void GetOwnedGameplayTags(FGameplayTagContainer& TagContainer) const override;
 
+	/** 	 
+	 *  Allows GameCode to add loose gameplaytags which are not backed by a GameplayEffect. 
+	 *
+	 *	Tags added this way are not replicated! 
+	 *	
+	 *	It is up to the calling GameCode to make sure these tags are added on a;; clients/server where necessary
+	 */
+
+	void AddLooseGameplayTag(const FGameplayTag& GameplayTag, int32 Count=1);
+
+	void AddLooseGameplayTags(const FGameplayTagContainer& GameplayTag, int32 Count = 1);
+
+	void RemoveLooseGameplayTag(const FGameplayTag& GameplayTag, int32 Count = 1);
+
+	void RemoveLooseGameplayTags(const FGameplayTagContainer& GameplayTag, int32 Count = 1);
+
 	/** Allow events to be registered for specific gameplay tags being added or removed */
 	FOnGameplayEffectTagCountChanged& RegisterGameplayTagEvent(FGameplayTag Tag);
 
+	/** Returns multicast delegate that is invoked whenever a tag is added or removed (but not if just count is increased. Only for 'new' and 'removed' events) */
+	FOnGameplayEffectTagCountChanged& RegisterGenericGameplayTagEvent();
+
 	FOnGameplayAttributeChange& RegisterGameplayAttributeEvent(FGameplayAttribute Attribute);
 
 	// --------------------------------------------
@@ -226,17 +260,15 @@ class GAMEPLAYABILITIES_API UAbilitySystemComponent : public UActorComponent, pu
 	
 	FOnActiveGameplayEffectRemoved* OnGameplayEffectRemovedDelegate(FActiveGameplayEffectHandle Handle);
 
-	FActiveGameplayEffectHandle ApplyGameplayEffectToTarget(UGameplayEffect *GameplayEffect, UAbilitySystemComponent *Target, float Level = FGameplayEffectLevelSpec::INVALID_LEVEL, FGameplayEffectContextHandle Context = FGameplayEffectContextHandle(), FModifierQualifier BaseQualifier = FModifierQualifier());
+	FActiveGameplayEffectHandle ApplyGameplayEffectToTarget(UGameplayEffect *GameplayEffect, UAbilitySystemComponent *Target, float Level = UGameplayEffect::INVALID_LEVEL, FGameplayEffectContextHandle Context = FGameplayEffectContextHandle(), FPredictionKey PredictionKey = FPredictionKey());
 
 	UFUNCTION(BlueprintCallable, Category = GameplayEffects, meta=(FriendlyName = "ApplyGameplayEffectToSelf"))
 	FActiveGameplayEffectHandle K2_ApplyGameplayEffectToSelf(const UGameplayEffect *GameplayEffect, float Level, FGameplayEffectContextHandle EffectContext);
 	
-	FActiveGameplayEffectHandle ApplyGameplayEffectToSelf(const UGameplayEffect *GameplayEffect, float Level, const FGameplayEffectContextHandle& EffectContext, FModifierQualifier BaseQualifier = FModifierQualifier());
+	FActiveGameplayEffectHandle ApplyGameplayEffectToSelf(const UGameplayEffect *GameplayEffect, float Level, const FGameplayEffectContextHandle& EffectContext, FPredictionKey PredictionKey = FPredictionKey());
 
 	int32 GetNumActiveGameplayEffect() const;
 
-	void AddDependancyToAttribute(FGameplayAttribute Attribute, const TWeakPtr<FAggregator> InDependant);
-
 	void SetBaseAttributeValueFromReplication(float NewValue, FGameplayAttribute Attribute);
 
 	/** Tests if all modifiers in this GameplayEffect will leave the attribute > 0.f */
@@ -255,24 +287,10 @@ class GAMEPLAYABILITIES_API UAbilitySystemComponent : public UActorComponent, pu
 	/** Removes all active effects that match given query */
 	void RemoveActiveEffects(const FActiveGameplayEffectQuery Query);
 
-	void OnRestackGameplayEffects();
-
-	// --------------------------------------------
-	// Temp / Debug
-	// --------------------------------------------
-
-	void TEMP_ApplyActiveGameplayEffects();
+	void OnRestackGameplayEffects();	
 	
 	void PrintAllGameplayEffects() const;
 
-	void PushGlobalCurveOveride(UCurveTable *OverrideTable)
-	{
-		if (OverrideTable)
-		{
-			GlobalCurveDataOverride.Overrides.Push(OverrideTable);
-		}
-	}
-
 	// ----------------------------------------------------------------------------------------------------------------
 	//
 	//	GameplayCues
@@ -632,28 +650,18 @@ public:
 
 private:
 
-	bool HasNetworkAuthorityToApplyGameplayEffect(const FModifierQualifier QualifierContext) const;
+	bool HasNetworkAuthorityToApplyGameplayEffect(FPredictionKey PredictionKey) const;
 
 	void ExecutePeriodicEffect(FActiveGameplayEffectHandle	Handle);
 
-	void ExecuteGameplayEffect(FGameplayEffectSpec &Spec, const FModifierQualifier &QualifierContext);
+	void ExecuteGameplayEffect(FGameplayEffectSpec &Spec, FPredictionKey PredictionKey);
 
 	void CheckDurationExpired(FActiveGameplayEffectHandle Handle);
 
-	bool AreGameplayEffectApplicationRequirementsSatisfied(const class UGameplayEffect* EffectToAdd, const FGameplayEffectContextHandle& EffectContext) const;
-
 	bool IsOwnerActorAuthoritative() const;
 
 	void OnAttributeGameplayEffectSpecExected(const FGameplayAttribute &Attribute, const struct FGameplayEffectSpec &Spec, struct FGameplayModifierEvaluatedData &Data);
 
-	const FGlobalCurveDataOverride* GetCurveDataOverride() const
-	{
-		// only return data if we have overrides. NULL if we don't.
-		return (GlobalCurveDataOverride.Overrides.Num() > 0 ? &GlobalCurveDataOverride : NULL);
-	}
-
-	FGlobalCurveDataOverride	GlobalCurveDataOverride;	
-
 	// --------------------------------------------
 	
 	UPROPERTY(ReplicatedUsing=OnRep_GameplayEffects)
diff --git a/Engine/Source/Runtime/GameplayAbilities/Public/AttributeSet.h b/Engine/Source/Runtime/GameplayAbilities/Public/AttributeSet.h
--- a/Engine/Source/Runtime/GameplayAbilities/Public/AttributeSet.h
+++ b/Engine/Source/Runtime/GameplayAbilities/Public/AttributeSet.h
@@ -59,7 +59,7 @@ struct GAMEPLAYABILITIES_API FGameplayAttribute
 
 	FString GetName() const
 	{
-		return Attribute->GetName();
+		return *GetNameSafe(Attribute);
 	}
 
 private:
diff --git a/Engine/Source/Runtime/GameplayAbilities/Public/GameplayEffect.h b/Engine/Source/Runtime/GameplayAbilities/Public/GameplayEffect.h
--- a/Engine/Source/Runtime/GameplayAbilities/Public/GameplayEffect.h
+++ b/Engine/Source/Runtime/GameplayAbilities/Public/GameplayEffect.h
@@ -7,6 +7,8 @@
 #include "GameplayTagContainer.h"
 #include "TimerManager.h"
 #include "GameplayEffectTypes.h"
+#include "GameplayEffectAggregator.h"
+#include "GameplayEffectCalculation.h"
 #include "GameplayEffect.generated.h"
 
 struct FActiveGameplayEffect;
@@ -14,7 +16,20 @@ struct FActiveGameplayEffect;
 class UGameplayEffect;
 class UGameplayEffectTemplate;
 class UAbilitySystemComponent;
+class UGameplayModCalculation;
 
+/** Enumeration for options of where to capture gameplay attributes from for gameplay effects */
+UENUM()
+enum class EGameplayEffectAttributeCaptureSource : uint8
+{
+	Source,	// Source (caster) of the gameplay effect
+	Target	// Target (recipient) of the gameplay effect
+};
+
+// Repurposing(?) this to be the inplace/custom/scoped modifier used by FGameplayEffectExecutionDefinition
+//	I would expect this to be the thing that is built off to be a "scalable float/attribute reference/custom magnitude" thing.
+//	This would potentially need to be snapshot/not
+// 
 USTRUCT()
 struct GAMEPLAYABILITIES_API FExtensionAttributeModifierInfo
 {
@@ -37,8 +52,20 @@ struct GAMEPLAYABILITIES_API FExtensionAttributeModifierInfo
 	/** How much this modifies what it is applied to */
 	UPROPERTY(EditDefaultsOnly, Category=GameplayModifier)
 	FScalableFloat Magnitude;
+
+	/** Source of the gameplay attribute */
+	UPROPERTY(EditDefaultsOnly, Category = Capture)
+	EGameplayEffectAttributeCaptureSource Source;
+
+	UPROPERTY(EditDefaultsOnly, Category = GameplayModifier)
+	FGameplayTagRequirements	SourceTags;
+
+	UPROPERTY(EditDefaultsOnly, Category = GameplayModifier)
+	FGameplayTagRequirements	TargetTags;
 };
 
+
+// Should this just go away now?
 USTRUCT()
 struct GAMEPLAYABILITIES_API FGameplayModifierCallback
 {
@@ -65,28 +92,105 @@ struct FGameplayEffectStackingCallbacks
 	TArray<TSubclassOf<class UGameplayEffectStackingExtension> >	ExtensionClasses;
 };
 
-/**
-* Defines how A GameplayEffect levels
-*	Normally, GameplayEffect levels are specified when they are created.
-*	They can also be tied to their instigators attribute.
-*		For example, a Damage apply GameplayEffect that 'levels' based on the PhysicalDamage attribute
-*/
+
+
+
+
+/** Struct defining gameplay attribute capture options for gameplay effects */
 USTRUCT()
-struct GAMEPLAYABILITIES_API FGameplayEffectLevelDef
+struct FGameplayEffectAttributeCaptureDefinition
 {
 	GENERATED_USTRUCT_BODY()
 
-	/** When true, whatever creates or owns this will pass in a level. E.g, level is not intrinsic to this definition. */
-	UPROPERTY(EditDefaultsOnly, Category = GameplayEffectLevel)
-	bool InheritLevelFromOwner;
+	FGameplayEffectAttributeCaptureDefinition()
+	{
+
+	}
+
+	FGameplayEffectAttributeCaptureDefinition(FGameplayAttribute InAttribute, EGameplayEffectAttributeCaptureSource InSource, bool InSnapshot)
+		: AttributeToCapture(InAttribute), AttributeSource(InSource), bSnapshot(InSnapshot)
+	{
+
+	}
+
+	/** Gameplay attribute to capture */
+	UPROPERTY(EditDefaultsOnly, Category=Capture)
+	FGameplayAttribute AttributeToCapture;
+
+	/** Source of the gameplay attribute */
+	UPROPERTY(EditDefaultsOnly, Category=Capture)
+	EGameplayEffectAttributeCaptureSource AttributeSource;
+
+	/** Whether the attribute should be snapshotted or not */
+	UPROPERTY(EditDefaultsOnly, Category=Capture)
+	bool bSnapshot;
+
+	/** Equality/Inequality operators */
+	bool operator==(const FGameplayEffectAttributeCaptureDefinition& Other) const;
+	bool operator!=(const FGameplayEffectAttributeCaptureDefinition& Other) const;
+
+	FString ToSimpleString() const;
+};
+
+/** Structure for defining custom logic when a gameplay effect execution happens */
+USTRUCT()
+struct FGameplayEffectExecutionDefinition
+{
+	GENERATED_USTRUCT_BODY()
+
+	void OnObjectPostEditChange(UObject* Object, FPropertyChangedEvent& PropertyChangedEvent);
+
+	// @todo: Ideally this would be an interface, but can't store one in a UPROPERTY. Need to figure out where we want custom
+	// execution to come from, as it could potentially be multiple places in reality. Might want native code, might want a blueprint graph,
+	// might want to implement it inline in the gameplay effect graph (maybe? would need to be static or instancing becomes a mess), might
+	// even just want it to be a "custom calculation" inline with fancy math operations.
+	UPROPERTY(EditDefaultsOnly, Category=Execution)
+	TSubclassOf<UGameplayEffectCalculation> CalculationClass;
+	
+	// @todo: This area should also be able to capture attributes because it's going to likely want them for the execution, as well as for
+	// scoped modifiers. Need to decide where they live/how this looks in the UI/etc.
+	// 
+	
+	
+
+	// FGameplayEffectAttributeCaptureDefinition
+
+	/** Modifiers that are applied "in place" during the execution calculation */
+	UPROPERTY(EditDefaultsOnly, Category = Execution)
+	TArray<FExtensionAttributeModifierInfo>	CalculationModifiers;
+};
+
+/** The runtime/mutable data of an ExecutionDefinition */
+USTRUCT()
+struct FGameplayEffectExecutionSpec
+{
+	GENERATED_USTRUCT_BODY()
+
+	//FAggregatorRef	
+};
+
+// @todo: Enum to choose between types; Details customization to hide the other one
+// @todo: Add Math Expression parser
+USTRUCT()
+struct FGameplayEffectMagnitude
+{
+	GENERATED_USTRUCT_BODY()
+
+	FGameplayEffectMagnitude()
+	: ScalableFloatMagnitude()
+	, CalculationClassMagnitude(nullptr)
+	{
+	}
+
+	UPROPERTY(EditDefaultsOnly, Category=Magnitude)
+	FScalableFloat ScalableFloatMagnitude;
 
-	/** If set, the gameplay effect's level will be tied to this attribute on the instigator */
-	UPROPERTY(EditDefaultsOnly, Category = GameplayEffectLevel, meta = (FilterMetaTag="HideFromLevelInfos"))
-	FGameplayAttribute	Attribute;
+	UPROPERTY(EditDefaultsOnly, Category=Magnitude)
+	TSubclassOf<UGameplayModCalculation> CalculationClassMagnitude;
 
-	/** If true, take snapshot of attribute level when the gameplay effect is initialized. Otherwise, the level of the gameplay effect will update as the attribute it is tied to updates */
-	UPROPERTY(EditDefaultsOnly, Category = GameplayEffectLevel)
-	bool TakeSnapshotOnInit;
+	/** Attributes to capture for usage in custom calculations for magnitude */
+	UPROPERTY(EditDefaultsOnly, Category=Magnitude)
+	TArray<FGameplayEffectAttributeCaptureDefinition> CustomMagnitudeCalculationAttributes;
 };
 
 /**
@@ -100,11 +204,8 @@ struct GAMEPLAYABILITIES_API FGameplayModifierInfo
 {
 	GENERATED_USTRUCT_BODY()
 
-	FGameplayModifierInfo()
-		: ModifierType( EGameplayMod::Attribute )
-		, EffectType( EGameplayModEffect::Magnitude )
-		, TargetEffect( NULL )
-		, ModifierOp(EGameplayModOp::Additive)
+	FGameplayModifierInfo()	
+	: ModifierOp(EGameplayModOp::Additive)
 	{
 
 	}
@@ -113,18 +214,6 @@ struct GAMEPLAYABILITIES_API FGameplayModifierInfo
 	UPROPERTY(EditDefaultsOnly, Category=GameplayModifier, meta=(FilterMetaTag="HideFromModifiers"))
 	FGameplayAttribute Attribute;
 
-	/** What this modifies - Attribute, OutgoingGEs, IncomingGEs, ACtiveGEs. */
-	UPROPERTY(EditDefaultsOnly, Category=GameplayModifier)
-	TEnumAsByte<EGameplayMod::Type> ModifierType;
-
-	/** If we modify an effect, this is what we modify about it (Duration, Magnitude, etc) */
-	UPROPERTY(EditDefaultsOnly, Category=GameplayModifier)
-	TEnumAsByte<EGameplayModEffect::Type> EffectType;
-
-	/** If we are linking a gameplay effect to another effect, this is the effect to link */
-	UPROPERTY(EditDefaultsOnly, Category=GameplayModifier)
-	UGameplayEffect* TargetEffect;
-
 	/** The numeric operation of this modifier: Override, Add, Multiply, etc  */
 	UPROPERTY(EditDefaultsOnly, Category=GameplayModifier)
 	TEnumAsByte<EGameplayModOp::Type> ModifierOp;
@@ -136,24 +225,20 @@ struct GAMEPLAYABILITIES_API FGameplayModifierInfo
 	UPROPERTY(EditDefaultsOnly, Category=GameplayModifier)
 	FScalableFloat Magnitude; // Not modified from defaults
 
-	UPROPERTY(EditDefaultsOnly, BlueprintReadWrite, Category=GameplayEffect)
-	FGameplayEffectLevelDef	LevelInfo;
-
-	/** The thing I modify requires these tags */
+	// @todo: Rename/fix-up/etc.
 	UPROPERTY(EditDefaultsOnly, Category=GameplayModifier)
-	FGameplayTagContainer RequiredTags;
+	FGameplayEffectMagnitude MagnitudeV2;
 
-	/** The thing I modify must not have any of these tags */
 	UPROPERTY(EditDefaultsOnly, Category=GameplayModifier)
-	FGameplayTagContainer IgnoreTags;
-	
-	/** This modifier's tags. These tags are passed to any other modifiers that this modifies. */
+	FGameplayTagRequirements	SourceTags;
+
 	UPROPERTY(EditDefaultsOnly, Category=GameplayModifier)
-	FGameplayTagContainer OwnedTags;
+	FGameplayTagRequirements	TargetTags;
+
 
 	FString ToSimpleString() const
 	{
-		return FString::Printf(TEXT("%s %s %s BaseVaue: %s"), *EGameplayModToString(ModifierType), *EGameplayModOpToString(ModifierOp), *EGameplayModEffectToString(EffectType), *Magnitude.ToSimpleString());
+		return FString::Printf(TEXT("%s BaseVaue: %s"), *EGameplayModOpToString(ModifierOp), *Magnitude.ToSimpleString());
 	}
 };
 
@@ -228,6 +313,9 @@ public:
 	/** Constant specifying that the combat effect has no period and doesn't check for over time application */
 	static const float NO_PERIOD;
 
+	/** No Level/Level not set */
+	static const float INVALID_LEVEL;
+
 #if WITH_EDITORONLY_DATA
 	/** Template to derive starting values and editing customization from */
 	UPROPERTY(EditDefaultsOnly, BlueprintReadOnly, Category = Template)
@@ -250,22 +338,8 @@ public:
 	UPROPERTY(EditDefaultsOnly, BlueprintReadOnly, Category=GameplayEffect)
 	TArray<FGameplayModifierInfo> Modifiers;
 
-	/** Array of level definitions that will determine how this GameplayEffect scales */
-	UPROPERTY(EditDefaultsOnly, BlueprintReadWrite, Category=GameplayEffect)
-	FGameplayEffectLevelDef	LevelInfo;
-		
-	// "I can only be applied to targets that have these tags"
-	// "I can only exist on CE buckets on targets that have these tags":
-	
-	/** Container of gameplay tags that have to be present on the target actor for the effect to be applied */
-	UPROPERTY(EditDefaultsOnly, BlueprintReadOnly, Category = Application)
-	FGameplayTagContainer ApplicationRequiredTargetTags;
-
-	// "I can only be applied if my instigator has these tags"
-
-	/** Container of gameplay tags that have to be present on the instigator actor for the effect to be applied */
-	UPROPERTY(EditDefaultsOnly, BlueprintReadOnly, Category = Application)
-	FGameplayTagContainer ApplicationRequiredInstigatorTags;
+	UPROPERTY(EditDefaultsOnly, BlueprintReadOnly, Category = GameplayEffect)
+	TArray<FGameplayEffectExecutionDefinition>	Executions;
 
 	/** Probability that this gameplay effect will be applied to the target actor (0.0 for never, 1.0 for always) */
 	UPROPERTY(EditDefaultsOnly, Category=Application, meta=(GameplayAttribute="True"))
@@ -279,68 +353,6 @@ public:
 	UPROPERTY(EditDefaultsOnly, BlueprintReadOnly, Category = GameplayEffect)
 	TArray<UGameplayEffect*> TargetEffects;
 
-	/** removes active gameplay effects and stops gameplay effects from applying if the tags and qualification context match */
-	UPROPERTY(EditDefaultsOnly, BlueprintReadOnly, Category = GameplayEffect)
-	TEnumAsByte<EGameplayImmunity::Type> AppliesImmunityTo;
-
-	// Modify duration of CEs
-
-	// Modify MaxStacks of CEs (maybe... probably not)
-
-	// ------------------------------------------------
-	// Functions
-	// ------------------------------------------------
-
-	/**
-	 * Determines if the set of supplied gameplay tags are enough to satisfy the application tag requirements of the effect
-	 * 
-	 * @param InstigatorTags	Owned gameplay tags of the instigator applying the effect
-	 * @param TargetTags		Owned gameplay tags of the target about to be affected by the effect
-	 * @return					True if the instigator and target actor tags meet the requirements for this gameplay effect, false otherwise
-	 */
-	bool AreApplicationTagRequirementsSatisfied(const FGameplayTagContainer& InstigatorTags, const FGameplayTagContainer& TargetTags) const;
-
-	// ------------------------------------------------
-	// New Tagging functionality
-	// ------------------------------------------------
-
-	/** "These are my tags" */
-	UPROPERTY(EditDefaultsOnly, BlueprintReadOnly, Category=Tags)
-	FGameplayTagContainer GameplayEffectTags;
-
-	/** "In order to affect another GameplayEffect, they must have ALL of these tags" */
-	UPROPERTY(EditDefaultsOnly, BlueprintReadOnly, Category=Tags)
-	FGameplayTagContainer GameplayEffectRequiredTags;
-	
-	/** "In order to affect another GameplayEffect, they must NOT have ANY of these tags" */
-	UPROPERTY(EditDefaultsOnly, BlueprintReadOnly, Category=Tags)
-	FGameplayTagContainer GameplayEffectIgnoreTags;
-
-	/**
-	* Can this GameplayEffect modify a GameplayEffect that owns Tags
-	*
-	* @param Tags	Owned gameplay tags of the gameplay effect this effect is being applied to
-	* @return		True if the tags meet the requirements for this gameplay effect, false otherwise
-	*/
-	bool AreGameplayEffectTagRequirementsSatisfied(const FGameplayTagContainer& Tags) const
-	{
-		bool bHasRequired = Tags.MatchesAll(GameplayEffectRequiredTags, true);
-		bool bHasIgnored = Tags.MatchesAny(GameplayEffectIgnoreTags, false);
-
-		return bHasRequired && !bHasIgnored;
-	}
-
-	/**
-	* Can this GameplayEffect modify the input parameter, based on tags
-	*
-	* @param GameplayEffectToBeModified	A GameplayEffect we are trying to apply this GameplayEffect to.
-	* @return							True if the tags owned by GameplayEffectToBeModified meet the requirements for this gameplay effect, false otherwise
-	*/
-	bool AreGameplayEffectTagRequirementsSatisfied(const UGameplayEffect *GameplayEffectToBeModified) const
-	{
-		return AreGameplayEffectTagRequirementsSatisfied(GameplayEffectToBeModified->GameplayEffectTags);
-	}
-
 	// ------------------------------------------------
 	// Gameplay tag interface
 	// ------------------------------------------------
@@ -348,33 +360,14 @@ public:
 	/** Overridden to return requirements tags */
 	virtual void GetOwnedGameplayTags(FGameplayTagContainer& TagContainer) const override;
 
-	/** Get the "clear tags" for the effect */
-	virtual void GetClearGameplayTags(FGameplayTagContainer& TagContainer) const;
-
 	void ValidateGameplayEffect();
 
-	/** "These tags are applied to the actor I am applied to" */
-	UPROPERTY(EditDefaultsOnly, BlueprintReadOnly, Category=Tags)
-	FGameplayTagContainer OwnedTagsContainer;
-
 	/** Container of gameplay tags to be cleared upon effect application; Any active effects with these tags that can be cleared, will be */
 	UPROPERTY(EditDefaultsOnly, BlueprintReadOnly, Category=Tags)
 	FGameplayTagContainer ClearTagsContainer;
-	
-	// Used to quickly tell if a GameplayEffect modifies another GameplayEffect (or a set of attributes)
-	bool ModifiesAnyProperties(EGameplayMod::Type ModType, const TSet<UProperty> & Properties)
-	{
-		return false;
-	}
 
 	virtual void PostLoad() override;
 
-	// -----------------------------------------------
-	
-	/** Should copies of this GameplayEffect be a snapshot of the current state or update when it does (linked) */
-	UPROPERTY(EditDefaultsOnly, BlueprintReadOnly, Category = Advanced)
-	TEnumAsByte<EGameplayEffectCopyPolicy::Type>	CopyPolicy;
-
 	// ----------------------------------------------
 
 	/** Cues to trigger non-simulated reactions in response to this GameplayEffect such as sounds, particle effects, etc */
@@ -400,470 +393,183 @@ public:
 	TSubclassOf<class UGameplayEffectStackingExtension> StackingExtension;
 
 
-	/** Hack! Fixes issues in PIE when you create a dataasset and it cant be used in networking until you reload the editor */
-	bool IsNameStableForNetworking() const override
-	{
-		return true;
-	}
-
-protected:
-	void ValidateStacking();
-
-
-};
-
-/**
- * FGameplayEffectLevelSpec
- *	Level Specification. This can be a static, constant level specified on creation or it can be dynamically tied to a source's attribute value.
- *	For example, a GameplayEffect could be made whose level is tied to its instigators PhysicalDamage or Intelligence attribute.
- */
 
-struct FGameplayEffectLevelSpec
-{
-	GAMEPLAYABILITIES_API static const float INVALID_LEVEL;
-
-	FGameplayEffectLevelSpec()
-		:ConstantLevel(INVALID_LEVEL)
-		,CachedLevel(INVALID_LEVEL)
-	{
-	}
-
-	FGameplayEffectLevelSpec(float InLevel, const FGameplayEffectLevelDef &Def, class AActor *InSource)
-		: ConstantLevel(InLevel)
-		, CachedLevel(InLevel)
-		, Source(InSource)
-	{
-		if (Def.Attribute.GetUProperty() != NULL)
-		{
-			// If given an attribute to base level on, then we throw out the passed in constant level.
-			Attribute = Def.Attribute;
-			ConstantLevel = INVALID_LEVEL;
-			CachedLevel = INVALID_LEVEL;
-		}
-
-		if (Def.TakeSnapshotOnInit)
-		{
-			SnapshotLevel();
-		}
-	}
-
-	void ApplyNewDef(const FGameplayEffectLevelDef &Def, TSharedPtr<FGameplayEffectLevelSpec> &OutSharedPtr) const
-	{
-		if (Def.InheritLevelFromOwner)
-		{
-			return;
-		}
-
-		check(OutSharedPtr.IsValid());
-		if (Def.Attribute != OutSharedPtr->Attribute)
-		{
-			// In def levels off something different
-			// make a new level spec
-			OutSharedPtr = TSharedPtr<FGameplayEffectLevelSpec>(new FGameplayEffectLevelSpec(INVALID_LEVEL, Def, Source.Get()));
-		}
-	}
-
-	/** Dynamic simply means the level may change. It is not constant. */
-	bool IsDynamic() const
-	{
-		return ConstantLevel == INVALID_LEVEL && Attribute.GetUProperty() != NULL;
-	}
-
-	/** Valid means we have some meaningful data. If we have an INVALID_LEVEL constant value and are not tied to a dynamic property, then we are invalid. */
-	bool IsValid() const
-	{
-		return ConstantLevel != INVALID_LEVEL || Attribute.GetUProperty() != NULL;
-	}
-
-	float GetLevel() const;
-
-	void SnapshotLevel()
-	{
-		// This should snapshot the current level (if dynamic/delegate) and save off its value so that it doesn't change
-		ConstantLevel = GetLevel();
-		Source = NULL;
-	}
-
-	void RegisterLevelDependancy(TWeakPtr<struct FAggregator> OwningAggregator);
-
-	void PrintAll() const;
-
-	mutable float ConstantLevel;	// Final/constant level. Once this is set we are locked at the given level.
-	mutable float CachedLevel;		// Last read value. Needed in case we lose our source, we use the last known level.
-
-	TWeakObjectPtr<class AActor>	Source;
-
-	FGameplayAttribute	Attribute;
-};
-
-
-/**
- * FAggregatorRef
- *	A reference to an FAggregator. The reference may be weak or hard, and this can be changed over the lifetime of the FAggregatorRef.
- *	
- *	There are cases where we want weak references in an aggregator chain. 
- *		For example a RunSpeed buff, which when it is destroyed we want the RunSpeed attribute aggregator to recalculate the RunSpeed value.
- *
- *	There are cases where we want to make a copy of what we are weak referencing and make the reference a hard ref to that copy
- *		For example, a DOT which is buffed is attached to a target. We want to make a copy of the DOT and its buff then give it to the target as a hard ref so that if 
- *		the buff expires on the source, the applied DOT is still buffed.
- *
- */
-USTRUCT()
-struct FAggregatorRef
-{
-	GENERATED_USTRUCT_BODY()
-
-	friend struct FAggregatorRef;
-
-	FAggregatorRef()
-	{
-	}
-
-	FAggregatorRef(struct FAggregator *Src)
-		: SharedPtr(Src)
-	{		 
-		WeakPtr = SharedPtr;
-	}
-
-	FAggregatorRef(const FAggregatorRef *Src)
-	{
-		SetSoftRef(Src);
-	}
-
-	void MakeHardRef()
-	{
-		check(WeakPtr.IsValid());
-		SharedPtr = WeakPtr.Pin();
-	}
-	void MakeSoftRef()
-	{
-		check(WeakPtr.IsValid());
-		SharedPtr.Reset();
-	}
-
-	void SetSoftRef(const FAggregatorRef *Src)
-	{
-		check(!SharedPtr.IsValid());
-		WeakPtr = Src->SharedPtr;
-	}
+	// ----------------------------------------------------------------------
+	//		Tag pass: properties that make sense "right now" (remove this comment DaveR/BillyB)
+	// ----------------------------------------------------------------------
 	
-	FAggregator * Get()
-	{
-		return WeakPtr.IsValid() ? WeakPtr.Pin().Get() : NULL;
-	}
-
-	const FAggregator* Get() const
-	{
-		return WeakPtr.IsValid() ? WeakPtr.Pin().Get() : NULL;
-	}
-
-	bool IsValid() const
-	{
-		return WeakPtr.IsValid();
-	}
-
-	/** Become a hard reference to a new copy of what we are reference  */
-	void MakeUnique();
-
-	/** Become a hard reference to a new copy of what we are reference AND make new copies/hard refs of the complete modifier chain in our FAggregator */
-	void MakeUniqueDeep();
-
-	bool NetSerialize(FArchive& Ar, class UPackageMap* Map, bool& bOutSuccess);	
-
-	FString ToString() const;
-	void PrintAll() const;
-
-private:
-
-	TSharedPtr<struct FAggregator>	SharedPtr;
-	TWeakPtr<struct FAggregator> WeakPtr;
-};
-
-template<>
-struct TStructOpsTypeTraits< FAggregatorRef > : public TStructOpsTypeTraitsBase
-{
-	enum
-	{
-		WithNetSerializer = true,
-		WithCopy = false
-	};
-};
-
-/**
- * GameplayModifierData
- *	This is the data that FAggregator aggregates and turns into FGameplayModifierEvaluatedData.
- *  It is distinct from FGameplayModifierEvaluatedData in that FGameplayModifierData ia level has not been applied to this data.
- *  FGameplayModifierData::Magnitude is an FScalableFloat which describes a numeric value for a given level.
- * 
- */
-struct FGameplayModifierData
-{
-	FGameplayModifierData()
-	{
-		
-	}
-
-	FGameplayModifierData(const FGameplayModifierInfo &Info, const FGlobalCurveDataOverride *CurveData)
-	{
-		Magnitude = Info.Magnitude.MakeFinalizedCopy(CurveData);
-		Tags = Info.OwnedTags;
-
-		// Fixme: this is static data, should be a reference
-		RequireTags = Info.RequiredTags;
-		IgnoreTags = Info.IgnoreTags;
 
-		if (Info.Callbacks.Num() > 0)
-		{
-			Callbacks = &Info.Callbacks;
-		}
-	}
-
-	FGameplayModifierData(FScalableFloat InMagnitude)
-	{
-		// Magnitude may scale based on our level
-		Magnitude = InMagnitude;
-		Callbacks = NULL;
-	}
-
-	FGameplayModifierData(float InMagnitude, const TArray<FGameplayModifierCallback>* InCallbacks)
-	{
-		// Magnitude will be fixed at this value
-		Magnitude.SetValue(InMagnitude);
-		Callbacks = InCallbacks;
-	}
-
-	// That magnitude that we modify by
-	FScalableFloat Magnitude;
+	/** The GameplayEffect's Tags: tags the the GE *has* and DOES NOT give to the actor. */
+	UPROPERTY(EditDefaultsOnly, BlueprintReadOnly, Category = Tags, meta=(DisplayName="GameplayEffectAssetTag"))
+	FGameplayTagContainer GameplayEffectTags;
 
-	// The tags I have
-	FGameplayTagContainer Tags;	
-	FGameplayTagContainer RequireTags;
-	FGameplayTagContainer IgnoreTags;
+	
+	/** "These tags are applied to the actor I am applied to" */
+	UPROPERTY(EditDefaultsOnly, BlueprintReadOnly, Category = Tags)
+	FGameplayTagContainer OwnedTagsContainer;
+	
+	/** Once Applied, these tags requirements are used to determined if the GameplayEffect is "on" or "off". A GameplayEffect can be off and do nothing, but still applied. */
+	UPROPERTY(EditDefaultsOnly, BlueprintReadOnly, Category = Tags)
+	FGameplayTagRequirements OngoingTagRequirements;
 
-	// Callback information for custom logic pre/post evaluation
-	const TArray<FGameplayModifierCallback>* Callbacks;
+	/** Tag requirements for this GameplayEffect to be applied to a target. This is pass/fail at the time of application. If fail, this GE fails to apply. */
+	UPROPERTY(EditDefaultsOnly, BlueprintReadOnly, Category = Tags)
+	FGameplayTagRequirements ApplicationTagRequirements;
 
-	void PrintAll() const;
 };
 
-/**
- * GameplayModifierEvaluatedData
- *	This is the output from an FAggregator: a numeric value and a set of GameplayTags.
- */
+// GE_REMOVE this may go away? Or still useful for storing 'Final results of shit we calculated'?
 struct FGameplayModifierEvaluatedData
 {
 	FGameplayModifierEvaluatedData()
 		: Magnitude(0.f)
-		, Callbacks(NULL)
 		, IsValid(false)
 	{
 	}
 
-	FGameplayModifierEvaluatedData(float InMagnitude, const TArray<FGameplayModifierCallback>* InCallbacks = NULL, FActiveGameplayEffectHandle InHandle = FActiveGameplayEffectHandle(), const FGameplayTagContainer *InTags = NULL)
+	FGameplayModifierEvaluatedData(float InMagnitude, FActiveGameplayEffectHandle InHandle = FActiveGameplayEffectHandle())
 		: Magnitude(InMagnitude)
-		, Callbacks(InCallbacks)
 		, Handle(InHandle)
 		, IsValid(true)
 	{
-		if (InTags)
-		{
-			Tags = *InTags;
-		}
 	}
 
+	FGameplayTagContainer	SourceTags;
+	FGameplayTagContainer	TargetTags;
+
 	float	Magnitude;
-	FGameplayTagContainer Tags;
-	const TArray<FGameplayModifierCallback>* Callbacks;
 	FActiveGameplayEffectHandle	Handle;	// Handle of the active gameplay effect that originated us. Will be invalid in many cases
 	bool IsValid;
-
-	// Helper function for building up final values during an aggregation
-	void Aggregate(OUT FGameplayTagContainer &OutTags, OUT float &OutMagnitude, const float Bias=0.f) const;
-
-	void InvokePreExecute(struct FGameplayEffectModCallbackData &Data) const;
-	void InvokePostExecute(const struct FGameplayEffectModCallbackData &Data) const;
-
-	void PrintAll() const;
 };
 
 /**
- * FAggregator - a data structure for aggregating stuff in GameplayEffects.
- *	Aggregates a numeric value (float) and a set of gameplay tags. This could be further extended.
- *
- *	Aggregation is done with BaseData + Mods[].
- *	-BaseData is simply the base data. We are initiliazed with base data and base data can be directly modified via ::ExecuteMod.
- *	-Mods[] are lists of other FAggregators. That is, we have a list for each EGameplayModOp: Add, multiply, override.
- *	-These lists contain FAggregatorRefs, which may be soft or hard refs to other FAggregators.
- *	-::Evalate() takes our BaseData, and then crawls through ours Mods[] list and aggregates a final output (FGameplayModifierEvaluatedData)
- *	-Results of ::Evaluate are cached in CachedData.
- *	-FAggregator also keeps a list of weak ptrs to other FAggregators that are dependant on us. If we change, we let these aggregators know, so they can invalidate their cached data.
- *
- *
+ * Modifier Specification
+ *	-Const data (FGameplayModifierInfo) tells us what we modify, what we can modify
+ *	-Mutable Aggregated data tells us how we modify (magnitude).
+ *  
+ * Modifiers can be modified. A modifier spec holds these modifications along with a reference to the const data about the modifier.
+ * 
  */
-struct FAggregator : public TSharedFromThis<FAggregator>
+struct FModifierSpec
 {
-	DECLARE_DELEGATE_OneParam(FOnDirty, FAggregator*);
-
-	FAggregator();
-	FAggregator(const FGameplayModifierData &InBaseData, TSharedPtr<FGameplayEffectLevelSpec> LevelInfo, const TCHAR* InDebugString);
-	FAggregator(const FScalableFloat &InBaseMagnitude, TSharedPtr<FGameplayEffectLevelSpec> LevelInfo, const TCHAR* InDebugString);
-	FAggregator(const FGameplayModifierEvaluatedData &InEvalData, const TCHAR* InDebugString);
-	FAggregator(const FAggregator &In);
-	virtual ~FAggregator();
-
-	FAggregator& MarkDirty();
-	void ClearAllDependancies();
-
-	const FGameplayModifierEvaluatedData& Evaluate() const;
-
-	void PreEvaluate(struct FGameplayEffectModCallbackData &Data) const;
-	void PostEvaluate(const struct FGameplayEffectModCallbackData &Data) const;
-
-	void TakeSnapshotOfLevel();
-	void ApplyMod(EGameplayModOp::Type ModType, FAggregatorRef Ref, bool TakeSnapshot);
-	
-	void ExecuteModAggr(EGameplayModOp::Type ModType, FAggregatorRef Ref);
-	void ExecuteMod(EGameplayModOp::Type ModType, const FGameplayModifierEvaluatedData& EvaluatedData);
-
-	void AddDependantAggregator(TWeakPtr<FAggregator> InDependant);
-
-	void RegisterLevelDependancies();
-
-	TSharedPtr<FGameplayEffectLevelSpec> Level;
-	FActiveGameplayEffectHandle	ActiveHandle;	// Handle to owning active effect. Will be null in many cases.
-
-	FGameplayModifierData		BaseData;
-	TArray<FAggregatorRef>		Mods[EGameplayModOp::Max];
-
-	TArray<TWeakPtr<FAggregator > >	Dependants;
+	FModifierSpec(const FGameplayModifierInfo& InInfo);
 
-	FOnDirty	OnDirty;
+	// @todo: Need a way to indicate whether the magnitude was properly calculated or not. Asking for a value early (before target is available, etc.
+	// will just result in incorrect values)
+	float GetEvaluatedMagnitude() const;
+	void CalculateMagnitude(OUT struct FGameplayEffectSpec& OwnerSpec);
 
-	void PrintAll() const;
-	void RefreshDependencies();
-	void MakeUniqueDeep();
-
-	void SetFromNetSerialize(float NetSerialize);
+	// Hard Ref to what we modify, this stuff is const and never changes
+	const FGameplayModifierInfo& Info;
 
-	// ----------------------------------------------------------------------------
-	// This is data only used in debugging/tracking where aggregator's came from
-	// ----------------------------------------------------------------------------
-#if SKILL_SYSTEM_AGGREGATOR_DEBUG
-	FString DebugString;
-	mutable int32 CopiesMade;
+	// @todo: Probably need to make the modifier info private so people don't go gunking around in the magnitude
+	/** In the event that the modifier spec requires custom magnitude calculations, this is the authoritative, last evaluated value of the magnitude */
+	float EvaluatedMagnitude;
 
 	FString ToSimpleString() const
 	{
-		return FString::Printf(TEXT("0x%X %s. CacheValid: %d Mods: [%d/%d/%d]"), this, *DebugString, CachedData.IsValid, 
-			GetNumValidMods(EGameplayModOp::Override), GetNumValidMods(EGameplayModOp::Additive), GetNumValidMods(EGameplayModOp::Multiplicitive) );
+		return Info.ToSimpleString();
 	}
+};
 
-	struct FAllocationStats
-	{
-		FAllocationStats()
-		{
-			Reset();
-		}
+/** Saves list of modified attributes, to use for gameplay cues or later processing */
+USTRUCT()
+struct FGameplayEffectModifiedAttribute
+{
+	GENERATED_USTRUCT_BODY()
 
-		void Reset()
-		{
-			FMemory::Memzero(this, sizeof(FAllocationStats));
-		}
+	/** The attribute that has been modified */
+	UPROPERTY()
+	FGameplayAttribute Attribute;
 
-		int32 DefaultCStor;
-		int32 ModifierCStor;
-		int32 ScalableFloatCstor;
-		int32 FloatCstor;
-		int32 CopyCstor;
+	/** Total magnitude applied to that attribute */
+	UPROPERTY()
+	float TotalMagnitude;
 
-		int32 DependantsUpdated;
-	};
+	FGameplayEffectModifiedAttribute() : TotalMagnitude(0.0f) {}
+};
 
-	static FAllocationStats AllocationStats;
-#else
+/** Struct used to hold the result of a gameplay attribute capture; Initially seeded by definition data, but then populated by ability system component when appropriate */
+USTRUCT()
+struct FGameplayEffectAttributeCaptureSpec
+{
+	// Allow these as friends so they can seed the aggregator, which we don't otherwise want exposed
+	friend struct FActiveGameplayEffectsContainer;
+	friend class UAbilitySystemComponent;
 
-	FString ToSimpleString() const
-	{
-		return FString::Printf(TEXT("CacheValid: %d Mods: [%d/%d/%d]"), CachedData.IsValid, 
-			GetNumValidMods(EGameplayModOp::Override), GetNumValidMods(EGameplayModOp::Additive), GetNumValidMods(EGameplayModOp::Multiplicitive) );
-	}
+	GENERATED_USTRUCT_BODY()
 
-#endif
+	// Constructors
+	FGameplayEffectAttributeCaptureSpec();
+	FGameplayEffectAttributeCaptureSpec(const FGameplayEffectAttributeCaptureDefinition& InDefinition);
+
+	// @todo:
+	// API should provide way to get:
+		// Final modified value
+		// Base value
+		// Bonus value (Final - Base)
+		// Mods (optionally matching tag filter)
 	
+	/** Simple accessor to backing capture definition */
+	const FGameplayEffectAttributeCaptureDefinition& GetBackingDefinition() const;
 
+	const FAggregator& GetAggregator() const;
+		
 private:
 
-	int32 GetNumValidMods(EGameplayModOp::Type Type) const
-	{
-		int32 Num=0;
-		for (const FAggregatorRef &Agg : Mods[Type])
-		{
-			if (Agg.IsValid())
-			{
-				Num++;
-			}
-		}
-		return Num;
-	}
+	/** Copy of the definition the spec should adhere to for capturing */
+	UPROPERTY()
+	FGameplayEffectAttributeCaptureDefinition BackingDefinition;
 
-	mutable FGameplayModifierEvaluatedData	CachedData;
+	/** Ref to the aggregator for the captured attribute */
+	FAggregatorRef AttributeAggregator;
 };
 
-/**
- * Modifier Specification
- *	-Const data (FGameplayModifierInfo) tells us what we modify, what we can modify
- *	-Mutable Aggregated data tells us how we modify (magnitude).
- *  
- * Modifiers can be modified. A modifier spec holds these modifications along with a reference to the const data about the modifier.
- * 
- */
-struct FModifierSpec
+/** Struct used to handle a collection of captured source and target attributes */
+USTRUCT()
+struct GAMEPLAYABILITIES_API FGameplayEffectAttributeCaptureSpecContainer
 {
-	FModifierSpec(const FGameplayModifierInfo& InInfo, TSharedPtr<FGameplayEffectLevelSpec> InLevel, const FGlobalCurveDataOverride* CurveData, const FGameplayEffectContextHandle& InEffectContext = FGameplayEffectContextHandle(), float Level = 0.f);
-
-	// Hard Ref to what we modify, this stuff is const and never changes
-	const FGameplayModifierInfo &Info;
-	
-	FAggregatorRef Aggregator;
+	GENERATED_USTRUCT_BODY()
 
-	TSharedPtr<FGameplayEffectSpec> TargetEffectSpec;
+public:
 
-	// returns true if this GameplayEffectSpec can modify things in QualifierContext, returns false otherwise.
-	bool CanModifyInContext(const FModifierQualifier &QualifierContext) const;
-	// returns true if this GameplayEffect can modify Other, false otherwise
-	bool CanModifyModifier(FModifierSpec &Other, const FModifierQualifier &QualifierContext) const;
-	
-	void ApplyModTo(FModifierSpec &Other, bool TakeSnapshot) const;
-	void ExecuteModOn(FModifierSpec &Other) const;
+	/**
+	 * Add a definition to be captured by the owner of the container. Will not add the definition if its exact
+	 * match already exists within the container.
+	 * 
+	 * @param InCaptureDefinition	Definition to capture with
+	 */
+	void AddCaptureDefinition(const FGameplayEffectAttributeCaptureDefinition& InCaptureDefinition);
 
-	FString ToSimpleString() const
-	{
-		return Info.ToSimpleString();
-	}
+	/**
+	 * Capture source or target attributes from the specified component. Should be called by the container's owner.
+	 * 
+	 * @param InAbilitySystemComponent	Component to capture attributes from
+	 * @param InCaptureSource			Whether to capture attributes as source or target
+	 */
+	void CaptureAttributes(class UAbilitySystemComponent* InAbilitySystemComponent, EGameplayEffectAttributeCaptureSource InCaptureSource);
 
-	// Can this GameplayEffect modify the input parameter, based on tags
-	// Returns true if it can modify the input parameter, false otherwise
-	bool AreTagRequirementsSatisfied(const FModifierSpec &ModifierToBeModified) const;
+	// @todo:
+	// API should provide access to querying specs by definition; Wrapper around @todo API for specs above
+	
+	const FAggregator& GetAggregator(const FGameplayEffectAttributeCaptureDefinition& InDefinition) const;
 
-	void PrintAll() const;
-};
+	/** Returns whether the container has at least one spec w/o snapshotted attributes */
+	bool HasNonSnapshottedAttributes() const;
 
-/** Saves list of modified attributes, to use for gameplay cues or later processing */
-USTRUCT()
-struct FGameplayEffectModifiedAttribute
-{
-	GENERATED_USTRUCT_BODY()
+private:
 
-	/** The attribute that has been modified */
+	/** Captured attributes from the source of a gameplay effect */
 	UPROPERTY()
-	FGameplayAttribute Attribute;
+	TArray<FGameplayEffectAttributeCaptureSpec> SourceAttributes;
 
-	/** Total magnitude applied to that attribute */
+	/** Captured attributes from the target of a gameplay effect */
 	UPROPERTY()
-	float TotalMagnitude;
+	TArray<FGameplayEffectAttributeCaptureSpec> TargetAttributes;
 
-	FGameplayEffectModifiedAttribute() : TotalMagnitude(0.0f) {}
+	/** If true, has at least one capture spec that did not request a snapshot */
+	UPROPERTY()
+	bool bHasNonSnapshottedAttributes;
 };
 
 /**
@@ -881,21 +587,20 @@ struct FGameplayEffectSpec
 	GENERATED_USTRUCT_BODY()
 
 	FGameplayEffectSpec()
-		: ModifierLevel( TSharedPtr< FGameplayEffectLevelSpec >( new FGameplayEffectLevelSpec() ) )
-		, Duration(new FAggregator(FGameplayModifierEvaluatedData(0.f, NULL, FActiveGameplayEffectHandle()), SKILL_AGG_DEBUG(TEXT("Uninitialized Duration"))))
-		, Period(new FAggregator(FGameplayModifierEvaluatedData(0.f, NULL, FActiveGameplayEffectHandle()), SKILL_AGG_DEBUG(TEXT("Uninitialized Period"))))
+		: Level(UGameplayEffect::INVALID_LEVEL)
+		, Duration(UGameplayEffect::INSTANT_APPLICATION)
+		, Period(UGameplayEffect::NO_PERIOD)
+		, ChanceToApplyToTarget(1.f)
+		, ChanceToExecuteOnGameplayEffect(1.f)
 		, bTopOfStack(false)
 	{
 		// If we initialize a GameplayEffectSpec with no level object passed in.
 	}
 
-	FGameplayEffectSpec(const UGameplayEffect* InDef, const FGameplayEffectContextHandle& InEffectContext, float Level, const FGlobalCurveDataOverride* CurveData);
+	FGameplayEffectSpec(const UGameplayEffect* InDef, const FGameplayEffectContextHandle& InEffectContext, float Level = UGameplayEffect::INVALID_LEVEL);
 	
 	UPROPERTY()
 	const UGameplayEffect* Def;
-
-	// Replicated	
-	TSharedPtr< FGameplayEffectLevelSpec > ModifierLevel;
 	
 	// Replicated
 	UPROPERTY()
@@ -904,6 +609,10 @@ struct FGameplayEffectSpec
 	// Replicated
 	UPROPERTY()
 	TArray<FGameplayEffectModifiedAttribute> ModifiedAttributes;
+	
+	/** Attributes captured by the spec that are relevant to custom calculations, potentially in owned modifiers, etc.; NOT replicated to clients */
+	UPROPERTY(NotReplicated)
+	FGameplayEffectAttributeCaptureSpecContainer CapturedRelevantAttributes;
 
 	// Looks for an existing modified attribute struct, may return NULL
 	const FGameplayEffectModifiedAttribute* GetModifiedAttribute(const FGameplayAttribute& Attribute) const;
@@ -921,7 +630,6 @@ struct FGameplayEffectSpec
 	float GetChanceToExecuteOnGameplayEffect() const;
 	float GetMagnitude(const FGameplayAttribute &Attribute) const;
 
-
 	EGameplayEffectStackingPolicy::Type GetStackingType() const;
 
 	// other effects that need to be applied to the target if this effect is successful
@@ -931,53 +639,63 @@ struct FGameplayEffectSpec
 	// instantaneous effects should have a duration of UGameplayEffect::INSTANT_APPLICATION
 	// effects that last forever should have a duration of UGameplayEffect::INFINITE_DURATION
 	UPROPERTY()
-	FAggregatorRef	Duration;
+	float Duration;
 
 	// The period in seconds of this effect.
 	// Nonperiodic effects should have a period of UGameplayEffect::NO_PERIOD
 	UPROPERTY()
-	FAggregatorRef	Period;
+	float Period;
 
 	// The chance, in a 0.0-1.0 range, that this GameplayEffect will be applied to the target Attribute or GameplayEffect.
 	UPROPERTY()
-	FAggregatorRef	ChanceToApplyToTarget;
+	float ChanceToApplyToTarget;
 
 	UPROPERTY()
-	FAggregatorRef	ChanceToExecuteOnGameplayEffect;
+	float ChanceToExecuteOnGameplayEffect;
 
 	// This should only be true if this is a stacking effect and at the top of its stack
 	// (FIXME: should this be part of the spec or FActiveGameplayEffect?)
 	UPROPERTY()
 	bool bTopOfStack;
 
-	// The spec needs to own these FModifierSpecs so that other people can keep TSharedPtr to it.
-	// The stuff in this array is OWNED by this spec
+	// Captured Source Tags on GameplayEffectSpec creation.
+	//	-Do we ever not snapshot these tags?
+	UPROPERTY()
+	FGameplayTagContainer	CapturedSourceTags;
 
+	// Tags from the target, captured during execute?
+	// Not sure if this is the right place or not, or how 'contextual' tags get injected (headshot, crit, etc)
+	// Does this overlap with info in the GameplayEffectContextHandle? Should it just all live there?
+	// For now - putting it here so places have one spot to pull from. Need to figure this out long term.
+	UPROPERTY()
+	FGameplayTagContainer	CapturedTargetTags;
+	
 	TArray<FModifierSpec> Modifiers;
+	
+	TArray<FGameplayEffectExecutionSpec>	Executions;
 
-	void MakeUnique();
-
-	void InitModifiers(const FGlobalCurveDataOverride* CurveData, const FGameplayEffectContextHandle& InEffectContext, float Level);
-
-	// returns the number of modifiers applied to the current GameplayEffectSpec by InSpec
-	// returns -1 if the current GameplayEffectSpec prevents InSpec from being applied
-	int32 ApplyModifiersFrom(const FGameplayEffectSpec &InSpec, const FModifierQualifier &QualifierContext);
-
-	int32 ExecuteModifiersFrom(const FGameplayEffectSpec &InSpec, const FModifierQualifier &QualifierContext);
-
-	// returns 1 if the modifier was applied, 0 otherwise
-	int32 ApplyModifier(const FModifierSpec &InMod, const FModifierQualifier &QualifierContext, bool bApplyAsSnapshot);
-
-	// returns true if modifiers applied to this GameplayEffectSpec should be a snapshot of the applied modifiers
-	// returns false if modifiers applied to this GameplayEffectSpec should link to the applied modifiers
-	bool ShouldApplyAsSnapshot(const FModifierQualifier &QualifierContext) const;
+	// GE_REMOVE: this feels strange but need to start standardizing how people get this info.
+	float GetModifierMagnitude(const FModifierSpec& ModSpec) const;
+	void CalculateModifierMagnitudes();
 
 	FString ToSimpleString() const
 	{
 		return FString::Printf(TEXT("%s"), *Def->GetName());
 	}
 
+	void SetLevel(float InLevel);
+
+	float GetLevel() const;
+
 	void PrintAll() const;
+
+	
+
+private:
+
+	// Replicated	
+	UPROPERTY()
+	float Level;	
 };
 
 /**
@@ -996,6 +714,7 @@ struct FActiveGameplayEffect : public FFastArraySerializerItem
 	FActiveGameplayEffect()
 		: StartGameStateTime(0)
 		, StartWorldTime(0.f)
+		, IsInhibited(true)
 	{
 	}
 
@@ -1005,11 +724,8 @@ struct FActiveGameplayEffect : public FFastArraySerializerItem
 		, PredictionKey(InPredictionKey)
 		, StartGameStateTime(InStartGameStateTime)
 		, StartWorldTime(CurrentWorldTime)
+		, IsInhibited(true)
 	{
-		for (FModifierSpec &Mod : Spec.Modifiers)
-		{
-			Mod.Aggregator.Get()->ActiveHandle = InHandle;
-		}
 	}
 
 	FActiveGameplayEffectHandle Handle;
@@ -1027,6 +743,10 @@ struct FActiveGameplayEffect : public FFastArraySerializerItem
 	UPROPERTY(NotReplicated)
 	float StartWorldTime;
 
+	// Not sure if this should replicate or not. If replicated, we may have trouble where IsHibited doesn't appear to change when we do tag checks (because it was previously inhibited, but replication made it inhibited).
+	UPROPERTY()
+	bool IsInhibited;
+
 	FOnActiveGameplayEffectRemoved	OnRemovedDelegate;
 
 	FTimerHandle PeriodHandle;
@@ -1049,6 +769,8 @@ struct FActiveGameplayEffect : public FFastArraySerializerItem
 		return Spec.GetPeriod();
 	}
 
+	void CheckOngoingTagRequirements(const FGameplayTagContainer& OwnerTags, struct FActiveGameplayEffectsContainer& OwningContainer);
+
 	bool CanBeStacked(const FActiveGameplayEffect& Other) const;
 
 	void PrintAll() const;
@@ -1108,30 +830,25 @@ struct FActiveGameplayEffectsContainer : public FFastArraySerializer
 	GENERATED_USTRUCT_BODY();
 
 	friend struct FActiveGameplayEffect;
+	friend class UAbilitySystemComponent;
 
 	FActiveGameplayEffectsContainer() : bNeedToRecalculateStacks(false) {};
 
+	UAbilitySystemComponent* Owner;
+
 	UPROPERTY()
 	TArray< FActiveGameplayEffect >	GameplayEffects;
 
-	UAbilitySystemComponent* Owner;
+	void RegisterWithOwner(UAbilitySystemComponent* Owner);	
 	
 	FActiveGameplayEffect& CreateNewActiveGameplayEffect(const FGameplayEffectSpec &Spec, FPredictionKey InPredictionKey);
 
 	FActiveGameplayEffect* GetActiveGameplayEffect(const FActiveGameplayEffectHandle Handle);
-
-	// returns true if none of the active effects provide immunity to Spec
-	// returns false if one (or more) of the active effects provides immunity to Spec
-	bool ApplyActiveEffectsTo(OUT FGameplayEffectSpec &Spec, const FModifierQualifier &QualifierContext) const;
-
-	void ApplySpecToActiveEffectsAndAttributes(FGameplayEffectSpec &Spec, const FModifierQualifier &QualifierContext);
 		
-	void ExecuteActiveEffectsFrom(FGameplayEffectSpec &Spec, const FModifierQualifier &QualifierContext);
+	void ExecuteActiveEffectsFrom(FGameplayEffectSpec &Spec, FPredictionKey PredictionKey = FPredictionKey() );
 	
 	void ExecutePeriodicGameplayEffect(FActiveGameplayEffectHandle Handle);	// This should not be outward facing to the skill system API, should only be called by the owning AbilitySystemComponent
 
-	void AddDependancyToAttribute(FGameplayAttribute Attribute, const TWeakPtr<FAggregator> InDependant);
-
 	bool RemoveActiveGameplayEffect(FActiveGameplayEffectHandle Handle);
 	
 	float GetGameplayEffectDuration(FActiveGameplayEffectHandle Handle) const;
@@ -1142,6 +859,13 @@ struct FActiveGameplayEffectsContainer : public FFastArraySerializer
 	// returns false if the handle points to an effect that is not in this container or is not applied because of the current stacking rules
 	bool IsGameplayEffectActive(FActiveGameplayEffectHandle Handle) const;
 
+	/**
+	 * Populate the specified capture spec with the data necessary to capture an attribute from the container
+	 * 
+	 * @param OutCaptureSpec	[OUT] Capture spec to populate with captured data
+	 */
+	void CaptureAttributeForGameplayEffect(OUT FGameplayEffectAttributeCaptureSpec& OutCaptureSpec);
+
 	void PrintAllGameplayEffects() const;
 
 	int32 GetNumGameplayEffects() const
@@ -1149,11 +873,6 @@ struct FActiveGameplayEffectsContainer : public FFastArraySerializer
 		return GameplayEffects.Num();
 	}
 
-	float GetGameplayEffectMagnitudeByTag(FActiveGameplayEffectHandle Handle, const FGameplayTag& InTag) const;
-
-	/** Registered as a callback for when a PropertyAggregator is dirty and we need to update the corresponding GameplayAttribute */
-	void OnPropertyAggregatorDirty(FAggregator* Aggregator, FGameplayAttribute Attribute);
-
 	void CheckDuration(FActiveGameplayEffectHandle Handle);
 
 	void StacksNeedToRecalculate();
@@ -1184,25 +903,21 @@ struct FActiveGameplayEffectsContainer : public FFastArraySerializer
 
 	float GetWorldTime() const;
 
-	void OnDurationAggregatorDirty(FAggregator* Aggregator, UAbilitySystemComponent* Owner, FActiveGameplayEffectHandle Handle);
-	
-	FOnGameplayAttributeChange& RegisterGameplayAttributeEvent(FGameplayAttribute Attribute);
-
 	bool HasReceivedEffectWithPredictedKey(FPredictionKey PredictionKey) const;
 
 	bool HasPredictedEffectWithPredictedKey(FPredictionKey PredictionKey) const;
 		
 	void SetBaseAttributeValueFromReplication(FGameplayAttribute Attribute, float BaseBalue);
 
-private:
+	// -------------------------------------------------------------------------------------------
 
-	FTimerHandle StackHandle;
+	FOnGameplayAttributeChange& RegisterGameplayAttributeEvent(FGameplayAttribute Attribute);
 
-	FAggregatorRef& FindOrCreateAttributeAggregator(FGameplayAttribute Attribute);
+	void OnOwnerTagChange(FGameplayTag TagChange, int32 NewCount);
 
-	TMap<FGameplayAttribute, FAggregatorRef> OngoingPropertyEffects;
+private:
 
-	TMap<FGameplayAttribute, FOnGameplayAttributeChange> AttributeChangeDelegates;
+	FTimerHandle StackHandle;
 
 	void InternalUpdateNumericalAttribute(FGameplayAttribute Attribute, float NewValue, const FGameplayEffectModCallbackData* ModData);
 
@@ -1212,8 +927,20 @@ private:
 	bool InternalRemoveActiveGameplayEffect(int32 Idx);
 
 	/** Called both in server side creation and replication creation/deletion */
-	void InternalOnActiveGameplayEffectAdded(const FActiveGameplayEffect& Effect);
+	void InternalOnActiveGameplayEffectAdded(FActiveGameplayEffect& Effect);
 	void InternalOnActiveGameplayEffectRemoved(const FActiveGameplayEffect& Effect);
+
+	// -------------------------------------------------------------------------------------------
+
+	TMap<FGameplayAttribute, FAggregatorRef>		AttributeAggregatorMap;
+
+	TMap<FGameplayAttribute, FOnGameplayAttributeChange> AttributeChangeDelegates;
+
+	TMap<FGameplayTag, TSet<FActiveGameplayEffectHandle> >	ActiveEffectTagDependencies;
+
+	FAggregatorRef& FindOrCreateAttributeAggregator(FGameplayAttribute Attribute);
+
+	void OnAttributeAggregatorDirty(FAggregator* Aggregator, FGameplayAttribute Attribute);
 };
 
 template<>
diff --git a/Engine/Source/Runtime/GameplayAbilities/Public/GameplayEffectAggregator.h b/Engine/Source/Runtime/GameplayAbilities/Public/GameplayEffectAggregator.h
--- a/Engine/Source/Runtime/GameplayAbilities/Public/GameplayEffectAggregator.h
+++ b/Engine/Source/Runtime/GameplayAbilities/Public/GameplayEffectAggregator.h
@@ -2,8 +2,6 @@
 
 #pragma once
 
-#if GE_REFACTOR
-
 #include "GameplayEffectTypes.h"
 
 // #include "GameplayEffectAggregator.generated.h"
@@ -31,24 +29,35 @@ struct GAMEPLAYABILITIES_API FAggregator : public TSharedFromThis<FAggregator>
 {
 	DECLARE_MULTICAST_DELEGATE_OneParam(FOnAggregatorDirty, FAggregator*);
 
-	FAggregator(float InBaseValue=0.f) : BaseValue(InBaseValue) { }
+	FAggregator(float InBaseValue=0.f) : CallbacksDisabled(false), BaseValue(InBaseValue) { }
 
 	void SetBaseValue(float NewBaseValue);
-	void ExecModOnBaseValue(const struct FModifierSpec& Mod);
-
-	static float StaticExecModOnBaseValue(float BaseValue, const struct FModifierSpec& Mod);
+	
+	void ExecModOnBaseValue(TEnumAsByte<EGameplayModOp::Type> ModifierOp, float EvaluatedMagnitude);
+	static float StaticExecModOnBaseValue(float BaseValue, TEnumAsByte<EGameplayModOp::Type> ModifierOp, float EvaluatedMagnitude);
 
 	void AddMod(float EvaluatedData, TEnumAsByte<EGameplayModOp::Type> ModifierOp, const FGameplayTagRequirements*	SourceTagReqs, const FGameplayTagRequirements* TargetTagReqs, FActiveGameplayEffectHandle ActiveHandle = FActiveGameplayEffectHandle());
 	void RemoveMod(FActiveGameplayEffectHandle ActiveHandle);
 
-	FOnAggregatorDirty OnDirty;
+	void DisableCallbacks();
+	void EnabledCallbacks();
 
+	/** Evaluates the Aggregator with the internal base value and given parameters */
 	float Evaluate(const FAggregatorEvaluateParameters& Parameters) const;
 
+	/** Evaluates the Aggregator with an arbitrary base value */
+	float EvaluateWithBase(float InlineBaseValue, const FAggregatorEvaluateParameters& Parameters) const;
+
 	void TakeSnapshotOf(const FAggregator& AggToSnapshot);
 
+	FOnAggregatorDirty OnDirty;
+
 private:
 
+	void BroadcastOnDirty();
+
+	bool CallbacksDisabled;
+
 	float	BaseValue;
 	TArray<FAggregatorMod>	Mods[EGameplayModOp::Max];
 
@@ -68,6 +77,4 @@ struct GAMEPLAYABILITIES_API FAggregatorRef
 	TSharedPtr<FAggregator>	Data;
 
 	void TakeSnapshotOf(const FAggregatorRef& RefToSnapshot);
-};
-
-#endif
\ No newline at end of file
+};
\ No newline at end of file
diff --git a/Engine/Source/Runtime/GameplayAbilities/Public/GameplayEffectTypes.h b/Engine/Source/Runtime/GameplayAbilities/Public/GameplayEffectTypes.h
--- a/Engine/Source/Runtime/GameplayAbilities/Public/GameplayEffectTypes.h
+++ b/Engine/Source/Runtime/GameplayAbilities/Public/GameplayEffectTypes.h
@@ -16,6 +16,7 @@
 #endif
 
 class UAbilitySystemComponent;
+class UGameplayEffect;
 
 struct FGameplayEffectSpec;
 struct FGameplayEffectModCallbackData;
@@ -49,66 +50,6 @@ namespace EGameplayModOp
 	};
 }
 
-/**
- * Tells us what thing a GameplayModifier modifies.
- */
-UENUM(BlueprintType)
-namespace EGameplayMod
-{
-	enum Type
-	{
-		Attribute = 0,		// Modifies this Attributes
-		OutgoingGE,			// Modifies Outgoing Gameplay Effects (that modify this Attribute)
-		IncomingGE,			// Modifies Incoming Gameplay Effects (that modify this Attribute)
-		ActiveGE,			// Modifies currently active Gameplay Effects
-
-		// This must always be at the end
-		Max					UMETA(DisplayName="Invalid")
-	};
-}
-
-/**
- * Tells us what thing a GameplayEffect provides immunity to. This must mirror the values in EGameplayMod
- */
-UENUM(BlueprintType)
-namespace EGameplayImmunity
-{
-	enum Type
-	{
-		None = 0,			// Does not provide immunity
-		OutgoingGE,			// Provides immunity to outgoing GEs
-		IncomingGE,			// Provides immunity to incoming GEs
-		ActiveGE,			// Provides immunity from currently active GEs
-
-		// This must always be at the end
-		Max					UMETA(DisplayName="Invalid")
-	};
-}
-
-static_assert((int32)EGameplayMod::OutgoingGE == (int32)EGameplayImmunity::OutgoingGE, "EGameplayMod::OutgoingGE and EGameplayImmunity::OutgoingGE must match. Did you forget to modify one of them?");
-static_assert((int32)EGameplayMod::IncomingGE == (int32)EGameplayImmunity::IncomingGE, "EGameplayMod::IncomingGE and EGameplayImmunity::IncomingGE must match. Did you forget to modify one of them?");
-static_assert((int32)EGameplayMod::ActiveGE == (int32)EGameplayImmunity::ActiveGE, "EGameplayMod::ActiveGE and EGameplayImmunity::ActiveGE must match. Did you forget to modify one of them?");
-static_assert((int32)EGameplayMod::Max == (int32)EGameplayImmunity::Max, "EGameplayMod and EGameplayImmunity must cover the same cases. Did you forget to modify one of them?");
-
-/**
- * Tells us what a GameplayEffect modifies when being applied to another GameplayEffect
- */
-UENUM(BlueprintType)
-namespace EGameplayModEffect
-{
-	enum Type
-	{
-		Magnitude			= 0x01,		// Modifies magnitude of a GameplayEffect (Always default for Attribute mod)
-		Duration			= 0x02,		// Modifies duration of a GameplayEffect
-		ChanceApplyTarget	= 0x04,		// Modifies chance to apply GameplayEffect to target
-		ChanceExecuteEffect	= 0x08,		// Modifies chance to execute GameplayEffect on GameplayEffect
-		LinkedGameplayEffect= 0x10,		// Adds a linked GameplayEffect to a GameplayEffect
-
-		// This must always be at the end
-		All					= 0xFF		UMETA(DisplayName="Invalid")
-	};
-}
-
 /**
  * Tells us how to handle copying gameplay effect when it is applied.
  *	Default means to use context - e.g, OutgoingGE is are always snapshots, IncomingGE is always Link
@@ -574,122 +515,6 @@ struct TStructOpsTypeTraits<FGameplayEffectContextHandle> : public TStructOpsTyp
 	};
 };
 
-
-// -----------------------------------------------------------
-
-/** 
- * Qualification context for applying modifiers.
- *	For example a Modifier may be setup in data to only apply to OutgoingGE mods.
- *	The FModifierQualifier is the data structure to hold the 'what type of modifier are we applying' data.
- *
- *  This should ideally only hold data that is outside of FGameplayEffectSpec or FGameplayModifierSpec.
- *	For example, specs know what they can and can't modify. We don't need to duplicate that there. 
- *  FModifierQualifier is meant to hold the data that comes from the calling context that is not intrinsic 
- *  to the existing data structures.
- *
- *	This class uses an optional initialization idiom such that you can do things like:
- *		FModifierQualifier().Handle(InHandle).(InType).etc...
- */
-USTRUCT(BlueprintType)
-struct FModifierQualifier
-{
-	GENERATED_USTRUCT_BODY()
-
-	FModifierQualifier()
-		: MyType(EGameplayMod::Max)
-	{
-	}
-
-	// ----------------------------------
-
-	FModifierQualifier& Type(EGameplayMod::Type InType)
-	{
-		MyType = InType;
-		return *this;
-	}
-
-	EGameplayMod::Type Type() const
-	{
-		return MyType;
-	}
-
-	// ----------------------------------
-	// PredictionKey is used by networking/replication to specify that a GameplayEffect has been predictively created/added.
-	
-	FModifierQualifier& PredictionKey(FPredictionKey InPredictionKey)
-	{
-		MyPredictionKey = InPredictionKey;
-		return *this;
-	}
-
-	FPredictionKey PredictionKey() const
-	{
-		return MyPredictionKey;
-	}
-
-	// ----------------------------------
-	// IgnoreHandle - ignore this handle completely. For example when executing an active gameplay effect, we never
-	// want to apply it to itself (either as IncomingGE or activeGE). It is ignored in all contexts.
-	//
-
-	FModifierQualifier& IgnoreHandle(FActiveGameplayEffectHandle InHandle)
-	{
-		MyIgnoreHandle = InHandle;
-		return *this;
-	}
-
-	FActiveGameplayEffectHandle IgnoreHandle() const
-	{
-		return MyIgnoreHandle;
-	}
-
-	// ----------------------------------
-	// ExclusiveTarget - sometimes we need to only apply a modifier to a specific active gameplay effect.
-	// We may only be able to use the handle if there are multiple instances of the same gameplay effect.
-	// This only applies in the context of applying/executing to a target. 'We only modify this active effect'.
-	// ExclusiveTarget is not checked in the context of applying outgoing/incoming GE modifiers to the spec.
-	//
-
-	FModifierQualifier& ExclusiveTarget(FActiveGameplayEffectHandle InHandle)
-	{
-		MyExclusiveTargetHandle = InHandle;
-		return *this;
-	}
-
-	FActiveGameplayEffectHandle ExclusiveTarget() const
-	{
-		return MyExclusiveTargetHandle;
-	}
-
-	// ----------------------------------
-
-	bool TestTarget(FActiveGameplayEffectHandle InHandle) const
-	{
-		if (MyIgnoreHandle.IsValid() && MyIgnoreHandle == InHandle)
-		{
-			return false;
-		}
-
-		if (MyExclusiveTargetHandle.IsValid() && MyExclusiveTargetHandle != InHandle)
-		{
-			return false;
-		}
-
-		return true;
-	}
-
-	FString ToString() const
-	{
-		return EGameplayModToString(MyType);
-	}
-
-private:
-	EGameplayMod::Type	MyType;
-	FActiveGameplayEffectHandle MyIgnoreHandle;		// Do not modify this handle
-	FActiveGameplayEffectHandle MyExclusiveTargetHandle;	// Only modify this handle
-	FPredictionKey MyPredictionKey;
-};
-
 // -----------------------------------------------------------
 
 
@@ -736,6 +561,10 @@ DECLARE_MULTICAST_DELEGATE(FOnActiveGameplayEffectRemoved);
 
 DECLARE_MULTICAST_DELEGATE_TwoParams(FOnGameplayAttributeChange, float ,const FGameplayEffectModCallbackData*);
 
+DECLARE_DELEGATE_RetVal(FGameplayTagContainer, FGetGameplayTags);
+
+DECLARE_DELEGATE_RetVal_OneParam(FOnGameplayEffectTagCountChanged&, FRegisterGameplayTagChangeDelegate, FGameplayTag);
+
 // -----------------------------------------------------------
 
 /** 
@@ -761,5 +590,34 @@ struct FGameplayTagCountContainer
 	TMap<struct FGameplayTag, FOnGameplayEffectTagCountChanged> GameplayTagEventMap;
 	TMap<struct FGameplayTag, int32> GameplayTagCountMap;
 
+	/** This is called when any tag is added new or removed completely (going too or from 0 count). Not called for other count increases (e.g, going from 2-3 count) */
+	FOnGameplayEffectTagCountChanged	OnAnyTagChangeDelegate;
+
 	EGameplayTagMatchType::Type TagContainerType;
+
+private:
+
+	// Fixme: This may not be adding tag parents correctly. The TagContainer version of this function properly adds parent tags
+	void UpdateTagMap_Internal(const struct FGameplayTag& Tag, int32 CountDelta);
+};
+
+// -----------------------------------------------------------
+
+/** Encapsulate require and ignore tags */
+USTRUCT(BlueprintType)
+struct FGameplayTagRequirements
+{
+	GENERATED_USTRUCT_BODY()
+
+	/** All of these tags must be present */
+	UPROPERTY(EditDefaultsOnly, Category = GameplayModifier)
+	FGameplayTagContainer RequireTags;
+
+	/** None of these tags may be present */
+	UPROPERTY(EditDefaultsOnly, Category = GameplayModifier)
+	FGameplayTagContainer IgnoreTags;
+
+	bool	RequirementsMet(FGameplayTagContainer Container) const;
+
+	static FGetGameplayTags	SnapshotTags(FGetGameplayTags TagDelegate);
 };
diff --git a/Engine/Source/Runtime/GameplayAbilities/Public/GameplayModCalculation.h b/Engine/Source/Runtime/GameplayAbilities/Public/GameplayModCalculation.h
--- a/Engine/Source/Runtime/GameplayAbilities/Public/GameplayModCalculation.h
+++ b/Engine/Source/Runtime/GameplayAbilities/Public/GameplayModCalculation.h
@@ -4,75 +4,17 @@
 
 #include "GameplayModCalculation.generated.h"
 
-USTRUCT()
-struct FGameplayModCalculationSpecHandle
-{
-	GENERATED_USTRUCT_BODY()
-
-	FGameplayModCalculationSpecHandle()
-		: Handle(INDEX_NONE)
-	{
-
-	}
-
-	FGameplayModCalculationSpecHandle(int32 InHandle)
-		: Handle(InHandle)
-	{
-
-	}
-
-	bool IsValid() const
-	{
-		return Handle != INDEX_NONE;
-	}
-
-	static FGameplayModCalculationSpecHandle GenerateNewHandle();
-
-	bool operator==(const FGameplayModCalculationSpecHandle& Other) const
-	{
-		return Handle == Other.Handle;
-	}
+// Forward declarations
+struct FGameplayEffectSpec;
+class UAbilitySystemComponent;
 
-	bool operator!=(const FGameplayModCalculationSpecHandle& Other) const
-	{
-		return Handle != Other.Handle;
-	}
-
-	friend uint32 GetTypeHash(const FGameplayModCalculationSpecHandle& InHandle)
-	{
-		return InHandle.Handle;
-	}
-
-	FString ToString() const
-	{
-		return FString::Printf(TEXT("%d"), Handle);
-	}
-
-private:
-
-	UPROPERTY()
-	int32 Handle;
-};
-
-/**
- * UGameplayModCalculation
- *	@todo document
- */
-UCLASS(BlueprintType)
-class GAMEPLAYABILITIES_API UGameplayModCalculation : public UDataAsset
+UCLASS(BlueprintType, Blueprintable, Abstract)
+class GAMEPLAYABILITIES_API UGameplayModCalculation : public UObject
 {
 
 public:
 	GENERATED_UCLASS_BODY()
 
-	/** */
-	UPROPERTY(EditDefaultsOnly, Category=Calculation)
-	TArray<FGameplayAttribute> RelevantSourceAttributes;
-
-	/** */
-	UPROPERTY(EditDefaultsOnly, Category=Calculation)
-	TArray<FGameplayAttribute> RelevantTargetAttributes;
-
-	FGameplayModCalculationSpecHandle GetContextForSource(const FGameplayEffectSpec& Spec);
-	FGameplayModCalculationSpecHandle GetContextForTarget(const FGameplayEffectSpec& Spec);
+	UFUNCTION(BlueprintNativeEvent, Category="Calculation")
+	float CalculateMagnitude(const FGameplayEffectSpec& Spec) const;
 };
\ No newline at end of file
