Gameplay ability fixes

* 'GameplayAbilityTargetActor' now checks the ability's Avatar when deciding it's net relevancy.

* Add option to reset 'FGameplayEffectContext::HitResult' so that the same gameplay effect can have multiple hit results set on it.

diff --git a/Engine/Source/Runtime/GameplayAbilities/Private/Abilities/GameplayAbilityTargetActor.cpp b/Engine/Source/Runtime/GameplayAbilities/Private/Abilities/GameplayAbilityTargetActor.cpp
--- a/Engine/Source/Runtime/GameplayAbilities/Private/Abilities/GameplayAbilityTargetActor.cpp
+++ b/Engine/Source/Runtime/GameplayAbilities/Private/Abilities/GameplayAbilityTargetActor.cpp
@@ -98,6 +98,14 @@ bool AGameplayAbilityTargetActor::IsNetRelevantFor(class APlayerController* Real
 		return false;
 	}
 
+	const FGameplayAbilityActorInfo* ActorInfo = (OwningAbility ? OwningAbility->GetCurrentActorInfo() : NULL);
+	AActor* Avatar = (ActorInfo ? ActorInfo->AvatarActor.Get() : NULL);
+
+	if (Avatar)
+	{
+		return Avatar->IsNetRelevantFor(RealViewer, Viewer, SrcLocation);
+	}
+
 	return Super::IsNetRelevantFor(RealViewer, Viewer, SrcLocation);
 }
 
diff --git a/Engine/Source/Runtime/GameplayAbilities/Private/GameplayEffectTypes.cpp b/Engine/Source/Runtime/GameplayAbilities/Private/GameplayEffectTypes.cpp
--- a/Engine/Source/Runtime/GameplayAbilities/Private/GameplayEffectTypes.cpp
+++ b/Engine/Source/Runtime/GameplayAbilities/Private/GameplayEffectTypes.cpp
@@ -25,8 +25,13 @@ void FGameplayEffectContext::AddInstigator(class AActor *InInstigator, class AAc
 	}
 }
 
-void FGameplayEffectContext::AddHitResult(const FHitResult InHitResult)
+void FGameplayEffectContext::AddHitResult(const FHitResult InHitResult, bool bReset)
 {
+	if (bReset && HitResult.IsValid())
+	{
+		HitResult.Reset();
+	}
+
 	check(!HitResult.IsValid());
 	HitResult = TSharedPtr<FHitResult>(new FHitResult(InHitResult));
 }
diff --git a/Engine/Source/Runtime/GameplayAbilities/Public/GameplayEffectTypes.h b/Engine/Source/Runtime/GameplayAbilities/Public/GameplayEffectTypes.h
--- a/Engine/Source/Runtime/GameplayAbilities/Public/GameplayEffectTypes.h
+++ b/Engine/Source/Runtime/GameplayAbilities/Public/GameplayEffectTypes.h
@@ -214,7 +214,7 @@ struct GAMEPLAYABILITIES_API FGameplayEffectContext
 		return InstigatorAbilitySystemComponent;
 	}
 
-	virtual void AddHitResult(const FHitResult InHitResult);
+	virtual void AddHitResult(const FHitResult InHitResult, bool bReset = false);
 
 	virtual const FHitResult* GetHitResult() const
 	{
@@ -415,11 +415,11 @@ struct FGameplayEffectContextHandle
 		return false;
 	}
 
-	void AddHitResult(const FHitResult InHitResult)
+	void AddHitResult(const FHitResult InHitResult, bool bReset = false)
 	{
 		if (IsValid())
 		{
-			Data->AddHitResult(InHitResult);
+			Data->AddHitResult(InHitResult, bReset);
 		}
 	}
 
