Check that targets have the required attributes before trying to apply a GameplayEffect

diff --git a/Engine/Source/Runtime/GameplayAbilities/Private/AbilitySystemComponent.cpp b/Engine/Source/Runtime/GameplayAbilities/Private/AbilitySystemComponent.cpp
--- a/Engine/Source/Runtime/GameplayAbilities/Private/AbilitySystemComponent.cpp
+++ b/Engine/Source/Runtime/GameplayAbilities/Private/AbilitySystemComponent.cpp
@@ -111,6 +111,11 @@ const UAttributeSet* UAbilitySystemComponent::GetAttributeSubobject(const TSubcl
 	return NULL;
 }
 
+bool UAbilitySystemComponent::HasAttributeSetForAttribute(FGameplayAttribute Attribute) const
+{
+	return (GetAttributeSubobject(Attribute.GetAttributeSetClass()) != nullptr);
+}
+
 void UAbilitySystemComponent::OnRegister()
 {
 	Super::OnRegister();
@@ -341,11 +346,23 @@ FActiveGameplayEffectHandle UAbilitySystemComponent::ApplyGameplayEffectSpecToSe
 	// Effects with other effects may be a mix so go with non-predictive
 	check((PredictionKey.IsValidKey() == false) || (Spec.GetPeriod() == UGameplayEffect::NO_PERIOD));
 
+	// Check Tag requirements
 	if (!HasNetworkAuthorityToApplyGameplayEffect(PredictionKey))
 	{
 		return FActiveGameplayEffectHandle();
 	}
 
+	// Check AttributeSet requirements: do we have everything this GameplayEffectSpec expects?
+	// We may want to cache this off in some way to make the runtime check quicker.
+	// We also need to handle things in the execution list
+	for (const FModifierSpec& Mod : Spec.Modifiers)
+	{
+		if (HasAttributeSetForAttribute(Mod.Info.Attribute) == false)
+		{
+			return FActiveGameplayEffectHandle();
+		}
+	}
+
 	// check if the effect being applied actually succeeds
 	float ChanceToApply = Spec.GetChanceToApplyToTarget();
 	if ((ChanceToApply < 1.f - SMALL_NUMBER) && (FMath::FRand() > ChanceToApply))
diff --git a/Engine/Source/Runtime/GameplayAbilities/Public/AbilitySystemComponent.h b/Engine/Source/Runtime/GameplayAbilities/Public/AbilitySystemComponent.h
--- a/Engine/Source/Runtime/GameplayAbilities/Public/AbilitySystemComponent.h
+++ b/Engine/Source/Runtime/GameplayAbilities/Public/AbilitySystemComponent.h
@@ -93,6 +93,8 @@ class GAMEPLAYABILITIES_API UAbilitySystemComponent : public UActorComponent, pu
 		return Subobject;
 	}
 
+	bool HasAttributeSetForAttribute(FGameplayAttribute Attribute) const;
+
 	const UAttributeSet* InitStats(TSubclassOf<class UAttributeSet> Attributes, const UDataTable* DataTable);
 
 	UFUNCTION(BlueprintCallable, Category="Skills", meta=(FriendlyName="InitStats"))
