CIS fix

diff --git a/Engine/Source/Editor/UnrealEd/Private/SSkeletonWidget.cpp b/Engine/Source/Editor/UnrealEd/Private/SSkeletonWidget.cpp
--- a/Engine/Source/Editor/UnrealEd/Private/SSkeletonWidget.cpp
+++ b/Engine/Source/Editor/UnrealEd/Private/SSkeletonWidget.cpp
@@ -10,6 +10,8 @@
 #include "Animation/Rig.h"
 #include "Editor/AnimGraph/Classes/AnimPreviewInstance.h"
 #include "SceneViewport.h"
+#include "NotificationManager.h"
+#include "SNotificationList.h"
 
 #define LOCTEXT_NAMESPACE "SkeletonWidget"
 
diff --git a/Engine/Source/Runtime/GameplayAbilities/Public/GameplayEffect.h b/Engine/Source/Runtime/GameplayAbilities/Public/GameplayEffect.h
--- a/Engine/Source/Runtime/GameplayAbilities/Public/GameplayEffect.h
+++ b/Engine/Source/Runtime/GameplayAbilities/Public/GameplayEffect.h
@@ -587,12 +587,12 @@ struct FGameplayEffectSpec
 	GENERATED_USTRUCT_BODY()
 
 	FGameplayEffectSpec()
-		: Level(UGameplayEffect::INVALID_LEVEL)
-		, Duration(UGameplayEffect::INSTANT_APPLICATION)
+		: Duration(UGameplayEffect::INSTANT_APPLICATION)
 		, Period(UGameplayEffect::NO_PERIOD)
 		, ChanceToApplyToTarget(1.f)
 		, ChanceToExecuteOnGameplayEffect(1.f)
 		, bTopOfStack(false)
+		, Level(UGameplayEffect::INVALID_LEVEL)
 	{
 		// If we initialize a GameplayEffectSpec with no level object passed in.
 	}
