calculate modifier magnitude on non duration based executed effects

diff --git a/Engine/Source/Runtime/GameplayAbilities/Private/GameplayEffect.cpp b/Engine/Source/Runtime/GameplayAbilities/Private/GameplayEffect.cpp
--- a/Engine/Source/Runtime/GameplayAbilities/Private/GameplayEffect.cpp
+++ b/Engine/Source/Runtime/GameplayAbilities/Private/GameplayEffect.cpp
@@ -553,6 +553,7 @@ void FActiveGameplayEffectsContainer::ExecuteActiveEffectsFrom(FGameplayEffectSp
 	Spec.CapturedSourceTags.RemoveAllTags();
 	Owner->GetOwnedGameplayTags(Spec.CapturedSourceTags);
 
+	Spec.CalculateModifierMagnitudes();
 
 	// ------------------------------------------------------
 	//	Modifiers
