remove optimize off pragma

diff --git a/Engine/Source/Runtime/GameplayAbilities/Private/AbilitySystemComponent.cpp b/Engine/Source/Runtime/GameplayAbilities/Private/AbilitySystemComponent.cpp
--- a/Engine/Source/Runtime/GameplayAbilities/Private/AbilitySystemComponent.cpp
+++ b/Engine/Source/Runtime/GameplayAbilities/Private/AbilitySystemComponent.cpp
@@ -13,8 +13,6 @@
 #include "UObjectToken.h"
 #include "MapErrors.h"
 
-#pragma optimize( "", off )
-
 DEFINE_LOG_CATEGORY(LogAbilitySystemComponent);
 
 #define LOCTEXT_NAMESPACE "AbilitySystemComponent"
