Fix CIS

diff --git a/Engine/Source/Runtime/GameplayAbilities/Private/GameplayEffect.cpp b/Engine/Source/Runtime/GameplayAbilities/Private/GameplayEffect.cpp
--- a/Engine/Source/Runtime/GameplayAbilities/Private/GameplayEffect.cpp
+++ b/Engine/Source/Runtime/GameplayAbilities/Private/GameplayEffect.cpp
@@ -10,8 +10,6 @@
 #include "AbilitySystemInterface.h"
 #include "GameplayModCalculation.h"
 
-#pragma optimize( "", off )
-
 const float UGameplayEffect::INFINITE_DURATION = -1.f;
 const float UGameplayEffect::INSTANT_APPLICATION = 0.f;
 const float UGameplayEffect::NO_PERIOD = 0.f;
