GameplayEffects refactor 2
-Setup Duration mods
-Fixed some issues with gameplayprediction

diff --git a/Engine/Source/Editor/GameplayAbilitiesEditor/Private/AttributeDetails.cpp b/Engine/Source/Editor/GameplayAbilitiesEditor/Private/AttributeDetails.cpp
--- a/Engine/Source/Editor/GameplayAbilitiesEditor/Private/AttributeDetails.cpp
+++ b/Engine/Source/Editor/GameplayAbilitiesEditor/Private/AttributeDetails.cpp
@@ -10,6 +10,7 @@
 #include "SSearchBox.h"
 #include "STextComboBox.h"
 #include "GameplayEffect.h"
+#include "AbilitySystemComponent.h"
 #include "GameplayEffectExtension.h"
 
 #define LOCTEXT_NAMESPACE "AttributeDetailsCustomization"
@@ -105,7 +106,7 @@ void FAttributePropertyDetails::CustomizeHeader( TSharedRef<IPropertyHandle> Str
 
 				for (TFieldIterator<UProperty> PropertyIt(Class, EFieldIteratorFlags::ExcludeSuper); PropertyIt; ++PropertyIt)
 				{
-					UProperty *Property = *PropertyIt;
+					UProperty* Property = *PropertyIt;
 
 					if (!FilterMetaStr.IsEmpty() && Property->HasMetaData(*FilterMetaStr))
 					{
@@ -121,6 +122,23 @@ void FAttributePropertyDetails::CustomizeHeader( TSharedRef<IPropertyHandle> Str
 					PropertiesToAdd.Add(Property);
 				}
 			}
+
+			// UAbilitySystemComponent can add 'system' attributes
+			if (Class->IsChildOf(UAbilitySystemComponent::StaticClass()) && !FKismetEditorUtilities::IsClassABlueprintSkeleton(Class))
+			{
+				for (TFieldIterator<UProperty> PropertyIt(Class, EFieldIteratorFlags::ExcludeSuper); PropertyIt; ++PropertyIt)
+				{
+					UProperty* Property = *PropertyIt;
+
+					// SystemAttributes have to be explicitly tagged
+					if (Property->HasMetaData(TEXT("SystemGameplayAttribute")) == false)
+					{
+						continue;
+					}
+
+					PropertiesToAdd.Add(Property);
+				}
+			}
 		}
 	}
 
diff --git a/Engine/Source/Editor/GameplayAbilitiesEditor/Private/K2Node_LatentAbilityCall.cpp b/Engine/Source/Editor/GameplayAbilitiesEditor/Private/K2Node_LatentAbilityCall.cpp
--- a/Engine/Source/Editor/GameplayAbilitiesEditor/Private/K2Node_LatentAbilityCall.cpp
+++ b/Engine/Source/Editor/GameplayAbilitiesEditor/Private/K2Node_LatentAbilityCall.cpp
@@ -267,6 +267,9 @@ void UK2Node_LatentAbilityCall::CreatePinsForClass(UClass* InClass)
 				check(bDefaultValueSet);
 				K2Schema->TrySetDefaultValue(*Pin, DefaultValueAsString);
 			}
+
+			// Copy tooltip from the property.
+			K2Schema->ConstructBasicPinTooltip(*Pin, Property->GetToolTipText(), Pin->PinToolTip);
 		}
 	}
 }
diff --git a/Engine/Source/Runtime/GameplayAbilities/Private/Abilities/GameplayAbility.cpp b/Engine/Source/Runtime/GameplayAbilities/Private/Abilities/GameplayAbility.cpp
--- a/Engine/Source/Runtime/GameplayAbilities/Private/Abilities/GameplayAbility.cpp
+++ b/Engine/Source/Runtime/GameplayAbilities/Private/Abilities/GameplayAbility.cpp
@@ -475,12 +475,28 @@ USkeletalMeshComponent* UGameplayAbility::GetOwningComponentFromActorInfo() cons
 	return NULL;
 }
 
-FGameplayEffectSpecHandle UGameplayAbility::GetOutgoingGameplayEffectSpec(UGameplayEffect* GameplayEffect, float Level) const
+FGameplayEffectSpecHandle UGameplayAbility::GetOutgoingGameplayEffectSpec(const UGameplayEffect* GameplayEffect, float Level) const
 {
 	check(CurrentActorInfo && CurrentActorInfo->AbilitySystemComponent.IsValid());
-	return CurrentActorInfo->AbilitySystemComponent->GetOutgoingSpec(GameplayEffect, Level);
+	return GetOutgoingGameplayEffectSpec(CurrentSpecHandle, CurrentActorInfo, CurrentActivationInfo, GameplayEffect, Level);
 }
 
+FGameplayEffectSpecHandle UGameplayAbility::GetOutgoingGameplayEffectSpec(const FGameplayAbilitySpecHandle Handle, const FGameplayAbilityActorInfo* ActorInfo, const FGameplayAbilityActivationInfo ActivationInfo, const UGameplayEffect* GameplayEffect, float Level) const
+{
+	check(ActorInfo);
+
+	FGameplayEffectSpecHandle NewHandle = ActorInfo->AbilitySystemComponent->GetOutgoingSpec(GameplayEffect, Level, GetEffectContext(ActorInfo));
+	if (NewHandle.IsValid())
+	{
+		ApplyAbilityTagsToGameplayEffectSpec(*NewHandle.Data.Get());
+	}
+	return NewHandle;
+}
+
+void UGameplayAbility::ApplyAbilityTagsToGameplayEffectSpec(FGameplayEffectSpec& Spec) const
+{
+	Spec.CapturedSourceTags.GetSpecTags().AppendTags(AbilityTags);
+}
 
 /** Fixme: Naming is confusing here */
 
@@ -736,11 +752,15 @@ FActiveGameplayEffectHandle UGameplayAbility::K2_ApplyGameplayEffectToOwner(cons
 	return ApplyGameplayEffectToOwner(CurrentSpecHandle, CurrentActorInfo, CurrentActivationInfo, GameplayEffect, GameplayEffectLevel);
 }
 
-FActiveGameplayEffectHandle UGameplayAbility::ApplyGameplayEffectToOwner(const FGameplayAbilitySpecHandle Handle, const FGameplayAbilityActorInfo* ActorInfo, const FGameplayAbilityActivationInfo ActivationInfo, const UGameplayEffect* GameplayEffect, int32 GameplayEffectLevel)
+FActiveGameplayEffectHandle UGameplayAbility::ApplyGameplayEffectToOwner(const FGameplayAbilitySpecHandle Handle, const FGameplayAbilityActorInfo* ActorInfo, const FGameplayAbilityActivationInfo ActivationInfo, const UGameplayEffect* GameplayEffect, float GameplayEffectLevel)
 {
 	if (ActivationInfo.ActivationMode == EGameplayAbilityActivationMode::Authority || ActivationInfo.ActivationMode == EGameplayAbilityActivationMode::Predicting)
 	{
-		return ActorInfo->AbilitySystemComponent->ApplyGameplayEffectToSelf(GameplayEffect, 1.f, GetEffectContext(ActorInfo), ActivationInfo.GetPredictionKeyForNewAction());
+		FGameplayEffectSpecHandle SpecHandle = GetOutgoingGameplayEffectSpec(Handle, ActorInfo, ActivationInfo, GameplayEffect, GameplayEffectLevel);
+		if (SpecHandle.IsValid())
+		{
+			return ActorInfo->AbilitySystemComponent->ApplyGameplayEffectSpecToSelf(*SpecHandle.Data.Get(), ActivationInfo.GetPredictionKeyForNewAction());
+		}
 	}
 
 	// We cannot apply GameplayEffects in this context. Return an empty handle.
diff --git a/Engine/Source/Runtime/GameplayAbilities/Private/AbilitySystemComponent.cpp b/Engine/Source/Runtime/GameplayAbilities/Private/AbilitySystemComponent.cpp
--- a/Engine/Source/Runtime/GameplayAbilities/Private/AbilitySystemComponent.cpp
+++ b/Engine/Source/Runtime/GameplayAbilities/Private/AbilitySystemComponent.cpp
@@ -111,7 +111,7 @@ const UAttributeSet* UAbilitySystemComponent::GetAttributeSubobject(const TSubcl
 
 bool UAbilitySystemComponent::HasAttributeSetForAttribute(FGameplayAttribute Attribute) const
 {
-	return (GetAttributeSubobject(Attribute.GetAttributeSetClass()) != nullptr);
+	return (Attribute.IsSystemAttribute() || GetAttributeSubobject(Attribute.GetAttributeSetClass()) != nullptr);
 }
 
 void UAbilitySystemComponent::OnRegister()
@@ -151,20 +151,32 @@ void UAbilitySystemComponent::SetNumericAttribute(const FGameplayAttribute &Attr
 
 float UAbilitySystemComponent::GetNumericAttribute(const FGameplayAttribute &Attribute)
 {
+	if(Attribute.IsSystemAttribute())
+	{
+		return 0.f;
+	}
+
 	const UAttributeSet* AttributeSet = GetAttributeSubobjectChecked(Attribute.GetAttributeSetClass());
 	return Attribute.GetNumericValueChecked(AttributeSet);
 }
 
-FGameplayEffectSpecHandle UAbilitySystemComponent::GetOutgoingSpec(UGameplayEffect* GameplayEffect, float Level) const
+FGameplayEffectSpecHandle UAbilitySystemComponent::GetOutgoingSpec(const UGameplayEffect* GameplayEffect, float Level, FGameplayEffectContextHandle Context) const
 {
 	SCOPE_CYCLE_COUNTER(STAT_GetOutgoingSpec);
-	// Fixme: we should build a map and cache these off. We can invalidate the map when an OutgoingGE modifier is applied or removed from us.
+	if (Context.IsValid() == false)
+	{
+		Context = GetEffectContext();
+	}
 
-	// By default use the owner and avatar as the instigator and causer
-	FGameplayEffectSpec* NewSpec = new FGameplayEffectSpec(GameplayEffect, GetEffectContext(), Level);
+	FGameplayEffectSpec* NewSpec = new FGameplayEffectSpec(GameplayEffect, Context, Level);
 	return FGameplayEffectSpecHandle(NewSpec);
 }
 
+FGameplayEffectSpecHandle UAbilitySystemComponent::GetOutgoingSpec(const UGameplayEffect* GameplayEffect, float Level) const
+{
+	return GetOutgoingSpec(GameplayEffect, Level, GetEffectContext());
+}
+
 FGameplayEffectContextHandle UAbilitySystemComponent::GetEffectContext() const
 {
 	FGameplayEffectContextHandle Context = FGameplayEffectContextHandle(UAbilitySystemGlobals::Get().AllocGameplayEffectContext());
@@ -240,7 +252,7 @@ void UAbilitySystemComponent::CaptureAttributeForGameplayEffect(OUT FGameplayEff
 {
 	// Verify the capture is happening on an attribute the component actually has a set for; if not, can't capture the value
 	const FGameplayAttribute& AttributeToCapture = OutCaptureSpec.BackingDefinition.AttributeToCapture;
-	if (AttributeToCapture.IsValid() && GetAttributeSubobject(AttributeToCapture.GetAttributeSetClass()))
+	if (AttributeToCapture.IsValid() && (AttributeToCapture.IsSystemAttribute() || GetAttributeSubobject(AttributeToCapture.GetAttributeSetClass())))
 	{
 		ActiveGameplayEffects.CaptureAttributeForGameplayEffect(OutCaptureSpec);
 	}
@@ -261,6 +273,32 @@ FOnGameplayAttributeChange& UAbilitySystemComponent::RegisterGameplayAttributeEv
 	return ActiveGameplayEffects.RegisterGameplayAttributeEvent(Attribute);
 }
 
+UProperty* UAbilitySystemComponent::GetOutgoingDurationProperty()
+{
+	static UProperty* DurationProperty = FindFieldChecked<UProperty>(UAbilitySystemComponent::StaticClass(), GET_MEMBER_NAME_CHECKED(UAbilitySystemComponent, OutgoingDuration));
+	return DurationProperty;
+}
+
+UProperty* UAbilitySystemComponent::GetIncomingDurationProperty()
+{
+	static UProperty* DurationProperty = FindFieldChecked<UProperty>(UAbilitySystemComponent::StaticClass(), GET_MEMBER_NAME_CHECKED(UAbilitySystemComponent, IncomingDuration));
+	return DurationProperty;
+}
+
+const FGameplayEffectAttributeCaptureDefinition& UAbilitySystemComponent::GetOutgoingDurationCapture()
+{
+	// We will just always take snapshots of the source's duration mods
+	static FGameplayEffectAttributeCaptureDefinition OutgoingDurationCapture(GetOutgoingDurationProperty(), EGameplayEffectAttributeCaptureSource::Source, true);
+	return OutgoingDurationCapture;
+
+}
+const FGameplayEffectAttributeCaptureDefinition& UAbilitySystemComponent::GetIncomingDurationCapture()
+{
+	// Never take snapshots of the target's duration mods: we are going to evaluate this on apply only.
+	static FGameplayEffectAttributeCaptureDefinition IncomingDurationCapture(GetIncomingDurationProperty(), EGameplayEffectAttributeCaptureSource::Target, false);
+	return IncomingDurationCapture;
+}
+
 // ------------------------------------------------------------------------
 
 void UAbilitySystemComponent::GetOwnedGameplayTags(FGameplayTagContainer& TagContainer) const
@@ -340,16 +378,27 @@ FActiveGameplayEffectHandle UAbilitySystemComponent::ApplyGameplayEffectSpecToTa
 
 FActiveGameplayEffectHandle UAbilitySystemComponent::ApplyGameplayEffectSpecToSelf(OUT FGameplayEffectSpec &Spec, FPredictionKey PredictionKey)
 {
-	// Temp, only non instant, non periodic GEs can be predictive
-	// Effects with other effects may be a mix so go with non-predictive
-	check((PredictionKey.IsValidKey() == false) || (Spec.GetPeriod() == UGameplayEffect::NO_PERIOD));
-
-	// Check Tag requirements
+	// Check Network Authority
 	if (!HasNetworkAuthorityToApplyGameplayEffect(PredictionKey))
 	{
 		return FActiveGameplayEffectHandle();
 	}
 
+	// Don't allow prediciton of periodic effects
+	if(PredictionKey.IsValidKey() && Spec.GetPeriod() > 0.f)
+	{
+		if(IsOwnerActorAuthoritative())
+		{
+			// Server continue with invliad prediction key
+			PredictionKey = FPredictionKey();
+		}
+		else
+		{
+			// Client just return now
+			return FActiveGameplayEffectHandle();
+		}
+	}
+
 	// Check AttributeSet requirements: do we have everything this GameplayEffectSpec expects?
 	// We may want to cache this off in some way to make the runtime check quicker.
 	// We also need to handle things in the execution list
@@ -400,6 +449,7 @@ FActiveGameplayEffectHandle UAbilitySystemComponent::ApplyGameplayEffectSpecToSe
 			{
 				ActiveGameplayEffects.StacksNeedToRecalculate();
 			}
+
 			FActiveGameplayEffect& NewActiveEffect = ActiveGameplayEffects.CreateNewActiveGameplayEffect(Spec, PredictionKey);
 			MyHandle = NewActiveEffect.Handle;
 			OurCopyOfSpec = &NewActiveEffect.Spec;
@@ -411,27 +461,15 @@ FActiveGameplayEffectHandle UAbilitySystemComponent::ApplyGameplayEffectSpecToSe
 			OurCopyOfSpec = StackSpec.Get();
 		}
 
-		// GE_REMOVE: Pretty sure this completely goes away? 
-
-		// Do a 1st order copy of the spec so that we can modify it
-		// (the one passed in is owned by the caller, we can't apply our incoming GEs to it)
-		// Note that at this point the spec has a bunch of modifiers. Those modifiers may
-		// have other modifiers. THOSE modifiers may or may not be copies of whatever.
-		//
-		// In theory, we don't modify 2nd order modifiers after they are 'attached'
-		// Long complex chains can be created but we never say 'Modify a GE that is modding another GE'
-		// 
-		// 
-		// OurCopyOfSpec->MakeUnique();
-
 		// if necessary add a modifier to OurCopyOfSpec to force it to have an infinite duration
 		if (bTreatAsInfiniteDuration)
 		{			
 			// This should just be a straight set of the duration float now
+			OurCopyOfSpec->SetDuration(UGameplayEffect::INFINITE_DURATION);
 		}
 	}
 
-	// Capture Our attributes (this may snapshot or link to our AttributeAggregators)
+
 	OurCopyOfSpec->CapturedRelevantAttributes.CaptureAttributes(this, EGameplayEffectAttributeCaptureSource::Target);
 	
 	// We still probably want to apply tags and stuff even if instant?
@@ -525,7 +563,7 @@ void UAbilitySystemComponent::InvokeGameplayCueEvent(const FGameplayEffectSpec &
 	float ExecuteLevel = Spec.GetLevel();
 
 	FGameplayCueParameters CueParameters;
-	CueParameters.EffectContext = Spec.EffectContext;
+	CueParameters.EffectContext = Spec.GetContext();
 
 	for (FGameplayEffectCue CueInfo : Spec.Def->GameplayCues)
 	{
@@ -549,9 +587,9 @@ void UAbilitySystemComponent::InvokeGameplayCueEvent(const FGameplayEffectSpec &
 
 		UAbilitySystemGlobals::Get().GetGameplayCueManager()->HandleGameplayCues(ActorAvatar, CueInfo.GameplayCueTags, EventType, CueParameters);
 
-		if (DebugGameplayCues && Spec.EffectContext.GetHitResult())
+		if (DebugGameplayCues && Spec.GetContext().GetHitResult())
 		{
-			DrawDebugSphere(GetWorld(), Spec.EffectContext.GetHitResult()->Location, 30.f, 32, FColor(255, 0, 0), true, 30.f);
+			DrawDebugSphere(GetWorld(), Spec.GetContext().GetHitResult()->Location, 30.f, 32, FColor(255, 0, 0), true, 30.f);
 			ABILITY_LOG(Warning, TEXT("   %s"), *CueInfo.GameplayCueTags.ToString());
 		}
 	}
diff --git a/Engine/Source/Runtime/GameplayAbilities/Private/AbilitySystemComponent_Abilities.cpp b/Engine/Source/Runtime/GameplayAbilities/Private/AbilitySystemComponent_Abilities.cpp
--- a/Engine/Source/Runtime/GameplayAbilities/Private/AbilitySystemComponent_Abilities.cpp
+++ b/Engine/Source/Runtime/GameplayAbilities/Private/AbilitySystemComponent_Abilities.cpp
@@ -975,7 +975,7 @@ void UAbilitySystemComponent::AbilityInputPressed(int32 InputID)
 							Instance->InputPressed(Spec.Handle, AbilityActorInfo.Get(), Spec.ActivationInfo);
 						}						
 					}
-					FPredictionKey NewKey = FPredictionKey::CreateNewPredictionKey();
+					FPredictionKey NewKey = FPredictionKey::CreateNewPredictionKey(this);
 					FScopedPredictionWindow ScopedPrediction(this, NewKey);
 					//We don't require (AbilityKeyPressCallbacks.IsBound() || AbilityKeyReleaseCallbacks.IsBound())), because we don't necessarily know now if we'll need this data later.
 					if (GetOwnerRole() != ROLE_Authority)
@@ -1005,7 +1005,7 @@ void UAbilitySystemComponent::AbilityInputReleased(int32 InputID)
 			if (Spec.Ability)
 			{
 				AbilitySpecInputReleased(Spec);
-				FPredictionKey NewKey = FPredictionKey::CreateNewPredictionKey();
+				FPredictionKey NewKey = FPredictionKey::CreateNewPredictionKey(this);
 				FScopedPredictionWindow ScopedPrediction(this, NewKey);
 				//We don't require (AbilityKeyPressCallbacks.IsBound() || AbilityKeyReleaseCallbacks.IsBound())), because we don't necessarily know now if we'll need this data later.
 				if (GetOwnerRole() != ROLE_Authority)
@@ -1063,7 +1063,7 @@ void UAbilitySystemComponent::AbilitySpecInputReleased(FGameplayAbilitySpec& Spe
 
 void UAbilitySystemComponent::InputConfirm()
 {
-	FPredictionKey NewKey = FPredictionKey::CreateNewPredictionKey();
+	FPredictionKey NewKey = FPredictionKey::CreateNewPredictionKey(this);
 	FScopedPredictionWindow ScopedPrediction(this, NewKey);
 
 	if (GetOwnerRole() != ROLE_Authority && ConfirmCallbacks.IsBound())
@@ -1077,7 +1077,7 @@ void UAbilitySystemComponent::InputConfirm()
 
 void UAbilitySystemComponent::InputCancel()
 {
-	FPredictionKey NewKey = FPredictionKey::CreateNewPredictionKey();
+	FPredictionKey NewKey = FPredictionKey::CreateNewPredictionKey(this);
 	FScopedPredictionWindow ScopedPrediction(this, NewKey);
 
 	if (GetOwnerRole() != ROLE_Authority && CancelCallbacks.IsBound())
diff --git a/Engine/Source/Runtime/GameplayAbilities/Private/AttributeSet.cpp b/Engine/Source/Runtime/GameplayAbilities/Private/AttributeSet.cpp
--- a/Engine/Source/Runtime/GameplayAbilities/Private/AttributeSet.cpp
+++ b/Engine/Source/Runtime/GameplayAbilities/Private/AttributeSet.cpp
@@ -23,6 +23,11 @@ float FGameplayAttribute::GetNumericValueChecked(const UAttributeSet* Src) const
 	return NumericProperty->GetFloatingPointPropertyValue(ValuePtr);
 }
 
+bool FGameplayAttribute::IsSystemAttribute() const
+{
+	return GetAttributeSetClass()->IsChildOf(UAbilitySystemComponent::StaticClass());
+}
+
 UAttributeSet::UAttributeSet(const FObjectInitializer& ObjectInitializer)
 	: Super(ObjectInitializer)
 {
diff --git a/Engine/Source/Runtime/GameplayAbilities/Private/GameplayEffect.cpp b/Engine/Source/Runtime/GameplayAbilities/Private/GameplayEffect.cpp
--- a/Engine/Source/Runtime/GameplayAbilities/Private/GameplayEffect.cpp
+++ b/Engine/Source/Runtime/GameplayAbilities/Private/GameplayEffect.cpp
@@ -166,10 +166,10 @@ bool FGameplayEffectAttributeCaptureSpecContainer::HasNonSnapshottedAttributes()
 
 FGameplayEffectSpec::FGameplayEffectSpec(const UGameplayEffect* InDef, const FGameplayEffectContextHandle& InEffectContext, float Level)
 : Def(InDef)
-, EffectContext(InEffectContext)
 {
 	check(Def);	
 	SetLevel(Level);
+	SetContext(InEffectContext);
 
 	// Init our ModifierSpecs
 	Modifiers.Reserve(Def->Modifiers.Num());
@@ -206,30 +206,41 @@ FGameplayEffectSpec::FGameplayEffectSpec(const UGameplayEffect* InDef, const FGa
 		}
 	}
 
-
-	
-
-	// Capture source tags	
-	InEffectContext.GetOwnedGameplayTags(CapturedSourceTags);
-
-	// Capture source Attributes
-	// Is this the right place to do it? Do we ever need to create spec and capture attributes at a later time? If so, this will need to move.
-	CapturedRelevantAttributes.CaptureAttributes(InEffectContext.GetInstigatorAbilitySystemComponent(), EGameplayEffectAttributeCaptureSource::Source);
+	// Add the GameplayEffect asset tags to the source Spec tags
+	CapturedSourceTags.GetSpecTags().AppendTags(InDef->GameplayEffectTags);
 
 	// Make TargetEffectSpecs too
 	for (UGameplayEffect* TargetDef : InDef->TargetEffects)
 	{
 		TargetEffectSpecs.Add(TSharedRef<FGameplayEffectSpec>(new FGameplayEffectSpec(TargetDef, EffectContext, Level)));
 	}
+
+	// Everything is setup now, capture data from our source
+	CaptureDataFromSource();
+}
+
+void FGameplayEffectSpec::CaptureDataFromSource()
+{
+	// Capture source actor tags
+	CapturedSourceTags.GetActorTags().RemoveAllTags();
+	EffectContext.GetOwnedGameplayTags(CapturedSourceTags.GetActorTags());
+
+	// Capture source Attributes
+	// Is this the right place to do it? Do we ever need to create spec and capture attributes at a later time? If so, this will need to move.
+	CapturedRelevantAttributes.CaptureAttributes(EffectContext.GetInstigatorAbilitySystemComponent(), EGameplayEffectAttributeCaptureSource::Source);
 }
 
 void FGameplayEffectSpec::SetLevel(float InLevel)
 {
 	Level = InLevel;
-	Duration = Def->Duration.GetValueAtLevel(InLevel);
-	Period = Def->Period.GetValueAtLevel(InLevel);
-	ChanceToApplyToTarget = Def->ChanceToApplyToTarget.GetValueAtLevel(InLevel);
-	ChanceToExecuteOnGameplayEffect = Def->ChanceToExecuteOnGameplayEffect.GetValueAtLevel(InLevel);
+	if (Def)
+	{
+		SetDuration(Def->Duration.GetValueAtLevel(InLevel));
+
+		Period = Def->Period.GetValueAtLevel(InLevel);
+		ChanceToApplyToTarget = Def->ChanceToApplyToTarget.GetValueAtLevel(InLevel);
+		ChanceToExecuteOnGameplayEffect = Def->ChanceToExecuteOnGameplayEffect.GetValueAtLevel(InLevel);
+	}
 }
 
 float FGameplayEffectSpec::GetLevel() const
@@ -242,6 +253,17 @@ float FGameplayEffectSpec::GetDuration() const
 	return Duration;
 }
 
+void FGameplayEffectSpec::SetDuration(float NewDuration)
+{
+	Duration = NewDuration;
+	if (Duration > 0.f)
+	{
+		// We may  have potential problems one day if a game is applying duration based gameplay effects from instantaneous effects
+		// (E.g., everytime fire damage is applied, a DOT is also applied). We may need to for Duration to always be capturd.
+		CapturedRelevantAttributes.AddCaptureDefinition(UAbilitySystemComponent::GetOutgoingDurationCapture());
+	}
+}
+
 float FGameplayEffectSpec::GetPeriod() const
 {
 	return Period;
@@ -347,6 +369,16 @@ EGameplayEffectStackingPolicy::Type FGameplayEffectSpec::GetStackingType() const
 	return Def->StackingPolicy;
 }
 
+void FGameplayEffectSpec::SetContext(FGameplayEffectContextHandle NewEffectContext)
+{
+	bool bWasAlreadyInit = EffectContext.IsValid();
+	EffectContext = NewEffectContext;	
+	if (bWasAlreadyInit)
+	{
+		CaptureDataFromSource();
+	}
+}
+
 // --------------------------------------------------------------------------------------------------------------------------------------------------------
 //
 //	FModifierSpec
@@ -526,6 +558,9 @@ void FActiveGameplayEffectsContainer::RegisterWithOwner(UAbilitySystemComponent*
 		// Binding raw is ok here, since the owner is literally the UObject that owns us. If we are destroyed, its because that uobject is destroyed,
 		// and if that is destroyed, the delegate wont be able to fire.
 		Owner->RegisterGenericGameplayTagEvent().AddRaw(this, &FActiveGameplayEffectsContainer::OnOwnerTagChange);
+
+		// Add Aggregators for system attributes
+		//FindOrCreateAttributeAggregator(
 	}
 }
 
@@ -547,9 +582,11 @@ void FActiveGameplayEffectsContainer::ExecuteActiveEffectsFrom(FGameplayEffectSp
 
 	float ChanceToExecute = Spec.GetChanceToExecuteOnGameplayEffect();		// Not implemented? Should we just remove ChanceToExecute?
 
-	// Capture our own tags. Should this ever be a setting or are Target tags always up to date? Seems simplier to keep them always up to date.
-	Spec.CapturedSourceTags.RemoveAllTags();
-	Owner->GetOwnedGameplayTags(Spec.CapturedSourceTags);
+	// Capture our own tags.
+	// TODO: We should only capture them if we need to. We may have snapshotted target tags (?) (in the case of dots with exotic setups?)
+
+	Spec.CapturedTargetTags.GetActorTags().RemoveAllTags();
+	Owner->GetOwnedGameplayTags(Spec.CapturedTargetTags.GetActorTags());
 
 	Spec.CalculateModifierMagnitudes();
 
@@ -730,9 +767,11 @@ FAggregatorRef& FActiveGameplayEffectsContainer::FindOrCreateAttributeAggregator
 	ABILITY_LOG(Log, TEXT("Creating new entry in AttributeAggregatorMap for %s. CurrentValue: %.2f"), *Attribute.GetName(), CurrentValueOfProperty);
 
 	FAggregator* NewAttributeAggregator = new FAggregator(CurrentValueOfProperty);
-
-	// FIXME: AddRaw is bad
-	NewAttributeAggregator->OnDirty.AddUObject(Owner, &UAbilitySystemComponent::OnAttributeAggregatorDirty, Attribute);
+	
+	if (Attribute.IsSystemAttribute() == false)
+	{
+		NewAttributeAggregator->OnDirty.AddUObject(Owner, &UAbilitySystemComponent::OnAttributeAggregatorDirty, Attribute);
+	}
 
 	return AttributeAggregatorMap.Add(Attribute, FAggregatorRef(NewAttributeAggregator));
 }
@@ -870,17 +909,47 @@ FActiveGameplayEffect& FActiveGameplayEffectsContainer::CreateNewActiveGameplayE
 	FActiveGameplayEffectHandle NewHandle = FActiveGameplayEffectHandle::GenerateNewHandle(Owner);
 	FActiveGameplayEffect& NewEffect = *new (GameplayEffects)FActiveGameplayEffect(NewHandle, Spec, GetWorldTime(), GetGameStateTime(), InPredictionKey);
 
+	// Calculate Duration mods if we have a real duration
+	float DurationBaseValue = NewEffect.Spec.GetDuration();
+	if (DurationBaseValue > 0.f)
+	{
+		const FAggregator& OutgoingAgg = NewEffect.Spec.CapturedRelevantAttributes.GetAggregator(UAbilitySystemComponent::GetOutgoingDurationCapture());
+		const FAggregator& IncomingAgg = *FindOrCreateAttributeAggregator(UAbilitySystemComponent::GetIncomingDurationProperty()).Get();
+		
+		FAggregator DurationAgg;
+
+		DurationAgg.AddModsFrom(OutgoingAgg);
+		DurationAgg.AddModsFrom(IncomingAgg);
+
+		FAggregatorEvaluateParameters Params;
+		Params.SourceTags = NewEffect.Spec.CapturedSourceTags.GetAggregatedTags();
+		Params.TargetTags = NewEffect.Spec.CapturedTargetTags.GetAggregatedTags();
+
+		float FinalDuration = DurationAgg.EvaluateWithBase(DurationBaseValue, Params);
+
+		// We cannot mod ourselves into an instant or infinite duration effect
+		if (FinalDuration <= 0.f)
+		{
+			ABILITY_LOG(Error, TEXT("GameplayEffect %s Duration was modified to %.2f. Clamping to 0.1s duration."), *NewEffect.Spec.Def->GetName(), FinalDuration);
+			FinalDuration = 0.1f;
+		}
+
+		NewEffect.Spec.SetDuration(FinalDuration);
+
+		// ABILITY_LOG(Warning, TEXT("SetDuration for %s. Base: %.2f, Final: %.2f"), *NewEffect.Spec.Def->GetName(), DurationBaseValue, FinalDuration);
+	}
+	
 	// register callbacks with the timer manager
 	if (Owner)
 	{
-		if (Spec.GetDuration() > 0.f)
+		if (NewEffect.Spec.GetDuration() > 0.f)
 		{
 			FTimerManager& TimerManager = Owner->GetWorld()->GetTimerManager();
 			FTimerDelegate Delegate = FTimerDelegate::CreateUObject(Owner, &UAbilitySystemComponent::CheckDurationExpired, NewHandle);
-			TimerManager.SetTimer(NewEffect.DurationHandle, Delegate, Spec.GetDuration(), false, Spec.GetDuration());
+			TimerManager.SetTimer(NewEffect.DurationHandle, Delegate, NewEffect.Spec.GetDuration(), false, NewEffect.Spec.GetDuration());
 		}
 		// The timer manager moves things from the pending list to the active list after checking the active list on the first tick so we need to execute here
-		if (Spec.GetPeriod() != UGameplayEffect::NO_PERIOD)
+		if (NewEffect.Spec.GetPeriod() != UGameplayEffect::NO_PERIOD)
 		{
 			FTimerManager& TimerManager = Owner->GetWorld()->GetTimerManager();
 			FTimerDelegate Delegate = FTimerDelegate::CreateUObject(Owner, &UAbilitySystemComponent::ExecutePeriodicEffect, NewHandle);
@@ -888,7 +957,7 @@ FActiveGameplayEffect& FActiveGameplayEffectsContainer::CreateNewActiveGameplayE
 			// If this is a periodic stacking effect, make sure that it's in sync with the others.
 			float FirstDelay = -1.f;
 			bool bApplyToNextTick = true;
-			if (Spec.GetStackingType() != EGameplayEffectStackingPolicy::Unlimited)
+			if (NewEffect.Spec.GetStackingType() != EGameplayEffectStackingPolicy::Unlimited)
 			{
 				for (FActiveGameplayEffect& Effect : GameplayEffects)
 				{
@@ -906,7 +975,7 @@ FActiveGameplayEffect& FActiveGameplayEffectsContainer::CreateNewActiveGameplayE
 				TimerManager.SetTimerForNextTick(Delegate); // not part of a stack or the first item on the stack, execute once right away
 			}
 
-			TimerManager.SetTimer(NewEffect.PeriodHandle, Delegate, Spec.GetPeriod(), true, FirstDelay); // this is going to be off by a frame for stacking because of the pending list
+			TimerManager.SetTimer(NewEffect.PeriodHandle, Delegate, NewEffect.Spec.GetPeriod(), true, FirstDelay); // this is going to be off by a frame for stacking because of the pending list
 		}
 	}
 	
diff --git a/Engine/Source/Runtime/GameplayAbilities/Private/GameplayEffectAggregator.cpp b/Engine/Source/Runtime/GameplayAbilities/Private/GameplayEffectAggregator.cpp
--- a/Engine/Source/Runtime/GameplayAbilities/Private/GameplayEffectAggregator.cpp
+++ b/Engine/Source/Runtime/GameplayAbilities/Private/GameplayEffectAggregator.cpp
@@ -5,8 +5,8 @@
 
 bool FAggregatorMod::Qualifies(const FAggregatorEvaluateParameters& Parameters) const
 {
-	bool SourceMet = (!SourceTagReqs || SourceTagReqs->RequirementsMet(Parameters.SourceTags));
-	bool TargetMet = (!TargetTagReqs || TargetTagReqs->RequirementsMet(Parameters.TargetTags));
+	bool SourceMet = (!SourceTagReqs || SourceTagReqs->IsEmpty()) || (Parameters.SourceTags && SourceTagReqs->RequirementsMet(*Parameters.SourceTags));
+	bool TargetMet = (!TargetTagReqs || TargetTagReqs->IsEmpty()) || (Parameters.TargetTags && TargetTagReqs->RequirementsMet(*Parameters.TargetTags));
 
 	return SourceMet && TargetMet;
 }
@@ -49,29 +49,29 @@ float FAggregator::StaticExecModOnBaseValue(float BaseValue, TEnumAsByte<EGamepl
 {
 	switch (ModifierOp)
 	{
-	case EGameplayModOp::Override:
-	{
-		BaseValue = EvaluatedMagnitude;
-		break;
-	}
-	case EGameplayModOp::Additive:
-	{
-		BaseValue += EvaluatedMagnitude;
-		break;
-	}
-	case EGameplayModOp::Multiplicitive:
-	{
-		BaseValue *= EvaluatedMagnitude;
-		break;
-	}
-	case EGameplayModOp::Division:
-	{
-		if (FMath::IsNearlyZero(EvaluatedMagnitude) == false)
+		case EGameplayModOp::Override:
 		{
-			BaseValue /= EvaluatedMagnitude;
+			BaseValue = EvaluatedMagnitude;
+			break;
+		}
+		case EGameplayModOp::Additive:
+		{
+			BaseValue += EvaluatedMagnitude;
+			break;
+		}
+		case EGameplayModOp::Multiplicitive:
+		{
+			BaseValue *= EvaluatedMagnitude;
+			break;
+		}
+		case EGameplayModOp::Division:
+		{
+			if (FMath::IsNearlyZero(EvaluatedMagnitude) == false)
+			{
+				BaseValue /= EvaluatedMagnitude;
+			}
+			break;
 		}
-		break;
-	}
 	}
 
 	return BaseValue;
@@ -126,6 +126,14 @@ void FAggregator::RemoveMod(FActiveGameplayEffectHandle ActiveHandle)
 	BroadcastOnDirty();
 }
 
+void FAggregator::AddModsFrom(const FAggregator& SourceAggregator)
+{
+	for(int32 idx = 0; idx < ARRAY_COUNT(Mods); ++idx)
+	{
+		Mods[idx].Append(SourceAggregator.Mods[idx]);
+	}
+}
+
 void FAggregator::RemoveModsWithActiveHandle(TArray<FAggregatorMod>& Mods, FActiveGameplayEffectHandle ActiveHandle)
 {
 	check(ActiveHandle.IsValid());
diff --git a/Engine/Source/Runtime/GameplayAbilities/Private/GameplayEffectExtension_LifestealTest.cpp b/Engine/Source/Runtime/GameplayAbilities/Private/GameplayEffectExtension_LifestealTest.cpp
--- a/Engine/Source/Runtime/GameplayAbilities/Private/GameplayEffectExtension_LifestealTest.cpp
+++ b/Engine/Source/Runtime/GameplayAbilities/Private/GameplayEffectExtension_LifestealTest.cpp
@@ -28,7 +28,7 @@ void UGameplayEffectExtension_LifestealTest::PostGameplayEffectExecute(const FGa
 	float HealthToRestore = -DamageDone * LifestealPCT;
 	if (HealthToRestore > 0.f)
 	{
-		UAbilitySystemComponent *Source = Data.EffectSpec.EffectContext.GetOriginalInstigatorAbilitySystemComponent();
+		UAbilitySystemComponent *Source = Data.EffectSpec.GetContext().GetOriginalInstigatorAbilitySystemComponent();
 
 		UGameplayEffect * LocalHealthRestore = HealthRestoreGameplayEffect;
 		if (!LocalHealthRestore)
diff --git a/Engine/Source/Runtime/GameplayAbilities/Private/GameplayEffectTypes.cpp b/Engine/Source/Runtime/GameplayAbilities/Private/GameplayEffectTypes.cpp
--- a/Engine/Source/Runtime/GameplayAbilities/Private/GameplayEffectTypes.cpp
+++ b/Engine/Source/Runtime/GameplayAbilities/Private/GameplayEffectTypes.cpp
@@ -279,7 +279,7 @@ void FGameplayTagCountContainer::UpdateTagMap_Internal(const FGameplayTag& Tag,
 	}
 }
 
-bool FGameplayTagRequirements::RequirementsMet(FGameplayTagContainer Container) const
+bool FGameplayTagRequirements::RequirementsMet(const FGameplayTagContainer& Container) const
 {
 	bool HasRequired = Container.MatchesAll(RequireTags, true);
 	bool HasIgnored = Container.MatchesAny(IgnoreTags, false);
@@ -287,6 +287,11 @@ bool FGameplayTagRequirements::RequirementsMet(FGameplayTagContainer Container)
 	return HasRequired && !HasIgnored;
 }
 
+bool FGameplayTagRequirements::IsEmpty() const
+{
+	return (RequireTags.Num() == 0 && IgnoreTags.Num() == 0);
+}
+
 void FActiveGameplayEffectsContainer::PrintAllGameplayEffects() const
 {
 	ABILITY_LOG_SCOPE(TEXT("ActiveGameplayEffects. Num: %d"), GameplayEffects.Num());
@@ -313,4 +318,41 @@ void FGameplayEffectSpec::PrintAll() const
 	ABILITY_LOG(Log, TEXT("Period: %.2f"), GetPeriod());
 
 	ABILITY_LOG(Log, TEXT("Modifiers:"));
+}
+
+const FGameplayTagContainer* FTagContainerAggregator::GetAggregatedTags() const
+{
+	if (CacheIsValid == false)
+	{
+		CacheIsValid = true;
+		CachedAggregator.RemoveAllTags(CapturedActorTags.Num() + CapturedSpecTags.Num() + ScopedTags.Num());
+		CachedAggregator.AppendTags(CapturedActorTags);
+		CachedAggregator.AppendTags(CapturedSpecTags);
+		CachedAggregator.AppendTags(ScopedTags);
+	}
+
+	return &CachedAggregator;
+}
+
+FGameplayTagContainer& FTagContainerAggregator::GetActorTags()
+{
+	CacheIsValid = false;
+	return CapturedActorTags;
+}
+
+const FGameplayTagContainer& FTagContainerAggregator::GetActorTags() const
+{
+	return CapturedActorTags;
+}
+
+FGameplayTagContainer& FTagContainerAggregator::GetSpecTags()
+{
+	CacheIsValid = false;
+	return CapturedSpecTags;
+}
+
+const FGameplayTagContainer& FTagContainerAggregator::GetSpecTags() const
+{
+	CacheIsValid = false;
+	return CapturedSpecTags;
 }
\ No newline at end of file
diff --git a/Engine/Source/Runtime/GameplayAbilities/Private/GameplayPrediction.cpp b/Engine/Source/Runtime/GameplayAbilities/Private/GameplayPrediction.cpp
--- a/Engine/Source/Runtime/GameplayAbilities/Private/GameplayPrediction.cpp
+++ b/Engine/Source/Runtime/GameplayAbilities/Private/GameplayPrediction.cpp
@@ -21,7 +21,8 @@ bool FPredictionKey::NetSerialize(FArchive& Ar, class UPackageMap* Map, bool& bO
 		/**
 		 *	Only serialize the payload if we have no owning connection (Client sending to server).
 		 *	or if the owning connection is this connection (Server only sends the prediction key to the client who gave it to us).
-		 */	
+		 */
+		
 		if (PredictiveConnection == nullptr || (Map == PredictiveConnection))
 		{
 			Ar << Current;
@@ -68,10 +69,15 @@ void FPredictionKey::GenerateDependantPredictionKey()
 	}
 }
 
-FPredictionKey FPredictionKey::CreateNewPredictionKey()
+FPredictionKey FPredictionKey::CreateNewPredictionKey(UAbilitySystemComponent* OwningComponent)
 {
 	FPredictionKey NewKey;
-	NewKey.GenerateNewPredictionKey();
+	
+	// We should never generate prediction keys on the authority
+	if(OwningComponent->GetOwnerRole() != ROLE_Authority)
+	{
+		NewKey.GenerateNewPredictionKey();
+	}
 	return NewKey;
 }
 
@@ -190,7 +196,7 @@ void FPredictionKeyDelegates::AddDependancy(FPredictionKey::KeyType ThisKey, FPr
 
 FScopedPredictionWindow::FScopedPredictionWindow(UAbilitySystemComponent* AbilitySystemComponent, FPredictionKey InPredictionKey)
 {
-	// This is used to set an already generated predictiot key as the current scoped prediction key.
+	// This is used to set an already generated prediction key as the current scoped prediction key.
 	// Should be called on the server for logical scopes where a given key is valid. E.g, "client gave me this key, we both are going to run Foo()".
 
 	Owner = AbilitySystemComponent;
diff --git a/Engine/Source/Runtime/GameplayAbilities/Public/Abilities/GameplayAbility.h b/Engine/Source/Runtime/GameplayAbilities/Public/Abilities/GameplayAbility.h
--- a/Engine/Source/Runtime/GameplayAbilities/Public/Abilities/GameplayAbility.h
+++ b/Engine/Source/Runtime/GameplayAbilities/Public/Abilities/GameplayAbility.h
@@ -341,7 +341,7 @@ protected:
 	FActiveGameplayEffectHandle K2_ApplyGameplayEffectToOwner(const UGameplayEffect* GameplayEffect, int32 GameplayEffectLevel = 1);
 
 	/** Non blueprintcallable, safe to call on CDO/NonInstance abilities */
-	FActiveGameplayEffectHandle ApplyGameplayEffectToOwner(const FGameplayAbilitySpecHandle Handle, const FGameplayAbilityActorInfo* ActorInfo, const FGameplayAbilityActivationInfo ActivationInfo, const UGameplayEffect* GameplayEffect, int32 GameplayEffectLevel);
+	FActiveGameplayEffectHandle ApplyGameplayEffectToOwner(const FGameplayAbilitySpecHandle Handle, const FGameplayAbilityActorInfo* ActorInfo, const FGameplayAbilityActivationInfo ActivationInfo, const UGameplayEffect* GameplayEffect, float GameplayEffectLevel);
 
 	UFUNCTION(BlueprintCallable, Category = Ability, FriendlyName = "ApplyGameplayEffectToTarget")
 	FActiveGameplayEffectHandle K2_ApplyGameplayEffectToTarget(FGameplayAbilityTargetDataHandle TargetData, const UGameplayEffect* GameplayEffect, int32 GameplayEffectLevel = 1);
@@ -401,7 +401,7 @@ protected:
 		SetCurrentActivationInfo(ActivationInfo);
 	}
 
-	public:
+public:
 
 	/** Returns the actor info associated with this ability, has cached pointers to useful objects */
 	UFUNCTION(BlueprintCallable, Category=Ability)
@@ -421,9 +421,15 @@ protected:
 
 	/** Convenience method for abilities to get outgoing gameplay effect specs (for example, to pass on to projectiles to apply to whoever they hit) */
 	UFUNCTION(BlueprintCallable, Category=Ability)
-	FGameplayEffectSpecHandle GetOutgoingGameplayEffectSpec(UGameplayEffect* GameplayEffect, float Level=1.f) const;
+	FGameplayEffectSpecHandle GetOutgoingGameplayEffectSpec(const UGameplayEffect* GameplayEffect, float Level=1.f) const;
+
+	FGameplayEffectSpecHandle GetOutgoingGameplayEffectSpec(const FGameplayAbilitySpecHandle Handle, const FGameplayAbilityActorInfo* ActorInfo, const FGameplayAbilityActivationInfo ActivationInfo, const UGameplayEffect* GameplayEffect, float Level = 1.f) const;
+	
 
-	protected:
+	/** Add the Ability's tags to the given GameplayEffectSpec. This is likely to be overriden per project. */
+	virtual void ApplyAbilityTagsToGameplayEffectSpec(FGameplayEffectSpec& Spec) const;
+
+protected:
 
 	bool IsSupportedForNetworking() const override;
 
@@ -452,10 +458,6 @@ protected:
 	UPROPERTY(EditDefaultsOnly, Category=Advanced)
 	TEnumAsByte<EGameplayAbilityNetExecutionPolicy::Type> NetExecutionPolicy;
 
-	/** This GameplayEffect represents the cooldown. It will be applied when the ability is committed and the ability cannot be used again until it is expired. */
-	UPROPERTY(EditDefaultsOnly, Category=Cooldowns)
-	class UGameplayEffect* CooldownGameplayEffect;
-
 	/** This GameplayEffect represents the cost (mana, stamina, etc) of the ability. It will be applied when the ability is committed. */
 	UPROPERTY(EditDefaultsOnly, Category=Costs)
 	class UGameplayEffect* CostGameplayEffect;
@@ -463,6 +465,25 @@ protected:
 	/** Triggers to determine if this ability should execute in response to an event */
 	UPROPERTY(EditDefaultsOnly, Category = Triggers)
 	TArray<FAbilityTriggerData> AbilityTriggers;
+
+	// ----------------------------------------------------------------------------------------------------------------
+	//
+	//	Cooldowns
+	//
+	// ----------------------------------------------------------------------------------------------------------------
+		
+
+	/** Cooldown tags that we check for. If we have cooldowns with these tags, we cannot activate this ability */
+	UPROPERTY()
+	FGameplayTagContainer CheckedCooldownTags;
+
+	UPROPERTY()
+	TArray<FGameplayAbilityCooldownInfo> AppliedCooldowns;
+	
+	/** Deprecated? */
+	UPROPERTY(EditDefaultsOnly, Category = Cooldowns)
+	class UGameplayEffect* CooldownGameplayEffect;
+
 	
 	// ----------------------------------------------------------------------------------------------------------------
 	//
diff --git a/Engine/Source/Runtime/GameplayAbilities/Public/Abilities/GameplayAbilityTypes.h b/Engine/Source/Runtime/GameplayAbilities/Public/Abilities/GameplayAbilityTypes.h
--- a/Engine/Source/Runtime/GameplayAbilities/Public/Abilities/GameplayAbilityTypes.h
+++ b/Engine/Source/Runtime/GameplayAbilities/Public/Abilities/GameplayAbilityTypes.h
@@ -4,6 +4,7 @@
 
 #include "GameplayEffect.h"
 #include "GameplayPrediction.h"
+#include "AttributeSet.h"
 #include "GameplayAbilityTypes.generated.h"
 
 class UGameplayEffect;
@@ -473,6 +474,18 @@ struct GAMEPLAYABILITIES_API FAttributeDefaults
 	class UDataTable*	DefaultStartingTable;
 };
 
+// Still WIP: not being used yet
+USTRUCT()
+struct GAMEPLAYABILITIES_API FGameplayAbilityCooldownInfo
+{
+	GENERATED_USTRUCT_BODY()
+
+	UPROPERTY(EditDefaultsOnly, Category = Cooldowns)
+	FScalableFloat	Cooldown;
+
+	UPROPERTY(EditDefaultsOnly, Category = Cooldowns)
+	FGameplayTagContainer CooldownTags;
+};
 
 /** Used for cleaning up predicted data on network clients */
 DECLARE_MULTICAST_DELEGATE(FAbilitySystemComponentPredictionKeyClear);
diff --git a/Engine/Source/Runtime/GameplayAbilities/Public/AbilitySystemComponent.h b/Engine/Source/Runtime/GameplayAbilities/Public/AbilitySystemComponent.h
--- a/Engine/Source/Runtime/GameplayAbilities/Public/AbilitySystemComponent.h
+++ b/Engine/Source/Runtime/GameplayAbilities/Public/AbilitySystemComponent.h
@@ -180,7 +180,10 @@ class GAMEPLAYABILITIES_API UAbilitySystemComponent : public UActorComponent, pu
 
 	/** Get an outgoing GameplayEffectSpec that is ready to be applied to other things. */
 	UFUNCTION(BlueprintCallable, Category = GameplayEffects)
-	FGameplayEffectSpecHandle GetOutgoingSpec(UGameplayEffect* GameplayEffect, float Level) const;
+	FGameplayEffectSpecHandle GetOutgoingSpec(const UGameplayEffect* GameplayEffect, float Level = 1.f) const;
+
+	/** Overloaded version that allows Context to be specified */
+	FGameplayEffectSpecHandle GetOutgoingSpec(const UGameplayEffect* GameplayEffect, float Level, FGameplayEffectContextHandle Context) const;
 
 	/** Create an EffectContext for the owner of this AbilitySystemComponent */
 	UFUNCTION(BlueprintCallable, Category = GameplayEffects)
@@ -204,10 +207,6 @@ class GAMEPLAYABILITIES_API UAbilitySystemComponent : public UActorComponent, pu
 	UFUNCTION(BlueprintCallable, Category = GameplayEffects)
 	bool IsGameplayEffectActive(FActiveGameplayEffectHandle InHandle) const;
 
-
-	// New Stuff ---------------------------------- (Remove this comment!)
-
-
 	/**
 	 * Populate the specified capture spec with the data necessary to capture an attribute from the component
 	 * 
@@ -256,6 +255,23 @@ class GAMEPLAYABILITIES_API UAbilitySystemComponent : public UActorComponent, pu
 
 	FOnGameplayAttributeChange& RegisterGameplayAttributeEvent(FGameplayAttribute Attribute);
 
+	// --------------------------------------------
+	// System Attributes
+	// --------------------------------------------
+	
+	UPROPERTY(meta=(SystemGameplayAttribute="true"))
+	float OutgoingDuration;
+
+	UPROPERTY(meta = (SystemGameplayAttribute = "true"))
+	float IncomingDuration;
+
+	static UProperty* GetOutgoingDurationProperty();
+	static UProperty* GetIncomingDurationProperty();
+
+	static const FGameplayEffectAttributeCaptureDefinition& GetOutgoingDurationCapture();
+	static const FGameplayEffectAttributeCaptureDefinition& GetIncomingDurationCapture();
+
+
 	// --------------------------------------------
 	// Possibly useful but not primary API functions:
 	// --------------------------------------------
diff --git a/Engine/Source/Runtime/GameplayAbilities/Public/AttributeSet.h b/Engine/Source/Runtime/GameplayAbilities/Public/AttributeSet.h
--- a/Engine/Source/Runtime/GameplayAbilities/Public/AttributeSet.h
+++ b/Engine/Source/Runtime/GameplayAbilities/Public/AttributeSet.h
@@ -32,17 +32,19 @@ struct GAMEPLAYABILITIES_API FGameplayAttribute
 		Attribute = NewProperty;
 	}
 
-	UProperty * GetUProperty() const
+	UProperty* GetUProperty() const
 	{
 		return Attribute;
 	}
 
-	UClass * GetAttributeSetClass() const
+	UClass* GetAttributeSetClass() const
 	{
 		check(Attribute);
 		return CastChecked<UClass>(Attribute->GetOuter());
 	}
 
+	bool IsSystemAttribute() const;
+
 	void SetNumericValueChecked(const float NewValue, class UAttributeSet* Dest) const;
 
 	float GetNumericValueChecked(const UAttributeSet* Src) const;
diff --git a/Engine/Source/Runtime/GameplayAbilities/Public/GameplayEffect.h b/Engine/Source/Runtime/GameplayAbilities/Public/GameplayEffect.h
--- a/Engine/Source/Runtime/GameplayAbilities/Public/GameplayEffect.h
+++ b/Engine/Source/Runtime/GameplayAbilities/Public/GameplayEffect.h
@@ -587,14 +587,15 @@ struct FGameplayEffectSpec
 	GENERATED_USTRUCT_BODY()
 
 	FGameplayEffectSpec()
-		: Duration(UGameplayEffect::INSTANT_APPLICATION)
+		: Def(nullptr)
+		, Duration(UGameplayEffect::INSTANT_APPLICATION)
 		, Period(UGameplayEffect::NO_PERIOD)
 		, ChanceToApplyToTarget(1.f)
 		, ChanceToExecuteOnGameplayEffect(1.f)
 		, bTopOfStack(false)
 		, Level(UGameplayEffect::INVALID_LEVEL)
 	{
-		// If we initialize a GameplayEffectSpec with no level object passed in.
+
 	}
 
 	FGameplayEffectSpec(const UGameplayEffect* InDef, const FGameplayEffectContextHandle& InEffectContext, float Level = UGameplayEffect::INVALID_LEVEL);
@@ -602,11 +603,6 @@ struct FGameplayEffectSpec
 	UPROPERTY()
 	const UGameplayEffect* Def;
 	
-	// Replicated
-	UPROPERTY()
-	FGameplayEffectContextHandle EffectContext; // This tells us how we got here (who / what applied us)
-
-	// Replicated
 	UPROPERTY()
 	TArray<FGameplayEffectModifiedAttribute> ModifiedAttributes;
 	
@@ -614,16 +610,19 @@ struct FGameplayEffectSpec
 	UPROPERTY(NotReplicated)
 	FGameplayEffectAttributeCaptureSpecContainer CapturedRelevantAttributes;
 
-	// Looks for an existing modified attribute struct, may return NULL
+	/** Looks for an existing modified attribute struct, may return NULL */
 	const FGameplayEffectModifiedAttribute* GetModifiedAttribute(const FGameplayAttribute& Attribute) const;
 	FGameplayEffectModifiedAttribute* GetModifiedAttribute(const FGameplayAttribute& Attribute);
 
-	// Adds a new modified attribute struct, will always add so check to see if it exists first
+	/** Adds a new modified attribute struct, will always add so check to see if it exists first */
 	FGameplayEffectModifiedAttribute* AddModifiedAttribute(const FGameplayAttribute& Attribute);
 
-	// Deletes any modified attributes that aren't needed. Call before replication
+	/** Deletes any modified attributes that aren't needed. Call before replication */
 	void PruneModifiedAttributes();
 
+	/** Sets duration. This should onlyt be called as the GameplayEffect is being created */
+	void SetDuration(float NewDuration);
+
 	float GetDuration() const;
 	float GetPeriod() const;
 	float GetChanceToApplyToTarget() const;
@@ -632,9 +631,17 @@ struct FGameplayEffectSpec
 
 	EGameplayEffectStackingPolicy::Type GetStackingType() const;
 
-	// other effects that need to be applied to the target if this effect is successful
+	/** other effects that need to be applied to the target if this effect is successful */
 	TArray< TSharedRef< FGameplayEffectSpec > > TargetEffectSpecs;
 
+	/** Set the context info: who and where this spec came from. */
+	void SetContext(FGameplayEffectContextHandle NewEffectContext);
+
+	FGameplayEffectContextHandle GetContext() const
+	{
+		return EffectContext;
+	}
+
 	// The duration in seconds of this effect
 	// instantaneous effects should have a duration of UGameplayEffect::INSTANT_APPLICATION
 	// effects that last forever should have a duration of UGameplayEffect::INFINITE_DURATION
@@ -658,17 +665,16 @@ struct FGameplayEffectSpec
 	UPROPERTY()
 	bool bTopOfStack;
 
-	// Captured Source Tags on GameplayEffectSpec creation.
-	//	-Do we ever not snapshot these tags?
+	// Captured Source Tags on GameplayEffectSpec creation.	
 	UPROPERTY()
-	FGameplayTagContainer	CapturedSourceTags;
+	FTagContainerAggregator	CapturedSourceTags;
 
 	// Tags from the target, captured during execute?
 	// Not sure if this is the right place or not, or how 'contextual' tags get injected (headshot, crit, etc)
 	// Does this overlap with info in the GameplayEffectContextHandle? Should it just all live there?
 	// For now - putting it here so places have one spot to pull from. Need to figure this out long term.
 	UPROPERTY()
-	FGameplayTagContainer	CapturedTargetTags;
+	FTagContainerAggregator	CapturedTargetTags;
 	
 	TArray<FModifierSpec> Modifiers;
 	
@@ -689,11 +695,13 @@ struct FGameplayEffectSpec
 
 	void PrintAll() const;
 
-	
-
 private:
 
-	// Replicated	
+	void CaptureDataFromSource();
+	
+	UPROPERTY()
+	FGameplayEffectContextHandle EffectContext; // This tells us how we got here (who / what applied us)
+	
 	UPROPERTY()
 	float Level;	
 };
diff --git a/Engine/Source/Runtime/GameplayAbilities/Public/GameplayEffectAggregator.h b/Engine/Source/Runtime/GameplayAbilities/Public/GameplayEffectAggregator.h
--- a/Engine/Source/Runtime/GameplayAbilities/Public/GameplayEffectAggregator.h
+++ b/Engine/Source/Runtime/GameplayAbilities/Public/GameplayEffectAggregator.h
@@ -8,8 +8,10 @@
 
 struct GAMEPLAYABILITIES_API FAggregatorEvaluateParameters
 {
-	FGameplayTagContainer	SourceTags;
-	FGameplayTagContainer	TargetTags;
+	FAggregatorEvaluateParameters() : SourceTags(nullptr), TargetTags(nullptr) { }
+
+	const FGameplayTagContainer*	SourceTags;
+	const FGameplayTagContainer*	TargetTags;
 };
 
 struct GAMEPLAYABILITIES_API FAggregatorMod
@@ -52,6 +54,8 @@ struct GAMEPLAYABILITIES_API FAggregator : public TSharedFromThis<FAggregator>
 
 	FOnAggregatorDirty OnDirty;
 
+	void AddModsFrom(const FAggregator& SourceAggregator);
+
 private:
 
 	void BroadcastOnDirty();
diff --git a/Engine/Source/Runtime/GameplayAbilities/Public/GameplayEffectTypes.h b/Engine/Source/Runtime/GameplayAbilities/Public/GameplayEffectTypes.h
--- a/Engine/Source/Runtime/GameplayAbilities/Public/GameplayEffectTypes.h
+++ b/Engine/Source/Runtime/GameplayAbilities/Public/GameplayEffectTypes.h
@@ -617,7 +617,40 @@ struct FGameplayTagRequirements
 	UPROPERTY(EditDefaultsOnly, Category = GameplayModifier)
 	FGameplayTagContainer IgnoreTags;
 
-	bool	RequirementsMet(FGameplayTagContainer Container) const;
+	bool	RequirementsMet(const FGameplayTagContainer& Container) const;
+	bool	IsEmpty() const;
 
 	static FGetGameplayTags	SnapshotTags(FGetGameplayTags TagDelegate);
 };
+
+USTRUCT()
+struct GAMEPLAYABILITIES_API FTagContainerAggregator
+{
+	GENERATED_USTRUCT_BODY()
+
+	FTagContainerAggregator() : CacheIsValid(false) {}
+
+	FGameplayTagContainer& GetActorTags();
+	const FGameplayTagContainer& GetActorTags() const;
+
+	FGameplayTagContainer& GetSpecTags();
+	const FGameplayTagContainer& GetSpecTags() const;
+
+
+
+	const FGameplayTagContainer* GetAggregatedTags() const;
+
+private:
+
+	UPROPERTY()
+	FGameplayTagContainer CapturedActorTags;
+
+	UPROPERTY()
+	FGameplayTagContainer CapturedSpecTags;
+
+	UPROPERTY()
+	FGameplayTagContainer ScopedTags;
+
+	mutable FGameplayTagContainer CachedAggregator;
+	mutable bool CacheIsValid;
+};
diff --git a/Engine/Source/Runtime/GameplayAbilities/Public/GameplayPrediction.h b/Engine/Source/Runtime/GameplayAbilities/Public/GameplayPrediction.h
--- a/Engine/Source/Runtime/GameplayAbilities/Public/GameplayPrediction.h
+++ b/Engine/Source/Runtime/GameplayAbilities/Public/GameplayPrediction.h
@@ -304,7 +304,7 @@ struct GAMEPLAYABILITIES_API FPredictionKey
 	bool bIsStale;
 
 	/** Construct a new prediction key with no dependancies */
-	static FPredictionKey CreateNewPredictionKey();
+	static FPredictionKey CreateNewPredictionKey(class UAbilitySystemComponent*);
 
 	/** Create a new dependant prediction key: keep our existing base or use the current key as the base. */
 	void GenerateDependantPredictionKey();
diff --git a/Engine/Source/Runtime/GameplayTags/Classes/GameplayTagContainer.h b/Engine/Source/Runtime/GameplayTags/Classes/GameplayTagContainer.h
--- a/Engine/Source/Runtime/GameplayTags/Classes/GameplayTagContainer.h
+++ b/Engine/Source/Runtime/GameplayTags/Classes/GameplayTagContainer.h
@@ -197,7 +197,7 @@ struct GAMEPLAYTAGS_API FGameplayTagContainer
 	virtual void RemoveTag(FGameplayTag TagToRemove);
 
 	/** Remove all tags from the container */
-	virtual void RemoveAllTags();
+	virtual void RemoveAllTags(int32 Slack=0);
 
 	/**
 	 * Serialize the tag container
diff --git a/Engine/Source/Runtime/GameplayTags/Private/GameplayTagContainer.cpp b/Engine/Source/Runtime/GameplayTags/Private/GameplayTagContainer.cpp
--- a/Engine/Source/Runtime/GameplayTags/Private/GameplayTagContainer.cpp
+++ b/Engine/Source/Runtime/GameplayTags/Private/GameplayTagContainer.cpp
@@ -177,9 +177,9 @@ void FGameplayTagContainer::RemoveTag(FGameplayTag TagToRemove)
 	GameplayTags.Remove(TagToRemove);
 }
 
-void FGameplayTagContainer::RemoveAllTags()
+void FGameplayTagContainer::RemoveAllTags(int32 Slack)
 {
-	GameplayTags.Empty();
+	GameplayTags.Empty(Slack);
 }
 
 bool FGameplayTagContainer::Serialize(FArchive& Ar)
