dont generate new prediction keys for targeting data that does not send results to server (with no server RPC, the server never knows about the prediction key, so its never valid on the server)

diff --git a/Engine/Source/Runtime/GameplayAbilities/Private/Abilities/Tasks/AbilityTask_WaitTargetData.cpp b/Engine/Source/Runtime/GameplayAbilities/Private/Abilities/Tasks/AbilityTask_WaitTargetData.cpp
--- a/Engine/Source/Runtime/GameplayAbilities/Private/Abilities/Tasks/AbilityTask_WaitTargetData.cpp
+++ b/Engine/Source/Runtime/GameplayAbilities/Private/Abilities/Tasks/AbilityTask_WaitTargetData.cpp
@@ -204,7 +204,7 @@ void UAbilityTask_WaitTargetData::OnTargetDataReadyCallback(FGameplayAbilityTarg
 {
 	check(AbilitySystemComponent.IsValid());
 
-	FScopedPredictionWindow	ScopedPrediction(Ability.Get());
+	FScopedPredictionWindow	ScopedPrediction(Ability.Get(), ShouldReplicateDataToServer());
 	
 	if (ShouldReplicateDataToServer())
 	{
diff --git a/Engine/Source/Runtime/GameplayAbilities/Private/GameplayPrediction.cpp b/Engine/Source/Runtime/GameplayAbilities/Private/GameplayPrediction.cpp
--- a/Engine/Source/Runtime/GameplayAbilities/Private/GameplayPrediction.cpp
+++ b/Engine/Source/Runtime/GameplayAbilities/Private/GameplayPrediction.cpp
@@ -205,7 +205,7 @@ FScopedPredictionWindow::FScopedPredictionWindow(UAbilitySystemComponent* Abilit
 
 }
 
-FScopedPredictionWindow::FScopedPredictionWindow(UGameplayAbility* GameplayAbilityInstance)
+FScopedPredictionWindow::FScopedPredictionWindow(UGameplayAbility* GameplayAbilityInstance, bool bCanGenerateNewKey)
 {
 	// "Gets" the current/"best" prediction key and sets it as the prediciton key for this given abilities CurrentActivationInfo.
 	//	-On the server, it means we look for the prediction key that was given to us and either stored on the ActivationIfno or on the ScopedPredictionKey.
@@ -242,7 +242,7 @@ FScopedPredictionWindow::FScopedPredictionWindow(UGameplayAbility* GameplayAbili
 			ActivationInfo.ActivationMode = EGameplayAbilityActivationMode::Predicting;
 			ScopedPredictionKey = Owner->ScopedPedictionKey;
 		}
-		else
+		else if(bCanGenerateNewKey)
 		{
 			// We don't have a valid key, so generate a new one.
 			ActivationInfo.GenerateNewPredictionKey();
diff --git a/Engine/Source/Runtime/GameplayAbilities/Public/GameplayPrediction.h b/Engine/Source/Runtime/GameplayAbilities/Public/GameplayPrediction.h
--- a/Engine/Source/Runtime/GameplayAbilities/Public/GameplayPrediction.h
+++ b/Engine/Source/Runtime/GameplayAbilities/Public/GameplayPrediction.h
@@ -445,7 +445,7 @@ struct FScopedPredictionWindow
 	FScopedPredictionWindow(UAbilitySystemComponent* AbilitySystemComponent, FPredictionKey InPredictionKey);
 
 	/** To be called in the callsite where the predictive code will take place. This generates a new PredictionKey and acts as a synchonization point between client and server for that key.  */
-	FScopedPredictionWindow(UGameplayAbility* GameplayAbilityInstance);
+	FScopedPredictionWindow(UGameplayAbility* GameplayAbilityInstance, bool CanGenerateNewKey=true);
 
 	~FScopedPredictionWindow();
 
