fixes for removal of finalizecurvedata

diff --git a/Engine/Source/Runtime/GameplayAbilities/Private/AttributeSet.cpp b/Engine/Source/Runtime/GameplayAbilities/Private/AttributeSet.cpp
--- a/Engine/Source/Runtime/GameplayAbilities/Private/AttributeSet.cpp
+++ b/Engine/Source/Runtime/GameplayAbilities/Private/AttributeSet.cpp
@@ -105,54 +105,24 @@ FAttributeMetaData::FAttributeMetaData()
 
 }
 
-void FScalableFloat::FinalizeCurveData(const FGlobalCurveDataOverride *GlobalOverrides)
+float FScalableFloat::GetValueAtLevel(float Level) const
 {
-	static const FString ContextString = TEXT("FScalableFloat::FinalizeCurveData");
-
-	// We are a static value, so do nothing.
-	if (Curve.RowName == NAME_None)
-	{
-		return;
-	}
-
-	// Tied to an explicit table, so bind now.
-	if (Curve.CurveTable != NULL)
-	{
-		FinalCurve = Curve.GetCurve(ContextString);
-		return;
-	}
-
-	// We have overrides
-	if (GlobalOverrides)
+	if (Curve.CurveTable != nullptr)
 	{
-		for (UCurveTable* OverrideTable : GlobalOverrides->Overrides)
+		if (FinalCurve == nullptr)
 		{
-			FinalCurve = OverrideTable->FindCurve(Curve.RowName, ContextString, false);
-			if (FinalCurve)
-			{
-				return;
-			}
+			static const FString ContextString = TEXT("FScalableFloat::FinalizeCurveData");
+			FinalCurve = Curve.GetCurve(ContextString);
 		}
-	}
 
-	// Look at global defaults
-	const UCurveTable* GlobalTable = IGameplayAbilitiesModule::Get().GetAbilitySystemGlobals()->GetGlobalCurveTable();
-	if (GlobalTable)
-	{
-		FinalCurve = GlobalTable->FindCurve(Curve.RowName, ContextString, false);
-	}
-
-	if (!FinalCurve)
-	{
-		ABILITY_LOG(Warning, TEXT("Unable to find RowName: %s for FScalableFloat."), *Curve.RowName.ToString());
-	}
-}
-
-float FScalableFloat::GetValueAtLevel(float Level) const
-{
-	if (FinalCurve)
-	{
-		return Value * FinalCurve->Eval(Level);
+		if (FinalCurve != nullptr)
+		{
+			return Value * FinalCurve->Eval(Level);
+		}
+		else
+		{
+			ABILITY_LOG(Error, TEXT("Unable to find RowName: %s for FScalableFloat."), *Curve.RowName.ToString());
+		}
 	}
 
 	return Value;
diff --git a/Engine/Source/Runtime/GameplayAbilities/Public/AttributeSet.h b/Engine/Source/Runtime/GameplayAbilities/Public/AttributeSet.h
--- a/Engine/Source/Runtime/GameplayAbilities/Public/AttributeSet.h
+++ b/Engine/Source/Runtime/GameplayAbilities/Public/AttributeSet.h
@@ -169,17 +169,8 @@ public:
 	UPROPERTY(EditDefaultsOnly, Category=ScalableFloat)
 	FCurveTableRowHandle	Curve;
 
-	void FinalizeCurveData(const FGlobalCurveDataOverride *GlobalOverrides);
-
 	float GetValueAtLevel(float Level) const;
 
-	FScalableFloat MakeFinalizedCopy(const FGlobalCurveDataOverride *GlobalOverrides) const
-	{
-		FScalableFloat Copy(*this);
-		Copy.FinalizeCurveData(GlobalOverrides);
-		return Copy;
-	}
-
 	bool IsStatic() const
 	{
 		return (Curve.RowName == NAME_None);
@@ -227,7 +218,8 @@ public:
 
 private:
 
-	FRichCurve * FinalCurve;
+	// Cached direct pointer to RichCurve we should evaluate
+	mutable FRichCurve* FinalCurve;
 };
 
 /**
