Gameplay Effect Refactor work: Primarily focused on supporting custom magnitude and execution calculations
- Gameplay Effect
   - Replaced old FScalableFloat representation of magnitude with a custom struct capable of calculating magnitude three different ways:
      - FScalableFloat
      - An attribute backed float that can capture the value of an attribute (base, final, or bonus) from source or target and then use it in a calculation
      - Custom class-based calculation that can capture the value of multiple attributes and then do whatever it wants with them
   - Enabled concept of custom execution classes for gameplay effects
      - Can provide a BP or native class that will execute every time a gameplay effect executes and is capable of generating modifiers to execute in the same frame; Allows for custom logic, complex formulas, etc.
      - Can also capture attributes for use in calculations
      - Can process "scoped" modifiers that are active only for the execution, to allow for stuff like +300% Base Damage on this execution only
- Aggregator
   - Allowed aggregator to also evaluate with tag filters, which only count mods that were applied with tags matching the filters
   - Added EvaluateBonus, which computes (Evaluate - BaseValue)
   - Removed ability to suppress callbacks
- Made gameplay attributes in the editor respect edit const
- Made abilities module unregister all the customizations and things it sets up
- Add empty shell of details customization for execution scoped modifiers
- Renamed UGameplayModCalculation to UGameplayModMagnitudeCalculation
- Constify'd GetSet() and GetSetChecked()
- Remove ModifierSpec from FGameplayEffectModCallbackData, evaluated data has relevant necessary information within it now

- todo:
   - Detail customizations for scoped modifiers, new magnitude struct, etc.

diff --git a/Engine/Source/Editor/GameplayAbilitiesEditor/Private/AttributeDetails.cpp b/Engine/Source/Editor/GameplayAbilitiesEditor/Private/AttributeDetails.cpp
--- a/Engine/Source/Editor/GameplayAbilitiesEditor/Private/AttributeDetails.cpp
+++ b/Engine/Source/Editor/GameplayAbilitiesEditor/Private/AttributeDetails.cpp
@@ -92,13 +92,13 @@ void FAttributePropertyDetails::CustomizeHeader( TSharedRef<IPropertyHandle> Str
 	else
 */
 	{
-		// Gather all UAttriubute classes
+		// Gather all UAttribute classes
 		for (TObjectIterator<UClass> ClassIt; ClassIt; ++ClassIt)
 		{
 			UClass *Class = *ClassIt;
 			if (Class->IsChildOf(UAttributeSet::StaticClass()) && !FKismetEditorUtilities::IsClassABlueprintSkeleton(Class))
 			{
-				// Allow entire classes to be filtered globalyl
+				// Allow entire classes to be filtered globally
 				if (Class->HasMetaData(TEXT("HideInDetailsView")))
 				{
 					continue;
@@ -147,7 +147,7 @@ void FAttributePropertyDetails::CustomizeHeader( TSharedRef<IPropertyHandle> Str
 		PropertyOptions.Add(MakeShareable(new FString(FString::Printf(TEXT("%s.%s"), *Property->GetOuter()->GetName(), *Property->GetName()))));
 	}
 
-	// Fixme: this should be unified to use SGameplayATtributeWidget instead of custom combo box
+	// Fixme: this should be unified to use SGameplayAttributeWidget instead of custom combo box
 
 	HeaderRow.
 		NameContent()
@@ -167,6 +167,7 @@ void FAttributePropertyDetails::CustomizeHeader( TSharedRef<IPropertyHandle> Str
 
 				SNew(STextComboBox)
 				.ContentPadding(FMargin(2.0f, 2.0f))
+				.IsEnabled(!StructPropertyHandle->IsEditConst())
 				.OptionsSource( &PropertyOptions )
 				.InitiallySelectedItem(GetPropertyType())
 				.OnSelectionChanged( this, &FAttributePropertyDetails::OnChangeProperty )
diff --git a/Engine/Source/Editor/GameplayAbilitiesEditor/Private/GameplayAbilitiesEditorModule.cpp b/Engine/Source/Editor/GameplayAbilitiesEditor/Private/GameplayAbilitiesEditorModule.cpp
--- a/Engine/Source/Editor/GameplayAbilitiesEditor/Private/GameplayAbilitiesEditorModule.cpp
+++ b/Engine/Source/Editor/GameplayAbilitiesEditor/Private/GameplayAbilitiesEditorModule.cpp
@@ -6,6 +6,7 @@
 #include "AttributeSet.h"
 #include "GameplayEffectDetails.h"
 #include "GameplayModifierInfoDetails.h"
+#include "GameplayEffectExecutionScopedModifierInfoDetails.h"
 
 #include "IAssetTypeActions.h"
 #include "AssetToolsModule.h"
@@ -26,6 +27,12 @@ protected:
 private:
 	/** All created asset type actions.  Cached here so that we can unregister it during shutdown. */
 	TArray< TSharedPtr<IAssetTypeActions> > CreatedAssetTypeActions;
+
+	/** Pin factory for abilities graph; Cached so it can be unregistered */
+	TSharedPtr<FGameplayAbilitiesGraphPanelPinFactory> GameplayAbilitiesGraphPanelPinFactory;
+
+	/** Node factory for abilities graph; Cached so it can be unregistered */
+	TSharedPtr<FGameplayAbilitiesGraphPanelNodeFactory> GameplayAbilitiesGraphPanelNodeFactory;
 };
 
 IMPLEMENT_MODULE(FGameplayAbilitiesEditorModule, GameplayAbilitiesEditor)
@@ -37,6 +44,7 @@ void FGameplayAbilitiesEditorModule::StartupModule()
 	PropertyModule.RegisterCustomPropertyTypeLayout( "GameplayAttribute", FOnGetPropertyTypeCustomizationInstance::CreateStatic( &FAttributePropertyDetails::MakeInstance ) );
 	PropertyModule.RegisterCustomPropertyTypeLayout( "ScalableFloat", FOnGetPropertyTypeCustomizationInstance::CreateStatic( &FScalableFloatDetails::MakeInstance ) );
 	PropertyModule.RegisterCustomPropertyTypeLayout( "GameplayModifierInfo", FOnGetPropertyTypeCustomizationInstance::CreateStatic( &FGameplayModifierInfoCustomization::MakeInstance ) );
+	PropertyModule.RegisterCustomPropertyTypeLayout( "GameplayEffectExecutionScopedModifierInfo", FOnGetPropertyTypeCustomizationInstance::CreateStatic( &FGameplayEffectExecutionScopedModifierInfoDetails::MakeInstance ) );
 
 	PropertyModule.RegisterCustomClassLayout( "AttributeSet", FOnGetDetailCustomizationInstance::CreateStatic( &FAttributeDetails::MakeInstance ) );
 	PropertyModule.RegisterCustomClassLayout( "GameplayEffect", FOnGetDetailCustomizationInstance::CreateStatic( &FGameplayEffectDetails::MakeInstance ) );
@@ -48,10 +56,10 @@ void FGameplayAbilitiesEditorModule::StartupModule()
 	CreatedAssetTypeActions.Add(Action);
 
 	// Register factories for pins and nodes
-	TSharedPtr<FGameplayAbilitiesGraphPanelPinFactory> GameplayAbilitiesGraphPanelPinFactory = MakeShareable(new FGameplayAbilitiesGraphPanelPinFactory());
+	GameplayAbilitiesGraphPanelPinFactory = MakeShareable(new FGameplayAbilitiesGraphPanelPinFactory());
 	FEdGraphUtilities::RegisterVisualPinFactory(GameplayAbilitiesGraphPanelPinFactory);
 
-	TSharedPtr<FGameplayAbilitiesGraphPanelNodeFactory> GameplayAbilitiesGraphPanelNodeFactory = MakeShareable(new FGameplayAbilitiesGraphPanelNodeFactory());
+	GameplayAbilitiesGraphPanelNodeFactory = MakeShareable(new FGameplayAbilitiesGraphPanelNodeFactory());
 	FEdGraphUtilities::RegisterVisualNodeFactory(GameplayAbilitiesGraphPanelNodeFactory);
 }
 
@@ -66,5 +74,43 @@ void FGameplayAbilitiesEditorModule::ShutdownModule()
 	// This function may be called during shutdown to clean up your module.  For modules that support dynamic reloading,
 	// we call this function before unloading the module.
 
+	// Unregister customizations
+	if (FModuleManager::Get().IsModuleLoaded("PropertyEditor"))
+	{
+		FPropertyEditorModule& PropertyModule = FModuleManager::LoadModuleChecked<FPropertyEditorModule>("PropertyEditor");
+		PropertyModule.UnregisterCustomClassLayout("GameplayEffect");
+		PropertyModule.UnregisterCustomClassLayout("AttributeSet");
+
+		PropertyModule.UnregisterCustomPropertyTypeLayout("GameplayEffectExecutionScopedModifierInfo");
+		PropertyModule.UnregisterCustomPropertyTypeLayout("GameplayModifierInfo");
+		PropertyModule.UnregisterCustomPropertyTypeLayout("ScalableFloat");
+		PropertyModule.UnregisterCustomPropertyTypeLayout("GameplayAttribute");
+	}
+
+	// Unregister asset type actions
+	if (FModuleManager::Get().IsModuleLoaded("AssetTools"))
+	{
+		IAssetTools& AssetToolsModule = FModuleManager::GetModuleChecked<FAssetToolsModule>("AssetTools").Get();
+		for (auto& AssetTypeAction : CreatedAssetTypeActions)
+		{
+			if (AssetTypeAction.IsValid())
+			{
+				AssetToolsModule.UnregisterAssetTypeActions(AssetTypeAction.ToSharedRef());
+			}
+		}
+	}
 	CreatedAssetTypeActions.Empty();
+
+	// Unregister graph factories
+	if (GameplayAbilitiesGraphPanelPinFactory.IsValid())
+	{
+		FEdGraphUtilities::UnregisterVisualPinFactory(GameplayAbilitiesGraphPanelPinFactory);
+		GameplayAbilitiesGraphPanelPinFactory.Reset();
+	}
+
+	if (GameplayAbilitiesGraphPanelNodeFactory.IsValid())
+	{
+		FEdGraphUtilities::UnregisterVisualNodeFactory(GameplayAbilitiesGraphPanelNodeFactory);
+		GameplayAbilitiesGraphPanelNodeFactory.Reset();
+	}
 }
diff --git a/Engine/Source/Editor/GameplayAbilitiesEditor/Private/GameplayEffectExecutionScopedModifierInfoDetails.cpp b/Engine/Source/Editor/GameplayAbilitiesEditor/Private/GameplayEffectExecutionScopedModifierInfoDetails.cpp
new file mode 100644
--- /dev/null
+++ b/Engine/Source/Editor/GameplayAbilitiesEditor/Private/GameplayEffectExecutionScopedModifierInfoDetails.cpp
@@ -0,0 +1,31 @@
+// Copyright 1998-2013 Epic Games, Inc. All Rights Reserved.
+
+#include "AbilitySystemEditorPrivatePCH.h"
+#include "GameplayEffectExecutionScopedModifierInfoDetails.h"
+#include "DetailWidgetRow.h"
+#include "IDetailPropertyRow.h"
+#include "PropertyHandle.h"
+#include "GameplayEffect.h"
+
+#define LOCTEXT_NAMESPACE "GameplayEffectExecutionScopedModifierInfoDetailsCustomization"
+
+TSharedRef<IPropertyTypeCustomization> FGameplayEffectExecutionScopedModifierInfoDetails::MakeInstance()
+{
+	return MakeShareable(new FGameplayEffectExecutionScopedModifierInfoDetails());
+}
+
+void FGameplayEffectExecutionScopedModifierInfoDetails::CustomizeHeader(TSharedRef<IPropertyHandle> StructPropertyHandle, FDetailWidgetRow& HeaderRow, IPropertyTypeCustomizationUtils& StructCustomizationUtils)
+{
+	HeaderRow
+	.NameContent()
+	[
+		StructPropertyHandle->CreatePropertyNameWidget()
+	];
+}
+
+void FGameplayEffectExecutionScopedModifierInfoDetails::CustomizeChildren(TSharedRef<IPropertyHandle> StructPropertyHandle, IDetailChildrenBuilder& StructBuilder, IPropertyTypeCustomizationUtils& StructCustomizationUtils)
+{
+	// todo: Implement me!
+}
+
+#undef LOCTEXT_NAMESPACE
\ No newline at end of file
diff --git a/Engine/Source/Editor/GameplayAbilitiesEditor/Private/GameplayEffectExecutionScopedModifierInfoDetails.h b/Engine/Source/Editor/GameplayAbilitiesEditor/Private/GameplayEffectExecutionScopedModifierInfoDetails.h
new file mode 100644
--- /dev/null
+++ b/Engine/Source/Editor/GameplayAbilitiesEditor/Private/GameplayEffectExecutionScopedModifierInfoDetails.h
@@ -0,0 +1,20 @@
+// Copyright 1998-2013 Epic Games, Inc. All Rights Reserved.
+
+#pragma once
+
+#include "Editor/PropertyEditor/Public/PropertyEditorModule.h"
+
+class IPropertyHandle;
+class FDetailWidgetRow;
+class IDetailChildrenBuilder;
+
+/** Details customization for FGameplayEffectExecutionScopedModifierInfo */
+class FGameplayEffectExecutionScopedModifierInfoDetails : public IPropertyTypeCustomization
+{
+public:
+	static TSharedRef<IPropertyTypeCustomization> MakeInstance();
+
+	/** WIP */
+	virtual void CustomizeHeader(TSharedRef<IPropertyHandle> StructPropertyHandle, FDetailWidgetRow& HeaderRow, IPropertyTypeCustomizationUtils& StructCustomizationUtils) override;
+	virtual void CustomizeChildren(TSharedRef<IPropertyHandle> StructPropertyHandle, IDetailChildrenBuilder& StructBuilder, IPropertyTypeCustomizationUtils& StructCustomizationUtils) override;
+};
\ No newline at end of file
diff --git a/Engine/Source/Runtime/GameplayAbilities/Private/Abilities/Tasks/AbilityTask_WaitAttributeChange.cpp b/Engine/Source/Runtime/GameplayAbilities/Private/Abilities/Tasks/AbilityTask_WaitAttributeChange.cpp
--- a/Engine/Source/Runtime/GameplayAbilities/Private/Abilities/Tasks/AbilityTask_WaitAttributeChange.cpp
+++ b/Engine/Source/Runtime/GameplayAbilities/Private/Abilities/Tasks/AbilityTask_WaitAttributeChange.cpp
@@ -43,8 +43,8 @@ void UAbilityTask_WaitAttributeChange::OnAttributeChange(float NewValue, const F
 	}
 	else
 	{
-		if ((WithTag.IsValid() && !Data->EvaluatedData.SourceTags.HasTag(WithTag, EGameplayTagMatchType::IncludeParentTags, EGameplayTagMatchType::Explicit)) ||
-			(WithoutTag.IsValid() && Data->EvaluatedData.SourceTags.HasTag(WithoutTag, EGameplayTagMatchType::IncludeParentTags, EGameplayTagMatchType::Explicit)))
+		if ((WithTag.IsValid() && !Data->EffectSpec.CapturedSourceTags.GetAggregatedTags()->HasTag(WithTag, EGameplayTagMatchType::IncludeParentTags, EGameplayTagMatchType::Explicit)) ||
+			(WithoutTag.IsValid() && Data->EffectSpec.CapturedSourceTags.GetAggregatedTags()->HasTag(WithoutTag, EGameplayTagMatchType::IncludeParentTags, EGameplayTagMatchType::Explicit)))
 		{
 			// Failed tag check
 			return;
diff --git a/Engine/Source/Runtime/GameplayAbilities/Private/AbilitySystemComponent.cpp b/Engine/Source/Runtime/GameplayAbilities/Private/AbilitySystemComponent.cpp
--- a/Engine/Source/Runtime/GameplayAbilities/Private/AbilitySystemComponent.cpp
+++ b/Engine/Source/Runtime/GameplayAbilities/Private/AbilitySystemComponent.cpp
@@ -248,6 +248,16 @@ bool UAbilitySystemComponent::IsGameplayEffectActive(FActiveGameplayEffectHandle
 	return ActiveGameplayEffects.IsGameplayEffectActive(InHandle);
 }
 
+const FGameplayTagContainer* UAbilitySystemComponent::GetGameplayEffectSourceTagsFromHandle(FActiveGameplayEffectHandle Handle) const
+{
+	return ActiveGameplayEffects.GetGameplayEffectSourceTagsFromHandle(Handle);
+}
+
+const FGameplayTagContainer* UAbilitySystemComponent::GetGameplayEffectTargetTagsFromHandle(FActiveGameplayEffectHandle Handle) const
+{
+	return ActiveGameplayEffects.GetGameplayEffectTargetTagsFromHandle(Handle);
+}
+
 void UAbilitySystemComponent::CaptureAttributeForGameplayEffect(OUT FGameplayEffectAttributeCaptureSpec& OutCaptureSpec)
 {
 	// Verify the capture is happening on an attribute the component actually has a set for; if not, can't capture the value
diff --git a/Engine/Source/Runtime/GameplayAbilities/Private/AbilitySystemTestAttributeSet.cpp b/Engine/Source/Runtime/GameplayAbilities/Private/AbilitySystemTestAttributeSet.cpp
--- a/Engine/Source/Runtime/GameplayAbilities/Private/AbilitySystemTestAttributeSet.cpp
+++ b/Engine/Source/Runtime/GameplayAbilities/Private/AbilitySystemTestAttributeSet.cpp
@@ -88,16 +88,16 @@ bool UAbilitySystemTestAttributeSet::PreGameplayEffectExecute(struct FGameplayEf
 
 void UAbilitySystemTestAttributeSet::PostGameplayEffectExecute(const struct FGameplayEffectModCallbackData &Data)
 {
-	static UProperty *HealthProperty = FindFieldChecked<UProperty>(UAbilitySystemTestAttributeSet::StaticClass(), GET_MEMBER_NAME_CHECKED(UAbilitySystemTestAttributeSet, Health));
-	static UProperty *DamageProperty = FindFieldChecked<UProperty>(UAbilitySystemTestAttributeSet::StaticClass(), GET_MEMBER_NAME_CHECKED(UAbilitySystemTestAttributeSet, Damage));
+	static UProperty* HealthProperty = FindFieldChecked<UProperty>(UAbilitySystemTestAttributeSet::StaticClass(), GET_MEMBER_NAME_CHECKED(UAbilitySystemTestAttributeSet, Health));
+	static UProperty* DamageProperty = FindFieldChecked<UProperty>(UAbilitySystemTestAttributeSet::StaticClass(), GET_MEMBER_NAME_CHECKED(UAbilitySystemTestAttributeSet, Damage));
 
-	UProperty *ModifiedProperty = Data.ModifierSpec.Info.Attribute.GetUProperty();
+	UProperty* ModifiedProperty = Data.EvaluatedData.Attribute.GetUProperty();
 
 	// What property was modified?
 	if (DamageProperty == ModifiedProperty)
 	{
 		// Anytime Damage is applied with 'Damage.Fire' tag, there is a chance to apply a burning DOT
-		if (Data.EvaluatedData.SourceTags.HasTag( IGameplayTagsModule::RequestGameplayTag(FName(TEXT("FireDamage"))), EGameplayTagMatchType::IncludeParentTags, EGameplayTagMatchType::Explicit))
+		if (Data.EffectSpec.CapturedSourceTags.GetAggregatedTags()->HasTag( IGameplayTagsModule::RequestGameplayTag(FName(TEXT("FireDamage"))), EGameplayTagMatchType::IncludeParentTags, EGameplayTagMatchType::Explicit))
 		{
 			// Logic to rand() a burning DOT, if successful, apply DOT GameplayEffect to the target
 		}
diff --git a/Engine/Source/Runtime/GameplayAbilities/Private/GameplayEffect.cpp b/Engine/Source/Runtime/GameplayAbilities/Private/GameplayEffect.cpp
--- a/Engine/Source/Runtime/GameplayAbilities/Private/GameplayEffect.cpp
+++ b/Engine/Source/Runtime/GameplayAbilities/Private/GameplayEffect.cpp
@@ -8,7 +8,8 @@
 #include "GameplayTagsModule.h"
 #include "AbilitySystemComponent.h"
 #include "AbilitySystemInterface.h"
-#include "GameplayModCalculation.h"
+#include "GameplayModMagnitudeCalculation.h"
+#include "GameplayEffectExecutionCalculation.h"
 
 const float UGameplayEffect::INFINITE_DURATION = -1.f;
 const float UGameplayEffect::INSTANT_APPLICATION = 0.f;
@@ -46,118 +47,172 @@ void UGameplayEffect::PostLoad()
 {
 	Super::PostLoad();
 
+	// Temporary post-load fix-up to preserve magnitude data
+	for (FGameplayModifierInfo& CurModInfo : Modifiers)
+	{
+		// If the old magnitude actually had some value in it, copy it over and then clear out the old data
+		if (CurModInfo.Magnitude.Value != 0.f || CurModInfo.Magnitude.Curve.IsValid())
+		{
+			CurModInfo.ModifierMagnitude.ScalableFloatMagnitude = CurModInfo.Magnitude;
+			CurModInfo.Magnitude = FScalableFloat();
+		}
+	}
+
 	ValidateGameplayEffect();
 }
 
 void UGameplayEffect::ValidateGameplayEffect()
 {
-	// todo: add a check here for instant effects that modify incoming/outgoing GEs
 }
 
-bool FGameplayEffectAttributeCaptureDefinition::operator==(const FGameplayEffectAttributeCaptureDefinition& Other) const
-{
-	return ((AttributeToCapture == Other.AttributeToCapture) && (AttributeSource == Other.AttributeSource) && (bSnapshot == Other.bSnapshot));
-}
+// --------------------------------------------------------------------------------------------------------------------------------------------------------
+//
+//	FAttributeBasedFloat
+//
+// --------------------------------------------------------------------------------------------------------------------------------------------------------
 
-bool FGameplayEffectAttributeCaptureDefinition::operator!=(const FGameplayEffectAttributeCaptureDefinition& Other) const
+float FAttributeBasedFloat::CalculateMagnitude(const FGameplayEffectSpec& InRelevantSpec) const
 {
-	return ((AttributeToCapture != Other.AttributeToCapture) || (AttributeSource != Other.AttributeSource) || (bSnapshot != Other.bSnapshot));
-}
+	const FGameplayEffectAttributeCaptureSpec* CaptureSpec = InRelevantSpec.CapturedRelevantAttributes.FindCaptureSpecByDefinition(BackingAttribute, true);
+	checkf(CaptureSpec, TEXT("Attempted to calculate an attribute-based float from spec: %s that did not have the required captured attribute: %s"), 
+		*InRelevantSpec.ToSimpleString(), *BackingAttribute.ToSimpleString());
 
-FString FGameplayEffectAttributeCaptureDefinition::ToSimpleString() const
-{
-	return FString::Printf(TEXT("Attribute: %s, Capture Point: %s, Snapshot: %d"), *AttributeToCapture.GetName(), AttributeSource == EGameplayEffectAttributeCaptureSource::Source ? TEXT("Source") : TEXT("Target"), bSnapshot);
-}
+	float AttribValue = 0.f;
 
-FGameplayEffectAttributeCaptureSpec::FGameplayEffectAttributeCaptureSpec()
-{
-}
+	// Base value can be calculated w/o evaluation parameters
+	if (AttributeCalculationType == EAttributeBasedFloatCalculationType::AttributeBaseValue)
+	{
+		CaptureSpec->AttemptCalculateAttributeBaseValue(AttribValue);
+	}
+	// Set up eval params to handle magnitude or bonus magnitude calculations
+	else
+	{
+		FAggregatorEvaluateParameters EvaluationParameters;
+		EvaluationParameters.SourceTags = InRelevantSpec.CapturedSourceTags.GetAggregatedTags();
+		EvaluationParameters.TargetTags = InRelevantSpec.CapturedTargetTags.GetAggregatedTags();
+		EvaluationParameters.AppliedSourceTagFilter = SourceTagFilter;
+		EvaluationParameters.AppliedTargetTagFilter = TargetTagFilter;
 
-FGameplayEffectAttributeCaptureSpec::FGameplayEffectAttributeCaptureSpec(const FGameplayEffectAttributeCaptureDefinition& InDefinition)
-	: BackingDefinition(InDefinition)
-{
-}
+		if (AttributeCalculationType == EAttributeBasedFloatCalculationType::AttributeMagnitude)
+		{
+			CaptureSpec->AttemptCalculateAttributeMagnitude(EvaluationParameters, AttribValue);
+		}
+		else
+		{
+			CaptureSpec->AttemptCalculateAttributeBonusMagnitude(EvaluationParameters, AttribValue);
+		}
+	}
 
-const FGameplayEffectAttributeCaptureDefinition& FGameplayEffectAttributeCaptureSpec::GetBackingDefinition() const
-{
-	return BackingDefinition;
+	const float SpecLvl = InRelevantSpec.GetLevel();
+	return ((Coefficient.GetValueAtLevel(SpecLvl) * (AttribValue + PreMultiplyAdditiveValue.GetValueAtLevel(SpecLvl))) + PostMultiplyAdditiveValue.GetValueAtLevel(SpecLvl));
 }
 
-const FAggregator& FGameplayEffectAttributeCaptureSpec::GetAggregator() const
+// --------------------------------------------------------------------------------------------------------------------------------------------------------
+//
+//	FGameplayEffectMagnitude
+//
+// --------------------------------------------------------------------------------------------------------------------------------------------------------
+
+bool FGameplayEffectModifierMagnitude::CanCalculateMagnitude(const FGameplayEffectSpec& InRelevantSpec) const
 {
-	FAggregator* Aggregator = AttributeAggregator.Get();
-	if (ensure(Aggregator))
-	{
-		return *Aggregator;
-	}
+	// Only can calculate magnitude properly if all required capture definitions are fulfilled by the spec
+	TArray<FGameplayEffectAttributeCaptureDefinition> ReqCaptureDefs;
+	GetAttributeCaptureDefinitions(ReqCaptureDefs);
 
-	ABILITY_LOG(Error, TEXT("FGameplayEffectAttributeCaptureSpec::GetAggregator called with no valid aggregator"));
-	static FAggregator InvalidAggregator;
-	return InvalidAggregator;
+	return InRelevantSpec.HasValidCapturedAttributes(ReqCaptureDefs);
 }
 
-void FGameplayEffectAttributeCaptureSpecContainer::AddCaptureDefinition(const FGameplayEffectAttributeCaptureDefinition& InCaptureDefinition)
+bool FGameplayEffectModifierMagnitude::AttemptCalculateMagnitude(const FGameplayEffectSpec& InRelevantSpec, OUT float& OutCalculatedMagnitude) const
 {
-	const bool bSourceAttribute = (InCaptureDefinition.AttributeSource == EGameplayEffectAttributeCaptureSource::Source);
-	TArray<FGameplayEffectAttributeCaptureSpec>& AttributeArray = (bSourceAttribute ? SourceAttributes : TargetAttributes);
+	OutCalculatedMagnitude = 0.f;
+	
+	const bool bCanCalc = CanCalculateMagnitude(InRelevantSpec);
+	if (bCanCalc)
+	{
+		switch (MagnitudeCalculationType)
+		{
+			case EGameplayEffectMagnitudeCalculation::ScalableFloat:
+			{
+				OutCalculatedMagnitude = ScalableFloatMagnitude.GetValueAtLevel(InRelevantSpec.GetLevel());
+			}
+			break;
 
-	// Only add additional captures if this exact capture definition isn't already being handled
-	if (!AttributeArray.ContainsByPredicate(
-			[&](const FGameplayEffectAttributeCaptureSpec& Element)
+			case EGameplayEffectMagnitudeCalculation::AttributeBased:
 			{
-				return Element.GetBackingDefinition() == InCaptureDefinition;
-			}))
-	{
-		AttributeArray.Add(FGameplayEffectAttributeCaptureSpec(InCaptureDefinition));
+				OutCalculatedMagnitude = AttributeBasedMagnitude.CalculateMagnitude(InRelevantSpec);
+			}
+			break;
 
-		if (!InCaptureDefinition.bSnapshot)
-		{
-			bHasNonSnapshottedAttributes = true;
+			case EGameplayEffectMagnitudeCalculation::CustomCalculationClass:
+			{
+				const UGameplayModMagnitudeCalculation* CalcCDO = CalculationClassMagnitude->GetDefaultObject<UGameplayModMagnitudeCalculation>();
+				check(CalcCDO);
+
+				OutCalculatedMagnitude = CalcCDO->CalculateMagnitude(InRelevantSpec);
+			}
+			break;
 		}
 	}
+
+	return bCanCalc;
 }
 
-void FGameplayEffectAttributeCaptureSpecContainer::CaptureAttributes(UAbilitySystemComponent* InAbilitySystemComponent, EGameplayEffectAttributeCaptureSource InCaptureSource)
+void FGameplayEffectModifierMagnitude::GetAttributeCaptureDefinitions(OUT TArray<FGameplayEffectAttributeCaptureDefinition>& OutCaptureDefs) const
 {
-	if (InAbilitySystemComponent)
+	OutCaptureDefs.Empty();
+
+	switch (MagnitudeCalculationType)
 	{
-		const bool bSourceComponent = (InCaptureSource == EGameplayEffectAttributeCaptureSource::Source);
-		TArray<FGameplayEffectAttributeCaptureSpec>& AttributeArray = (bSourceComponent ? SourceAttributes : TargetAttributes);
+		case EGameplayEffectMagnitudeCalculation::AttributeBased:
+		{
+			OutCaptureDefs.Add(AttributeBasedMagnitude.BackingAttribute);
+		}
+		break;
 
-		// Capture every spec's requirements from the specified component
-		for (FGameplayEffectAttributeCaptureSpec& CurCaptureSpec : AttributeArray)
+		case EGameplayEffectMagnitudeCalculation::CustomCalculationClass:
 		{
-			InAbilitySystemComponent->CaptureAttributeForGameplayEffect(CurCaptureSpec);
+			if (CalculationClassMagnitude)
+			{
+				const UGameplayModMagnitudeCalculation* CalcCDO = CalculationClassMagnitude->GetDefaultObject<UGameplayModMagnitudeCalculation>();
+				check(CalcCDO);
+
+				OutCaptureDefs.Append(CalcCDO->GetAttributeCaptureDefinitions());
+			}
 		}
+		break;
 	}
 }
 
-const FAggregator& FGameplayEffectAttributeCaptureSpecContainer::GetAggregator(const FGameplayEffectAttributeCaptureDefinition& InDefinition) const
-{
-	const bool bSourceAttribute = (InDefinition.AttributeSource == EGameplayEffectAttributeCaptureSource::Source);
-	const TArray<FGameplayEffectAttributeCaptureSpec>& AttributeArray = (bSourceAttribute ? SourceAttributes : TargetAttributes);
 
-	const FGameplayEffectAttributeCaptureSpec* MatchingSpec = AttributeArray.FindByPredicate(
-		[&](const FGameplayEffectAttributeCaptureSpec& Element)
-		{
-			return Element.GetBackingDefinition() == InDefinition;
-		});
+// --------------------------------------------------------------------------------------------------------------------------------------------------------
+//
+//	FGameplayEffectExecutionDefinition
+//
+// --------------------------------------------------------------------------------------------------------------------------------------------------------
 
-	if (MatchingSpec)
+void FGameplayEffectExecutionDefinition::GetAttributeCaptureDefinitions(OUT TArray<FGameplayEffectAttributeCaptureDefinition>& OutCaptureDefs) const
+{
+	OutCaptureDefs.Empty();
+
+	if (CalculationClass)
 	{
-		return MatchingSpec->GetAggregator();
+		const UGameplayEffectExecutionCalculation* CalculationCDO = Cast<UGameplayEffectExecutionCalculation>(CalculationClass->ClassDefaultObject);
+		check(CalculationCDO);
+
+		OutCaptureDefs.Append(CalculationCDO->GetAttributeCaptureDefinitions());
 	}
 
-	ABILITY_LOG(Error, TEXT("FGameplayEffectAttributeCaptureSpecContainer::GetAggregator called for attribute definition %s when that attribute was not captured."), *InDefinition.ToSimpleString());
-	static FAggregator InvalidAggregator;
-	return InvalidAggregator;
-}
+	// Scoped modifiers might have custom magnitude calculations, requiring additional captured attributes
+	for (const FGameplayEffectExecutionScopedModifierInfo& CurScopedMod : CalculationModifiers)
+	{
+		TArray<FGameplayEffectAttributeCaptureDefinition> ScopedModMagDefs;
+		CurScopedMod.ModifierMagnitude.GetAttributeCaptureDefinitions(ScopedModMagDefs);
 
-bool FGameplayEffectAttributeCaptureSpecContainer::HasNonSnapshottedAttributes() const
-{
-	return bHasNonSnapshottedAttributes;
+		OutCaptureDefs.Append(ScopedModMagDefs);
+	}
 }
 
+
 // --------------------------------------------------------------------------------------------------------------------------------------------------------
 //
 //	FGameplayEffectSpec
@@ -179,10 +234,31 @@ FGameplayEffectSpec::FGameplayEffectSpec(const UGameplayEffect* InDef, const FGa
 		new (Modifiers)FModifierSpec(ModInfo);
 	}
 
+	// Prep the spec with all of the attribute captures it will need to perform
+	SetupAttributeCaptureDefinitions();
+	
+	// Add the GameplayEffect asset tags to the source Spec tags
+	CapturedSourceTags.GetSpecTags().AppendTags(InDef->GameplayEffectTags);
+
+	// Make TargetEffectSpecs too
+	for (UGameplayEffect* TargetDef : InDef->TargetEffects)
+	{
+		TargetEffectSpecs.Add(TSharedRef<FGameplayEffectSpec>(new FGameplayEffectSpec(TargetDef, EffectContext, Level)));
+	}
+
+	// Everything is setup now, capture data from our source
+	CaptureDataFromSource();
+}
+
+void FGameplayEffectSpec::SetupAttributeCaptureDefinitions()
+{
 	// Gather all capture definitions from modifiers
 	for (const FModifierSpec& CurModSpec : Modifiers)
 	{
-		for (const FGameplayEffectAttributeCaptureDefinition& CurCaptureDef : CurModSpec.Info.MagnitudeV2.CustomMagnitudeCalculationAttributes)
+		TArray<FGameplayEffectAttributeCaptureDefinition> CalculationCaptureDefs;
+		CurModSpec.Info.ModifierMagnitude.GetAttributeCaptureDefinitions(CalculationCaptureDefs);
+		
+		for (const FGameplayEffectAttributeCaptureDefinition& CurCaptureDef : CalculationCaptureDefs)
 		{
 			CapturedRelevantAttributes.AddCaptureDefinition(CurCaptureDef);
 		}
@@ -191,32 +267,14 @@ FGameplayEffectSpec::FGameplayEffectSpec(const UGameplayEffect* InDef, const FGa
 	// Gather all capture definitions from executions
 	for (const FGameplayEffectExecutionDefinition& Exec : Def->Executions)
 	{
-		if (Exec.CalculationClass == nullptr)
-			continue;
-
-		TArray<FGameplayEffectAttributeCaptureDefinition> Reqs;
-		const UGameplayEffectCalculation* CDO = Cast<UGameplayEffectCalculation>(Exec.CalculationClass->ClassDefaultObject);
-		
-		CDO->GetCaptureDefinitions(Reqs);
+		TArray<FGameplayEffectAttributeCaptureDefinition> ExecutionCaptureDefs;
+		Exec.GetAttributeCaptureDefinitions(ExecutionCaptureDefs);
 
-		// Maybe it's better to pass in CapturedRelevantAttributes.AddCaptureDefinition(CurCaptureDef); to GetCaptureDefinitions?
-		for (const FGameplayEffectAttributeCaptureDefinition& Req : Reqs)
+		for (const FGameplayEffectAttributeCaptureDefinition& CurExecCaptureDef : ExecutionCaptureDefs)
 		{
-			CapturedRelevantAttributes.AddCaptureDefinition(Req);
+			CapturedRelevantAttributes.AddCaptureDefinition(CurExecCaptureDef);
 		}
 	}
-
-	// Add the GameplayEffect asset tags to the source Spec tags
-	CapturedSourceTags.GetSpecTags().AppendTags(InDef->GameplayEffectTags);
-
-	// Make TargetEffectSpecs too
-	for (UGameplayEffect* TargetDef : InDef->TargetEffects)
-	{
-		TargetEffectSpecs.Add(TSharedRef<FGameplayEffectSpec>(new FGameplayEffectSpec(TargetDef, EffectContext, Level)));
-	}
-
-	// Everything is setup now, capture data from our source
-	CaptureDataFromSource();
 }
 
 void FGameplayEffectSpec::CaptureDataFromSource()
@@ -282,8 +340,6 @@ float FGameplayEffectSpec::GetChanceToExecuteOnGameplayEffect() const
 float FGameplayEffectSpec::GetModifierMagnitude(const FModifierSpec& ModSpec) const
 {
 	return ModSpec.EvaluatedMagnitude;
-
-	//return ModSpec.Info.Magnitude.GetValueAtLevel(Level);
 }
 
 void FGameplayEffectSpec::CalculateModifierMagnitudes()
@@ -311,6 +367,11 @@ float FGameplayEffectSpec::GetMagnitude(const FGameplayAttribute &Attribute) con
 	return CurrentMagnitude;
 }
 
+bool FGameplayEffectSpec::HasValidCapturedAttributes(const TArray<FGameplayEffectAttributeCaptureDefinition>& InCaptureDefsToCheck) const
+{
+	return CapturedRelevantAttributes.HasValidCapturedAttributes(InCaptureDefsToCheck);
+}
+
 const FGameplayEffectModifiedAttribute* FGameplayEffectSpec::GetModifiedAttribute(const FGameplayAttribute& Attribute) const
 {
 	for (const FGameplayEffectModifiedAttribute& ModifiedAttribute : ModifiedAttributes)
@@ -386,8 +447,8 @@ void FGameplayEffectSpec::SetContext(FGameplayEffectContextHandle NewEffectConte
 // --------------------------------------------------------------------------------------------------------------------------------------------------------
 
 FModifierSpec::FModifierSpec(const FGameplayModifierInfo& InInfo)
-: Info(InInfo)
-, EvaluatedMagnitude(0.f)
+	: Info(InInfo)
+	, EvaluatedMagnitude(0.f)
 {
 
 }
@@ -399,25 +460,178 @@ float FModifierSpec::GetEvaluatedMagnitude() const
 
 void FModifierSpec::CalculateMagnitude(OUT FGameplayEffectSpec& OwnerSpec)
 {	
-	EvaluatedMagnitude = Info.Magnitude.GetValueAtLevel(OwnerSpec.GetLevel());
+	if (!Info.ModifierMagnitude.AttemptCalculateMagnitude(OwnerSpec, EvaluatedMagnitude))
+	{
+		EvaluatedMagnitude = 0.f;
+		ABILITY_LOG(Warning, TEXT("Modifier on spec: %s was asked to CalculateMagnitude and failed, falling back to 0."), *OwnerSpec.ToSimpleString());
+	}
+}
+
+FGameplayEffectAttributeCaptureSpec::FGameplayEffectAttributeCaptureSpec()
+{
+}
 
-	// @todo: this should have an enumeration over which to use
-	/*
-	if (Info.MagnitudeV2.CalculationClassMagnitude)
+FGameplayEffectAttributeCaptureSpec::FGameplayEffectAttributeCaptureSpec(const FGameplayEffectAttributeCaptureDefinition& InDefinition)
+	: BackingDefinition(InDefinition)
+{
+}
+
+bool FGameplayEffectAttributeCaptureSpec::HasValidCapture() const
+{
+	return (AttributeAggregator.Get() != nullptr);
+}
+
+bool FGameplayEffectAttributeCaptureSpec::AttemptCalculateAttributeMagnitude(const FAggregatorEvaluateParameters& InEvalParams, OUT float& OutMagnitude) const
+{
+	OutMagnitude = 0.f;
+
+	FAggregator* Agg = AttributeAggregator.Get();
+	if (Agg)
 	{
-		const UGameplayModCalculation* ModCalcCDO = Info.MagnitudeV2.CalculationClassMagnitude->GetDefaultObject<UGameplayModCalculation>();
-		check(ModCalcCDO);
+		OutMagnitude = Agg->Evaluate(InEvalParams);
+		return true;
+	}
+
+	return false;
+}
+
+bool FGameplayEffectAttributeCaptureSpec::AttemptCalculateAttributeBaseValue(OUT float& OutBaseValue) const
+{
+	OutBaseValue = 0.f;
 
-		EvaluatedMagnitude = ModCalcCDO->CalculateMagnitude(OwnerSpec);
+	FAggregator* Agg = AttributeAggregator.Get();
+	if (Agg)
+	{
+		OutBaseValue = Agg->GetBaseValue();
+		return true;
 	}
-	else
+
+	return false;
+}
+
+bool FGameplayEffectAttributeCaptureSpec::AttemptCalculateAttributeBonusMagnitude(const FAggregatorEvaluateParameters& InEvalParams, OUT float& OutBonusMagnitude) const
+{
+	OutBonusMagnitude = 0.f;
+
+	FAggregator* Agg = AttributeAggregator.Get();
+	if (Agg)
+	{
+		OutBonusMagnitude = Agg->EvaluateBonus(InEvalParams);
+		return true;
+	}
+
+	return false;
+}
+
+bool FGameplayEffectAttributeCaptureSpec::AttemptGetAttributeAggregatorSnapshot(OUT FAggregator& OutAggregatorSnapshot) const
+{
+	FAggregator* Agg = AttributeAggregator.Get();
+	if (Agg)
+	{
+		OutAggregatorSnapshot.TakeSnapshotOf(*Agg);
+		return true;
+	}
+
+	return false;
+}
+
+bool FGameplayEffectAttributeCaptureSpec::AttemptAddAggregatorModsToAggregator(OUT FAggregator& OutAggregatorToAddTo) const
+{
+	FAggregator* Agg = AttributeAggregator.Get();
+	if (Agg)
+	{
+		OutAggregatorToAddTo.AddModsFrom(*Agg);
+		return true;
+	}
+
+	return false;
+}
+
+const FGameplayEffectAttributeCaptureDefinition& FGameplayEffectAttributeCaptureSpec::GetBackingDefinition() const
+{
+	return BackingDefinition;
+}
+
+void FGameplayEffectAttributeCaptureSpecContainer::AddCaptureDefinition(const FGameplayEffectAttributeCaptureDefinition& InCaptureDefinition)
+{
+	const bool bSourceAttribute = (InCaptureDefinition.AttributeSource == EGameplayEffectAttributeCaptureSource::Source);
+	TArray<FGameplayEffectAttributeCaptureSpec>& AttributeArray = (bSourceAttribute ? SourceAttributes : TargetAttributes);
+
+	// Only add additional captures if this exact capture definition isn't already being handled
+	if (!AttributeArray.ContainsByPredicate(
+			[&](const FGameplayEffectAttributeCaptureSpec& Element)
+			{
+				return Element.GetBackingDefinition() == InCaptureDefinition;
+			}))
+	{
+		AttributeArray.Add(FGameplayEffectAttributeCaptureSpec(InCaptureDefinition));
+
+		if (!InCaptureDefinition.bSnapshot)
+		{
+			bHasNonSnapshottedAttributes = true;
+		}
+	}
+}
+
+void FGameplayEffectAttributeCaptureSpecContainer::CaptureAttributes(UAbilitySystemComponent* InAbilitySystemComponent, EGameplayEffectAttributeCaptureSource InCaptureSource)
+{
+	if (InAbilitySystemComponent)
+	{
+		const bool bSourceComponent = (InCaptureSource == EGameplayEffectAttributeCaptureSource::Source);
+		TArray<FGameplayEffectAttributeCaptureSpec>& AttributeArray = (bSourceComponent ? SourceAttributes : TargetAttributes);
+
+		// Capture every spec's requirements from the specified component
+		for (FGameplayEffectAttributeCaptureSpec& CurCaptureSpec : AttributeArray)
+		{
+			InAbilitySystemComponent->CaptureAttributeForGameplayEffect(CurCaptureSpec);
+		}
+	}
+}
+
+const FGameplayEffectAttributeCaptureSpec* FGameplayEffectAttributeCaptureSpecContainer::FindCaptureSpecByDefinition(const FGameplayEffectAttributeCaptureDefinition& InDefinition, bool bOnlyIncludeValidCapture) const
+{
+	const bool bSourceAttribute = (InDefinition.AttributeSource == EGameplayEffectAttributeCaptureSource::Source);
+	const TArray<FGameplayEffectAttributeCaptureSpec>& AttributeArray = (bSourceAttribute ? SourceAttributes : TargetAttributes);
+
+	const FGameplayEffectAttributeCaptureSpec* MatchingSpec = AttributeArray.FindByPredicate(
+		[&](const FGameplayEffectAttributeCaptureSpec& Element)
 	{
-		EvaluatedMagnitude = Info.MagnitudeV2.ScalableFloatMagnitude.GetValueAtLevel(OwnerSpec.GetLevel());
+		return Element.GetBackingDefinition() == InDefinition;
+	});
+
+	// Null out the found results if the caller only wants valid captures and we don't have one yet
+	if (MatchingSpec && bOnlyIncludeValidCapture && !MatchingSpec->HasValidCapture())
+	{
+		MatchingSpec = nullptr;
+	}
+
+	return MatchingSpec;
+}
+
+bool FGameplayEffectAttributeCaptureSpecContainer::HasValidCapturedAttributes(const TArray<FGameplayEffectAttributeCaptureDefinition>& InCaptureDefsToCheck) const
+{
+	bool bHasValid = true;
+
+	for (const FGameplayEffectAttributeCaptureDefinition& CurDef : InCaptureDefsToCheck)
+	{
+		const FGameplayEffectAttributeCaptureSpec* CaptureSpec = FindCaptureSpecByDefinition(CurDef, true);
+		if (!CaptureSpec)
+		{
+			bHasValid = false;
+			break;
+		}
 	}
-	*/
+
+	return bHasValid;
+}
+
+bool FGameplayEffectAttributeCaptureSpecContainer::HasNonSnapshottedAttributes() const
+{
+	return bHasNonSnapshottedAttributes;
 }
 
 
+
 // --------------------------------------------------------------------------------------------------------------------------------------------------------
 //
 //	FActiveGameplayEffect
@@ -567,7 +781,7 @@ void FActiveGameplayEffectsContainer::RegisterWithOwner(UAbilitySystemComponent*
 /** This is the main function that executes a GameplayEffect on Attributes and ActiveGameplayEffects */
 void FActiveGameplayEffectsContainer::ExecuteActiveEffectsFrom(FGameplayEffectSpec &Spec, FPredictionKey PredictionKey)
 {
-	// If there are no modifiers, we always want to apply the GameplayCue. If there are modifiers, we only want to invoke the GameplayCue if one of them went through (coudl be blocked by immunity or % chance roll)
+	// If there are no modifiers, we always want to apply the GameplayCue. If there are modifiers, we only want to invoke the GameplayCue if one of them went through (could be blocked by immunity or % chance roll)
 	bool InvokeGameplayCueExecute = (Spec.Modifiers.Num() == 0);
 
 	// check if this is a stacking effect and if it is the active stacking effect.
@@ -597,123 +811,36 @@ void FActiveGameplayEffectsContainer::ExecuteActiveEffectsFrom(FGameplayEffectSp
 
 	for (const FModifierSpec &Mod : Spec.Modifiers)
 	{
-		UAttributeSet* AttributeSet = const_cast<UAttributeSet*>(Owner->GetAttributeSubobject(Mod.Info.Attribute.GetAttributeSetClass()));
-		if (AttributeSet == NULL)
-		{
-			// Our owner doesn't have this attribute, so we can't do anything
-			ABILITY_LOG(Log, TEXT("%s does not have attribute %s. Skipping modifier"), *Owner->GetPathName(), *Mod.Info.Attribute.GetName());
-			continue;
-		}
-
-		ABILITY_LOG_SCOPE(TEXT("Executing Attribute Mod %s"), *Mod.ToSimpleString());
-
-		// First, evaluate all of our data
-			
-		// GE_Remove: still fixing this up:
-
-		float EvaluatedMagnitude = Mod.Info.Magnitude.GetValueAtLevel(Spec.GetLevel());
-
-		FGameplayModifierEvaluatedData EvaluatedData(EvaluatedMagnitude);
-
-		FGameplayEffectModCallbackData ExecuteData(Spec, Mod, EvaluatedData, *Owner);
-
-		/** 
-		 *  This should apply 'gamewide' rules. Such as clamping Health to MaxHealth or granting +3 health for every point of strength, etc 
-		 *	PreAttributeModify can return false to 'throw out' this modification.
-		 */
-		if (AttributeSet->PreGameplayEffectExecute(ExecuteData) == false)
-		{
-			continue;
-		}
-
-		// Todo: Tags/application checks here - make sure we can still apply (GE_Remove ??? is this still valid?)
-		InvokeGameplayCueExecute = true;
-
-		// Do we have active GE's that are already modifying this?
-		FAggregatorRef* RefPtr = AttributeAggregatorMap.Find(Mod.Info.Attribute);
-		if (RefPtr)
-		{
-			ABILITY_LOG(Log, TEXT("Property %s has active mods. Adding to Aggregator."), *Mod.Info.Attribute.GetName());
-			RefPtr->Get()->ExecModOnBaseValue(Mod.Info.ModifierOp, Mod.GetEvaluatedMagnitude() );
-		}
-		else
-		{
-			// Modify the property inplace, without putting it in the AttributeAggregatorMap map
-			float		CurrentValueOfProperty = Owner->GetNumericAttribute(Mod.Info.Attribute);
-			const float NewPropertyValue = FAggregator::StaticExecModOnBaseValue(CurrentValueOfProperty, Mod.Info.ModifierOp, Mod.GetEvaluatedMagnitude());
-
-			InternalUpdateNumericalAttribute(Mod.Info.Attribute, NewPropertyValue, &ExecuteData);
-		}
-
-		FGameplayEffectModifiedAttribute* ModifiedAttribute = Spec.GetModifiedAttribute(Mod.Info.Attribute);
-		if (!ModifiedAttribute)
-		{
-			// If we haven't already created a modified attribute holder, create it
-			ModifiedAttribute = Spec.AddModifiedAttribute(Mod.Info.Attribute);
-		}
-		ModifiedAttribute->TotalMagnitude += EvaluatedData.Magnitude;
-
-		/** This should apply 'gamewide' rules. Such as clamping Health to MaxHealth or granting +3 health for every point of strength, etc */
-		AttributeSet->PostGameplayEffectExecute(ExecuteData);
+		FGameplayModifierEvaluatedData EvalData(Mod.Info.Attribute, Mod.Info.ModifierOp, Mod.GetEvaluatedMagnitude());
+		InvokeGameplayCueExecute |= InternalExecuteMod(Spec, EvalData);
 	}
 
 	// ------------------------------------------------------
 	//	Executions
 	//		This will run custom code to 'do stuff'
 	// ------------------------------------------------------
-	 
-	for (int32 ExecIdx=0; ExecIdx < Spec.Def->Executions.Num(); ++ExecIdx)
+	
+	for (const FGameplayEffectExecutionDefinition& CurExecDef : Spec.Def->Executions)
 	{
-		const FGameplayEffectExecutionDefinition& Exec = Spec.Def->Executions[ExecIdx];
-		if (Exec.CalculationClass == nullptr)
-			continue;
-
-		const UGameplayEffectCalculation* CDO = Cast<UGameplayEffectCalculation>(Exec.CalculationClass->ClassDefaultObject);
-		check(CDO);
-
-
-		// Apply custom inline mods
-		FActiveGameplayEffectHandle ModifierHandle = FActiveGameplayEffectHandle::GenerateNewHandle(Owner);
-		TSet<FAggregator*>	ModifiedAggregators;
-		for (const FExtensionAttributeModifierInfo& Mod : Exec.CalculationModifiers)
+		if (CurExecDef.CalculationClass)
 		{
-			// Evaluate this mods current value. Right now this is always a scalable float, but we'd like to support
-			// pulling arbitrary attribute values here too - these would be pulled from the Spec's captured attribute only?
-			// But bottom line - we evaluate them right here and that is their value for this scope. Right?
-			
-			// We do a const cast here becase we are being really careful! We will add Mods below but we will also remove them once the 
-			// custom calculation is finished. Custom Exec functions will only get const references to the aggregator.
-			
-			// If the custom exec functions want to add mods inplace, they can do it either like:
-			//	A) like we are doing it here (careful const cast + remove)
-			//  B) We give them a way to clone the aggregators on the stack and allow them to do whatever they want to those copies
-			//  C) We provide another optional parameter into ::Evaluate that takes additional mods.
-			 
-			
-			// @todo: Just always assuming snapshot for now, which is wrong. This should use a different struct.
-			FAggregator& Aggregator = const_cast<FAggregator&>(Spec.CapturedRelevantAttributes.GetAggregator(FGameplayEffectAttributeCaptureDefinition(Mod.Attribute, Mod.Source, true)));
-
-
-			float EvaluatedMagnitude = Mod.Magnitude.GetValueAtLevel(Spec.GetLevel());
+			const UGameplayEffectExecutionCalculation* ExecCDO = CurExecDef.CalculationClass->GetDefaultObject<UGameplayEffectExecutionCalculation>();
+			check(ExecCDO);
 
-			Aggregator.DisableCallbacks();
-			Aggregator.AddMod(EvaluatedMagnitude, Mod.ModifierOp, &Mod.SourceTags, &Mod.TargetTags, ModifierHandle);
-
-			ModifiedAggregators.Add(&Aggregator);
-		}
+			TArray<FGameplayModifierEvaluatedData> OutModifiers;
 
-		CDO->Execute(Spec, ExecIdx, Owner);
+			// Run the custom execution
+			FGameplayEffectCustomExecutionParameters ExecutionParams(Spec, CurExecDef.CalculationModifiers, Owner);
+			ExecCDO->Execute(ExecutionParams, OutModifiers);
 
-		// Remove those custom inline mods
-		for (FAggregator* Aggregator : ModifiedAggregators)
-		{
-			Aggregator->RemoveMod(ModifierHandle);
-			Aggregator->EnabledCallbacks();
+			// Execute any mods the custom execution yielded
+			for (FGameplayModifierEvaluatedData& CurExecMod : OutModifiers)
+			{
+				InvokeGameplayCueExecute |= InternalExecuteMod(Spec, CurExecMod);
+			}
 		}
-
 	}
 
-
 	// ------------------------------------------------------
 	//	Invoke GameplayCue events
 	// ------------------------------------------------------
@@ -864,6 +991,34 @@ bool FActiveGameplayEffectsContainer::IsGameplayEffectActive(FActiveGameplayEffe
 	return false;
 }
 
+const FGameplayTagContainer* FActiveGameplayEffectsContainer::GetGameplayEffectSourceTagsFromHandle(FActiveGameplayEffectHandle Handle) const
+{
+	// @todo: Need to consider this with tag changes
+	for (const FActiveGameplayEffect& Effect : GameplayEffects)
+	{
+		if (Effect.Handle == Handle)
+		{
+			return Effect.Spec.CapturedSourceTags.GetAggregatedTags();
+		}
+	}
+
+	return nullptr;
+}
+
+const FGameplayTagContainer* FActiveGameplayEffectsContainer::GetGameplayEffectTargetTagsFromHandle(FActiveGameplayEffectHandle Handle) const
+{
+	// @todo: Need to consider this with tag changes
+	for (const FActiveGameplayEffect& Effect : GameplayEffects)
+	{
+		if (Effect.Handle == Handle)
+		{
+			return Effect.Spec.CapturedTargetTags.GetAggregatedTags();
+		}
+	}
+
+	return nullptr;
+}
+
 void FActiveGameplayEffectsContainer::CaptureAttributeForGameplayEffect(OUT FGameplayEffectAttributeCaptureSpec& OutCaptureSpec)
 {
 	FAggregatorRef& AttributeAggregator = FindOrCreateAttributeAggregator(OutCaptureSpec.BackingDefinition.AttributeToCapture);
@@ -890,6 +1045,62 @@ void FActiveGameplayEffectsContainer::InternalUpdateNumericalAttribute(FGameplay
 	}
 }
 
+bool FActiveGameplayEffectsContainer::InternalExecuteMod(FGameplayEffectSpec& Spec, FGameplayModifierEvaluatedData& ModEvalData)
+{
+	bool bExecuted = false;
+
+	UAttributeSet* AttributeSet = const_cast<UAttributeSet*>(Owner->GetAttributeSubobject(ModEvalData.Attribute.GetAttributeSetClass()));
+	if (AttributeSet)
+	{
+		ABILITY_LOG_SCOPE(TEXT("Executing Attribute Mod %s"), *ModEvalData.ToSimpleString());
+
+		FGameplayEffectModCallbackData ExecuteData(Spec, ModEvalData, *Owner);
+
+		/**
+		 *  This should apply 'gamewide' rules. Such as clamping Health to MaxHealth or granting +3 health for every point of strength, etc
+		 *	PreAttributeModify can return false to 'throw out' this modification.
+		 */
+		if (AttributeSet->PreGameplayEffectExecute(ExecuteData))
+		{
+			// Do we have active GE's that are already modifying this?
+			FAggregatorRef* RefPtr = AttributeAggregatorMap.Find(ModEvalData.Attribute);
+			if (RefPtr)
+			{
+				ABILITY_LOG(Log, TEXT("Property %s has active mods. Adding to Aggregator."), *ModEvalData.Attribute.GetName());
+				RefPtr->Get()->ExecModOnBaseValue(ModEvalData.ModifierOp, ModEvalData.Magnitude);
+			}
+			else
+			{
+				// Modify the property in place, without putting it in the AttributeAggregatorMap map
+				float		CurrentValueOfProperty = Owner->GetNumericAttribute(ModEvalData.Attribute);
+				const float NewPropertyValue = FAggregator::StaticExecModOnBaseValue(CurrentValueOfProperty, ModEvalData.ModifierOp, ModEvalData.Magnitude);
+
+				InternalUpdateNumericalAttribute(ModEvalData.Attribute, NewPropertyValue, &ExecuteData);
+			}
+
+			FGameplayEffectModifiedAttribute* ModifiedAttribute = Spec.GetModifiedAttribute(ModEvalData.Attribute);
+			if (!ModifiedAttribute)
+			{
+				// If we haven't already created a modified attribute holder, create it
+				ModifiedAttribute = Spec.AddModifiedAttribute(ModEvalData.Attribute);
+			}
+			ModifiedAttribute->TotalMagnitude += ModEvalData.Magnitude;
+
+			/** This should apply 'gamewide' rules. Such as clamping Health to MaxHealth or granting +3 health for every point of strength, etc */
+			AttributeSet->PostGameplayEffectExecute(ExecuteData);
+
+			bExecuted = true;
+		}
+	}
+	else
+	{
+		// Our owner doesn't have this attribute, so we can't do anything
+		ABILITY_LOG(Log, TEXT("%s does not have attribute %s. Skipping modifier"), *Owner->GetPathName(), *ModEvalData.Attribute.GetName());
+	}
+
+	return bExecuted;
+}
+
 void FActiveGameplayEffectsContainer::StacksNeedToRecalculate()
 {
 	if (!bNeedToRecalculateStacks)
@@ -913,12 +1124,16 @@ FActiveGameplayEffect& FActiveGameplayEffectsContainer::CreateNewActiveGameplayE
 	float DurationBaseValue = NewEffect.Spec.GetDuration();
 	if (DurationBaseValue > 0.f)
 	{
-		const FAggregator& OutgoingAgg = NewEffect.Spec.CapturedRelevantAttributes.GetAggregator(UAbilitySystemComponent::GetOutgoingDurationCapture());
 		const FAggregator& IncomingAgg = *FindOrCreateAttributeAggregator(UAbilitySystemComponent::GetIncomingDurationProperty()).Get();
 		
 		FAggregator DurationAgg;
 
-		DurationAgg.AddModsFrom(OutgoingAgg);
+		const FGameplayEffectAttributeCaptureSpec* CaptureSpec = NewEffect.Spec.CapturedRelevantAttributes.FindCaptureSpecByDefinition(UAbilitySystemComponent::GetOutgoingDurationCapture(), true);
+		if (CaptureSpec)
+		{
+			CaptureSpec->AttemptAddAggregatorModsToAggregator(DurationAgg);
+		}
+
 		DurationAgg.AddModsFrom(IncomingAgg);
 
 		FAggregatorEvaluateParameters Params;
@@ -1509,6 +1724,18 @@ UAbilitySystemComponent* FActiveGameplayEffectHandle::GetOwningAbilitySystemComp
 
 	return nullptr;	
 }
+
+const UAbilitySystemComponent* FActiveGameplayEffectHandle::GetOwningAbilitySystemComponent() const
+{
+	TWeakObjectPtr<UAbilitySystemComponent>* Ptr = GlobalActiveGameplayEffectHandles::Map.Find(*this);
+	if (Ptr)
+	{
+		return Ptr->Get();
+	}
+
+	return nullptr;
+}
+
 // -----------------------------------------------------------------
 
 bool FActiveGameplayEffectQuery::Matches(const FActiveGameplayEffect& Effect) const
diff --git a/Engine/Source/Runtime/GameplayAbilities/Private/GameplayEffectAggregator.cpp b/Engine/Source/Runtime/GameplayAbilities/Private/GameplayEffectAggregator.cpp
--- a/Engine/Source/Runtime/GameplayAbilities/Private/GameplayEffectAggregator.cpp
+++ b/Engine/Source/Runtime/GameplayAbilities/Private/GameplayEffectAggregator.cpp
@@ -2,13 +2,33 @@
 
 #include "AbilitySystemPrivatePCH.h"
 #include "GameplayEffectAggregator.h"
+#include "AbilitySystemComponent.h"
 
 bool FAggregatorMod::Qualifies(const FAggregatorEvaluateParameters& Parameters) const
 {
-	bool SourceMet = (!SourceTagReqs || SourceTagReqs->IsEmpty()) || (Parameters.SourceTags && SourceTagReqs->RequirementsMet(*Parameters.SourceTags));
-	bool TargetMet = (!TargetTagReqs || TargetTagReqs->IsEmpty()) || (Parameters.TargetTags && TargetTagReqs->RequirementsMet(*Parameters.TargetTags));
+	bool bSourceMet = (!SourceTagReqs || SourceTagReqs->IsEmpty()) || (Parameters.SourceTags && SourceTagReqs->RequirementsMet(*Parameters.SourceTags));
+	bool bTargetMet = (!TargetTagReqs || TargetTagReqs->IsEmpty()) || (Parameters.TargetTags && TargetTagReqs->RequirementsMet(*Parameters.TargetTags));
 
-	return SourceMet && TargetMet;
+	bool bSourceFilterMet = (Parameters.AppliedSourceTagFilter.Num() == 0);
+	bool bTargetFilterMet = (Parameters.AppliedTargetTagFilter.Num() == 0);
+	
+	const UAbilitySystemComponent* HandleComponent = ActiveHandle.GetOwningAbilitySystemComponent();
+	if (HandleComponent)
+	{
+		if (!bSourceFilterMet)
+		{
+			const FGameplayTagContainer* SourceTags = HandleComponent->GetGameplayEffectSourceTagsFromHandle(ActiveHandle);
+			bSourceFilterMet = (SourceTags && SourceTags->MatchesAll(Parameters.AppliedSourceTagFilter, false));
+		}
+
+		if (!bTargetFilterMet)
+		{
+			const FGameplayTagContainer* TargetTags = HandleComponent->GetGameplayEffectTargetTagsFromHandle(ActiveHandle);
+			bTargetFilterMet = (TargetTags && TargetTags->MatchesAll(Parameters.AppliedTargetTagFilter, false));
+		}
+	}
+
+	return bSourceMet && bTargetMet && bSourceFilterMet && bTargetFilterMet;
 }
 
 float FAggregator::Evaluate(const FAggregatorEvaluateParameters& Parameters) const
@@ -39,6 +59,16 @@ float FAggregator::EvaluateWithBase(float InlineBaseValue, const FAggregatorEval
 	return ((InlineBaseValue + Additive) * Multiplicitive) / Division;
 }
 
+float FAggregator::EvaluateBonus(const FAggregatorEvaluateParameters& Parameters) const
+{
+	return (Evaluate(Parameters) - GetBaseValue());
+}
+
+float FAggregator::GetBaseValue() const
+{
+	return BaseValue;
+}
+
 void FAggregator::SetBaseValue(float NewBaseValue)
 {
 	BaseValue = NewBaseValue;
@@ -158,21 +188,7 @@ void FAggregator::TakeSnapshotOf(const FAggregator& AggToSnapshot)
 
 void FAggregator::BroadcastOnDirty()
 {
-	if (!CallbacksDisabled)
-	{
-		OnDirty.Broadcast(this);
-	}
-
-}
-
-void FAggregator::DisableCallbacks()
-{
-	CallbacksDisabled= true;
-}
-
-void FAggregator::EnabledCallbacks()
-{
-	CallbacksDisabled = false;
+	OnDirty.Broadcast(this);
 }
 
 void FAggregatorRef::TakeSnapshotOf(const FAggregatorRef& RefToSnapshot)
diff --git a/Engine/Source/Runtime/GameplayAbilities/Private/GameplayEffectCalculation.cpp b/Engine/Source/Runtime/GameplayAbilities/Private/GameplayEffectCalculation.cpp
--- a/Engine/Source/Runtime/GameplayAbilities/Private/GameplayEffectCalculation.cpp
+++ b/Engine/Source/Runtime/GameplayAbilities/Private/GameplayEffectCalculation.cpp
@@ -9,12 +9,7 @@ UGameplayEffectCalculation::UGameplayEffectCalculation(const FObjectInitializer&
 
 }
 
-void UGameplayEffectCalculation::Execute(FGameplayEffectSpec& Spec, int32 ExecutionIdx, UAbilitySystemComponent* TargetAbilitySystemComponent) const
+const TArray<FGameplayEffectAttributeCaptureDefinition>& UGameplayEffectCalculation::GetAttributeCaptureDefinitions() const
 {
-	
+	return RelevantAttributesToCapture;
 }
-
-void UGameplayEffectCalculation::GetCaptureDefinitions(OUT TArray<FGameplayEffectAttributeCaptureDefinition>& Definitions) const
-{
-
-}
\ No newline at end of file
diff --git a/Engine/Source/Runtime/GameplayAbilities/Private/GameplayEffectExecutionCalculation.cpp b/Engine/Source/Runtime/GameplayAbilities/Private/GameplayEffectExecutionCalculation.cpp
new file mode 100644
--- /dev/null
+++ b/Engine/Source/Runtime/GameplayAbilities/Private/GameplayEffectExecutionCalculation.cpp
@@ -0,0 +1,159 @@
+// Copyright 1998-2014 Epic Games, Inc. All Rights Reserved.
+
+#include "AbilitySystemPrivatePCH.h"
+#include "GameplayEffectExecutionCalculation.h"
+#include "AbilitySystemComponent.h"
+
+FGameplayEffectCustomExecutionParameters::FGameplayEffectCustomExecutionParameters()
+	: OwningSpec(nullptr)
+	, TargetAbilitySystemComponent(nullptr)
+{
+}
+
+FGameplayEffectCustomExecutionParameters::FGameplayEffectCustomExecutionParameters(const FGameplayEffectSpec& InOwningSpec, const TArray<FGameplayEffectExecutionScopedModifierInfo>& InScopedMods, UAbilitySystemComponent* InTargetAbilityComponent)
+	: OwningSpec(&InOwningSpec)
+	, TargetAbilitySystemComponent(InTargetAbilityComponent)
+{
+	check(InOwningSpec.Def);
+
+	FActiveGameplayEffectHandle ModifierHandle = FActiveGameplayEffectHandle::GenerateNewHandle(InTargetAbilityComponent);
+
+	for (const FGameplayEffectExecutionScopedModifierInfo& CurScopedMod : InScopedMods)
+	{
+		FAggregator* ScopedAggregator = ScopedModifierAggregators.Find(CurScopedMod.CapturedAttribute);
+		if (!ScopedAggregator)
+		{
+			const FGameplayEffectAttributeCaptureSpec* CaptureSpec = InOwningSpec.CapturedRelevantAttributes.FindCaptureSpecByDefinition(CurScopedMod.CapturedAttribute, true);
+
+			FAggregator SnapshotAgg;
+			if (CaptureSpec && CaptureSpec->AttemptGetAttributeAggregatorSnapshot(SnapshotAgg))
+			{
+				ScopedAggregator = &(ScopedModifierAggregators.Add(CurScopedMod.CapturedAttribute, SnapshotAgg));
+			}
+		}
+		
+		float ModEvalValue = 0.f;
+		if (ScopedAggregator && CurScopedMod.ModifierMagnitude.AttemptCalculateMagnitude(InOwningSpec, ModEvalValue))
+		{
+			ScopedAggregator->AddMod(ModEvalValue, CurScopedMod.ModifierOp, &CurScopedMod.SourceTags, &CurScopedMod.TargetTags, ModifierHandle);
+		}
+		else
+		{
+			ABILITY_LOG(Warning, TEXT("Attempted to apply a scoped modifier from %s's %s magnitude calculation that could not be properly calculated. Some attributes necessary for the calculation were missing."), *InOwningSpec.Def->GetName(), *CurScopedMod.CapturedAttribute.ToSimpleString());
+		}
+	}
+}
+
+const FGameplayEffectSpec& FGameplayEffectCustomExecutionParameters::GetOwningSpec() const
+{
+	check(OwningSpec);
+	return *OwningSpec;
+}
+
+UAbilitySystemComponent* FGameplayEffectCustomExecutionParameters::GetTargetAbilitySystemComponent() const
+{
+	return TargetAbilitySystemComponent.Get();
+}
+
+bool FGameplayEffectCustomExecutionParameters::AttemptCalculateCapturedAttributeMagnitude(const FGameplayEffectAttributeCaptureDefinition& InCaptureDef, const FAggregatorEvaluateParameters& InEvalParams, OUT float& OutMagnitude) const
+{
+	check(OwningSpec);
+
+	OutMagnitude = 0.f;
+
+	const FAggregator* CalcAgg = ScopedModifierAggregators.Find(InCaptureDef);
+	if (CalcAgg)
+	{
+		OutMagnitude = CalcAgg->Evaluate(InEvalParams);
+		return true;
+	}
+	else
+	{
+		const FGameplayEffectAttributeCaptureSpec* CaptureSpec = OwningSpec->CapturedRelevantAttributes.FindCaptureSpecByDefinition(InCaptureDef, true);
+		if (CaptureSpec)
+		{
+			return CaptureSpec->AttemptCalculateAttributeMagnitude(InEvalParams, OutMagnitude);
+		}
+	}
+
+	return false;
+}
+
+bool FGameplayEffectCustomExecutionParameters::AttemptCalculateCapturedAttributeBaseValue(const FGameplayEffectAttributeCaptureDefinition& InCaptureDef, OUT float& OutBaseValue) const
+{
+	check(OwningSpec);
+
+	OutBaseValue = 0.f;
+
+	const FAggregator* CalcAgg = ScopedModifierAggregators.Find(InCaptureDef);
+	if (CalcAgg)
+	{
+		OutBaseValue = CalcAgg->GetBaseValue();
+		return true;
+	}
+	else
+	{
+		const FGameplayEffectAttributeCaptureSpec* CaptureSpec = OwningSpec->CapturedRelevantAttributes.FindCaptureSpecByDefinition(InCaptureDef, true);
+		if (CaptureSpec)
+		{
+			return CaptureSpec->AttemptCalculateAttributeBaseValue(OutBaseValue);
+		}
+	}
+
+	return false;
+}
+
+bool FGameplayEffectCustomExecutionParameters::AttemptCalculateCapturedAttributeBonusMagnitude(const FGameplayEffectAttributeCaptureDefinition& InCaptureDef, const FAggregatorEvaluateParameters& InEvalParams, OUT float& OutBonusMagnitude) const
+{
+	check(OwningSpec);
+
+	OutBonusMagnitude = 0.f;
+
+	const FAggregator* CalcAgg = ScopedModifierAggregators.Find(InCaptureDef);
+	if (CalcAgg)
+	{
+		OutBonusMagnitude = CalcAgg->EvaluateBonus(InEvalParams);
+		return true;
+	}
+	else
+	{
+		const FGameplayEffectAttributeCaptureSpec* CaptureSpec = OwningSpec->CapturedRelevantAttributes.FindCaptureSpecByDefinition(InCaptureDef, true);
+		if (CaptureSpec)
+		{
+			return CaptureSpec->AttemptCalculateAttributeBonusMagnitude(InEvalParams, OutBonusMagnitude);
+		}
+	}
+
+	return false;
+}
+
+bool FGameplayEffectCustomExecutionParameters::AttemptGetCapturedAttributeAggregatorSnapshot(const FGameplayEffectAttributeCaptureDefinition& InCaptureDef, OUT FAggregator& OutSnapshottedAggregator) const
+{
+	check(OwningSpec);
+
+	const FAggregator* CalcAgg = ScopedModifierAggregators.Find(InCaptureDef);
+	if (CalcAgg)
+	{
+		OutSnapshottedAggregator.TakeSnapshotOf(*CalcAgg);
+		return true;
+	}
+	else
+	{
+		const FGameplayEffectAttributeCaptureSpec* CaptureSpec = OwningSpec->CapturedRelevantAttributes.FindCaptureSpecByDefinition(InCaptureDef, true);
+		if (CaptureSpec)
+		{
+			return CaptureSpec->AttemptGetAttributeAggregatorSnapshot(OutSnapshottedAggregator);
+		}
+	}
+
+	return false;
+}
+
+UGameplayEffectExecutionCalculation::UGameplayEffectExecutionCalculation(const FObjectInitializer& ObjectInitializer)
+	: Super(ObjectInitializer)
+{
+}
+
+void UGameplayEffectExecutionCalculation::Execute_Implementation(const FGameplayEffectCustomExecutionParameters& ExecutionParams, OUT TArray<FGameplayModifierEvaluatedData>& OutAdditionalModifiers) const
+{
+}
diff --git a/Engine/Source/Runtime/GameplayAbilities/Private/GameplayEffectTypes.cpp b/Engine/Source/Runtime/GameplayAbilities/Private/GameplayEffectTypes.cpp
--- a/Engine/Source/Runtime/GameplayAbilities/Private/GameplayEffectTypes.cpp
+++ b/Engine/Source/Runtime/GameplayAbilities/Private/GameplayEffectTypes.cpp
@@ -5,6 +5,21 @@
 #include "GameplayTagsModule.h"
 #include "GameplayEffectTypes.h"
 
+bool FGameplayEffectAttributeCaptureDefinition::operator==(const FGameplayEffectAttributeCaptureDefinition& Other) const
+{
+	return ((AttributeToCapture == Other.AttributeToCapture) && (AttributeSource == Other.AttributeSource) && (bSnapshot == Other.bSnapshot));
+}
+
+bool FGameplayEffectAttributeCaptureDefinition::operator!=(const FGameplayEffectAttributeCaptureDefinition& Other) const
+{
+	return ((AttributeToCapture != Other.AttributeToCapture) || (AttributeSource != Other.AttributeSource) || (bSnapshot != Other.bSnapshot));
+}
+
+FString FGameplayEffectAttributeCaptureDefinition::ToSimpleString() const
+{
+	return FString::Printf(TEXT("Attribute: %s, Capture: %s, Snapshot: %d"), *AttributeToCapture.GetName(), AttributeSource == EGameplayEffectAttributeCaptureSource::Source ? TEXT("Source") : TEXT("Target"), bSnapshot);
+}
+
 // --------------------------------------------------------------------------------------------------------------------------------------------------------
 //
 //	FGameplayEffectContext
diff --git a/Engine/Source/Runtime/GameplayAbilities/Private/GameplayModCalculation.cpp b/Engine/Source/Runtime/GameplayAbilities/Private/GameplayModCalculation.cpp
deleted file mode 100644
diff --git a/Engine/Source/Runtime/GameplayAbilities/Private/GameplayModMagnitudeCalculation.cpp b/Engine/Source/Runtime/GameplayAbilities/Private/GameplayModMagnitudeCalculation.cpp
new file mode 100644
--- /dev/null
+++ b/Engine/Source/Runtime/GameplayAbilities/Private/GameplayModMagnitudeCalculation.cpp
@@ -0,0 +1,15 @@
+// Copyright 1998-2013 Epic Games, Inc. All Rights Reserved.
+
+#include "AbilitySystemPrivatePCH.h"
+#include "GameplayModMagnitudeCalculation.h"
+
+UGameplayModMagnitudeCalculation::UGameplayModMagnitudeCalculation(const FObjectInitializer& ObjectInitializer)
+	: Super(ObjectInitializer)
+{
+
+}
+
+float UGameplayModMagnitudeCalculation::CalculateMagnitude_Implementation(const FGameplayEffectSpec& Spec) const
+{
+	return 0.f;
+}
diff --git a/Engine/Source/Runtime/GameplayAbilities/Public/AbilitySystemComponent.h b/Engine/Source/Runtime/GameplayAbilities/Public/AbilitySystemComponent.h
--- a/Engine/Source/Runtime/GameplayAbilities/Public/AbilitySystemComponent.h
+++ b/Engine/Source/Runtime/GameplayAbilities/Public/AbilitySystemComponent.h
@@ -67,13 +67,13 @@ class GAMEPLAYABILITIES_API UAbilitySystemComponent : public UActorComponent, pu
 
 	/** Finds existing AttributeSet */
 	template <class T >
-	const T*	GetSet()
+	const T*	GetSet() const
 	{
 		return (T*)GetAttributeSubobject(T::StaticClass());
 	}
 
 	template <class T >
-	const T*	GetSetChecked()
+	const T*	GetSetChecked() const
 	{
 		return (T*)GetAttributeSubobjectChecked(T::StaticClass());
 	}
@@ -207,6 +207,24 @@ class GAMEPLAYABILITIES_API UAbilitySystemComponent : public UActorComponent, pu
 	UFUNCTION(BlueprintCallable, Category = GameplayEffects)
 	bool IsGameplayEffectActive(FActiveGameplayEffectHandle InHandle) const;
 
+	/**
+	 * Get the source tags from the gameplay spec represented by the specified handle, if possible
+	 * 
+	 * @param Handle	Handle of the gameplay effect to retrieve source tags from
+	 * 
+	 * @return Source tags from the gameplay spec represented by the handle, if possible
+	 */
+	const FGameplayTagContainer* GetGameplayEffectSourceTagsFromHandle(FActiveGameplayEffectHandle Handle) const;
+
+	/**
+	 * Get the target tags from the gameplay spec represented by the specified handle, if possible
+	 * 
+	 * @param Handle	Handle of the gameplay effect to retrieve target tags from
+	 * 
+	 * @return Target tags from the gameplay spec represented by the handle, if possible
+	 */
+	const FGameplayTagContainer* GetGameplayEffectTargetTagsFromHandle(FActiveGameplayEffectHandle Handle) const;
+
 	/**
 	 * Populate the specified capture spec with the data necessary to capture an attribute from the component
 	 * 
diff --git a/Engine/Source/Runtime/GameplayAbilities/Public/AttributeSet.h b/Engine/Source/Runtime/GameplayAbilities/Public/AttributeSet.h
--- a/Engine/Source/Runtime/GameplayAbilities/Public/AttributeSet.h
+++ b/Engine/Source/Runtime/GameplayAbilities/Public/AttributeSet.h
@@ -158,7 +158,12 @@ struct GAMEPLAYABILITIES_API FScalableFloat
 		: Value(0.f)
 		, FinalCurve(NULL)
 	{
+	}
 
+	FScalableFloat(float InInitialValue)
+		: Value(InInitialValue)
+		, FinalCurve(NULL)
+	{
 	}
 
 public:
diff --git a/Engine/Source/Runtime/GameplayAbilities/Public/GameplayEffect.h b/Engine/Source/Runtime/GameplayAbilities/Public/GameplayEffect.h
--- a/Engine/Source/Runtime/GameplayAbilities/Public/GameplayEffect.h
+++ b/Engine/Source/Runtime/GameplayAbilities/Public/GameplayEffect.h
@@ -16,15 +16,8 @@ struct FActiveGameplayEffect;
 class UGameplayEffect;
 class UGameplayEffectTemplate;
 class UAbilitySystemComponent;
-class UGameplayModCalculation;
-
-/** Enumeration for options of where to capture gameplay attributes from for gameplay effects */
-UENUM()
-enum class EGameplayEffectAttributeCaptureSource : uint8
-{
-	Source,	// Source (caster) of the gameplay effect
-	Target	// Target (recipient) of the gameplay effect
-};
+class UGameplayModMagnitudeCalculation;
+class UGameplayEffectExecutionCalculation;
 
 // Repurposing(?) this to be the inplace/custom/scoped modifier used by FGameplayEffectExecutionDefinition
 //	I would expect this to be the thing that is built off to be a "scalable float/attribute reference/custom magnitude" thing.
@@ -92,105 +85,217 @@ struct FGameplayEffectStackingCallbacks
 	TArray<TSubclassOf<class UGameplayEffectStackingExtension> >	ExtensionClasses;
 };
 
+/** Enumeration outlining the possible gameplay effect magnitude calculation policies */
+UENUM()
+enum class EGameplayEffectMagnitudeCalculation : uint8
+{
+	ScalableFloat,			// Use a simple, scalable float for the calculation
+	AttributeBased,			// Perform a calculation based upon an attribute
+	CustomCalculationClass	// Perform a custom calculation, capable of capturing and acting on multiple attributes, in either BP or native
+};
 
+/** Enumeration outlining the possible attribute based float calculation policies */
+UENUM()
+enum class EAttributeBasedFloatCalculationType : uint8
+{
+	AttributeMagnitude,			// Use the final evaluated magnitude of the attribute
+	AttributeBaseValue,			// Use the base value of the attribute
+	AttributeBonusMagnitude		// Use the "bonus" evaluated magnitude of the attribute: Equivalent to (FinalMag - BaseValue)
+};
 
-
-
-/** Struct defining gameplay attribute capture options for gameplay effects */
+/** 
+ * Struct representing a float whose magnitude is dictated by a backing attribute and a calculation policy, follows basic form of:
+ * (Coefficient * (PreMultiplyAdditiveValue + [Eval'd Attribute Value According to Policy])) + PostMultiplyAdditiveValue
+ */
 USTRUCT()
-struct FGameplayEffectAttributeCaptureDefinition
+struct FAttributeBasedFloat
 {
 	GENERATED_USTRUCT_BODY()
 
-	FGameplayEffectAttributeCaptureDefinition()
-	{
+public:
 
-	}
+	/** Constructor */
+	FAttributeBasedFloat()
+		: Coefficient(1.f)
+		, PreMultiplyAdditiveValue(0.f)
+		, PostMultiplyAdditiveValue(0.f)
+		, BackingAttribute()
+		, AttributeCalculationType(EAttributeBasedFloatCalculationType::AttributeMagnitude)
+	{}
 
-	FGameplayEffectAttributeCaptureDefinition(FGameplayAttribute InAttribute, EGameplayEffectAttributeCaptureSource InSource, bool InSnapshot)
-		: AttributeToCapture(InAttribute), AttributeSource(InSource), bSnapshot(InSnapshot)
-	{
+	/**
+	 * Calculate and return the magnitude of the float given the specified gameplay effect spec.
+	 * 
+	 * @note:	This function assumes (and asserts on) the existence of the required captured attribute within the spec.
+	 *			It is the responsibility of the caller to verify that the spec is properly setup before calling this function.
+	 *			
+	 *	@param InRelevantSpec	Gameplay effect spec providing the backing attribute capture
+	 *	
+	 *	@return Evaluated magnitude based upon the spec & calculation policy
+	 */
+	float CalculateMagnitude(const FGameplayEffectSpec& InRelevantSpec) const;
 
-	}
+	/** Coefficient to the attribute calculation */
+	UPROPERTY(EditDefaultsOnly, Category=AttributeFloat)
+	FScalableFloat Coefficient;
 
-	/** Gameplay attribute to capture */
-	UPROPERTY(EditDefaultsOnly, Category=Capture)
-	FGameplayAttribute AttributeToCapture;
+	/** Additive value to the attribute calculation, added in before the coefficient applies */
+	UPROPERTY(EditDefaultsOnly, Category=AttributeFloat)
+	FScalableFloat PreMultiplyAdditiveValue;
 
-	/** Source of the gameplay attribute */
-	UPROPERTY(EditDefaultsOnly, Category=Capture)
-	EGameplayEffectAttributeCaptureSource AttributeSource;
+	/** Additive value to the attribute calculation, added in after the coefficient applies */
+	UPROPERTY(EditDefaultsOnly, Category=AttributeFloat)
+	FScalableFloat PostMultiplyAdditiveValue;
+
+	/** Attribute backing the calculation */
+	UPROPERTY(EditDefaultsOnly, Category=AttributeFloat)
+	FGameplayEffectAttributeCaptureDefinition BackingAttribute;
 
-	/** Whether the attribute should be snapshotted or not */
-	UPROPERTY(EditDefaultsOnly, Category=Capture)
-	bool bSnapshot;
+	/** Calculation policy in regards to the attribute */
+	UPROPERTY(EditDefaultsOnly, Category=AttributeFloat)
+	EAttributeBasedFloatCalculationType AttributeCalculationType;
 
-	/** Equality/Inequality operators */
-	bool operator==(const FGameplayEffectAttributeCaptureDefinition& Other) const;
-	bool operator!=(const FGameplayEffectAttributeCaptureDefinition& Other) const;
+	/** Filter to use on source tags; If specified, only modifiers applied with all of these tags will factor into the calculation */
+	UPROPERTY(EditDefaultsOnly, Category=AttributeFloat)
+	FGameplayTagContainer SourceTagFilter;
 
-	FString ToSimpleString() const;
+	/** Filter to use on target tags; If specified, only modifiers applied with all of these tags will factor into the calculation */
+	UPROPERTY(EditDefaultsOnly, Category=AttributeFloat)
+	FGameplayTagContainer TargetTagFilter;
 };
 
-/** Structure for defining custom logic when a gameplay effect execution happens */
+/** Struct representing the magnitude of a gameplay effect modifier, potentially calculated in numerous different ways */
 USTRUCT()
-struct FGameplayEffectExecutionDefinition
+struct FGameplayEffectModifierMagnitude
 {
 	GENERATED_USTRUCT_BODY()
 
-	void OnObjectPostEditChange(UObject* Object, FPropertyChangedEvent& PropertyChangedEvent);
+public:
 
-	// @todo: Ideally this would be an interface, but can't store one in a UPROPERTY. Need to figure out where we want custom
-	// execution to come from, as it could potentially be multiple places in reality. Might want native code, might want a blueprint graph,
-	// might want to implement it inline in the gameplay effect graph (maybe? would need to be static or instancing becomes a mess), might
-	// even just want it to be a "custom calculation" inline with fancy math operations.
-	UPROPERTY(EditDefaultsOnly, Category=Execution)
-	TSubclassOf<UGameplayEffectCalculation> CalculationClass;
-	
-	// @todo: This area should also be able to capture attributes because it's going to likely want them for the execution, as well as for
-	// scoped modifiers. Need to decide where they live/how this looks in the UI/etc.
-	// 
-	
-	
+	/** Constructor */
+	FGameplayEffectModifierMagnitude()
+		: MagnitudeCalculationType(EGameplayEffectMagnitudeCalculation::ScalableFloat)
+		, ScalableFloatMagnitude()
+		, AttributeBasedMagnitude()
+		, CalculationClassMagnitude(nullptr)
+	{
+	}
 
-	// FGameplayEffectAttributeCaptureDefinition
+	/**
+	 * Determines if the magnitude can be properly calculated with the specified gameplay effect spec (could fail if relying on an attribute not present, etc.)
+	 * 
+	 * @param InRelevantSpec	Gameplay effect spec to check for magnitude calculation
+	 * 
+	 * @return Whether or not the magnitude can be properly calculated
+	 */
+	bool CanCalculateMagnitude(const FGameplayEffectSpec& InRelevantSpec) const;
 
-	/** Modifiers that are applied "in place" during the execution calculation */
-	UPROPERTY(EditDefaultsOnly, Category = Execution)
-	TArray<FExtensionAttributeModifierInfo>	CalculationModifiers;
-};
+	/**
+	 * Attempts to calculate the magnitude given the provided spec. May fail if necessary information (such as captured attributes) is missing from
+	 * the spec.
+	 * 
+	 * @param InRelevantSpec			Gameplay effect spec to use to calculate the magnitude with
+	 * @param OutCalculatedMagnitude	[OUT] Calculated value of the magnitude, will be set to 0.f in the event of failure
+	 * 
+	 * @return True if the calculation was successful, false if it was not
+	 */
+	bool AttemptCalculateMagnitude(const FGameplayEffectSpec& InRelevantSpec, OUT float& OutCalculatedMagnitude) const;
 
-/** The runtime/mutable data of an ExecutionDefinition */
-USTRUCT()
-struct FGameplayEffectExecutionSpec
-{
-	GENERATED_USTRUCT_BODY()
+	/**
+	 * Gather all of the attribute capture definitions necessary to compute the magnitude and place them into the provided array
+	 * 
+	 * @param OutCaptureDefs	[OUT] Array populated with necessary attribute capture definitions
+	 */
+	void GetAttributeCaptureDefinitions(OUT TArray<FGameplayEffectAttributeCaptureDefinition>& OutCaptureDefs) const;
+
+protected:
+
+	/** Type of calculation to perform to derive the magnitude */
+	UPROPERTY(EditDefaultsOnly, Category=Magnitude)
+	EGameplayEffectMagnitudeCalculation MagnitudeCalculationType;
 
-	//FAggregatorRef	
+	/** Magnitude value represented by a scalable float */
+	UPROPERTY(EditDefaultsOnly, Category=Magnitude)
+	FScalableFloat ScalableFloatMagnitude;
+
+	/** Magnitude value represented by an attribute-based float */
+	UPROPERTY(EditDefaultsOnly, Category=Magnitude)
+	FAttributeBasedFloat AttributeBasedMagnitude;
+
+	/** Magnitude value represented by a custom calculation class */
+	UPROPERTY(EditDefaultsOnly, Category=Magnitude)
+	TSubclassOf<UGameplayModMagnitudeCalculation> CalculationClassMagnitude;
+
+	// @hack: @todo: This is temporary to aid in post-load fix-up w/o exposing members publicly
+	friend class UGameplayEffect;
 };
 
-// @todo: Enum to choose between types; Details customization to hide the other one
-// @todo: Add Math Expression parser
+/** 
+ * Struct representing modifier info used exclusively for "scoped" executions that happen instantaneously. These are
+ * folded into a calculation only for the extent of the calculation and never permanently added to an aggregator.
+ */
 USTRUCT()
-struct FGameplayEffectMagnitude
+struct FGameplayEffectExecutionScopedModifierInfo
 {
 	GENERATED_USTRUCT_BODY()
 
-	FGameplayEffectMagnitude()
-	: ScalableFloatMagnitude()
-	, CalculationClassMagnitude(nullptr)
+	// Constructors
+	FGameplayEffectExecutionScopedModifierInfo()
+		: ModifierOp(EGameplayModOp::Additive)
+	{}
+
+	FGameplayEffectExecutionScopedModifierInfo(const FGameplayEffectAttributeCaptureDefinition& InCaptureDef)
+		: CapturedAttribute(InCaptureDef)
+		, ModifierOp(EGameplayModOp::Additive)
 	{
 	}
 
-	UPROPERTY(EditDefaultsOnly, Category=Magnitude)
-	FScalableFloat ScalableFloatMagnitude;
+	/** Backing attribute that the scoped modifier is for */
+	UPROPERTY(VisibleDefaultsOnly, Category=Execution)
+	FGameplayEffectAttributeCaptureDefinition CapturedAttribute;
 
-	UPROPERTY(EditDefaultsOnly, Category=Magnitude)
-	TSubclassOf<UGameplayModCalculation> CalculationClassMagnitude;
+	/** Modifier operation to perform */
+	UPROPERTY(EditDefaultsOnly, Category=Execution)
+	TEnumAsByte<EGameplayModOp::Type> ModifierOp;
 
-	/** Attributes to capture for usage in custom calculations for magnitude */
-	UPROPERTY(EditDefaultsOnly, Category=Magnitude)
-	TArray<FGameplayEffectAttributeCaptureDefinition> CustomMagnitudeCalculationAttributes;
+	/** Magnitude of the scoped modifier */
+	UPROPERTY(EditDefaultsOnly, Category=Execution)
+	FGameplayEffectModifierMagnitude ModifierMagnitude;
+
+	/** Source tag requirements for the modifier to apply */
+	UPROPERTY(EditDefaultsOnly, Category=Execution)
+	FGameplayTagRequirements SourceTags;
+
+	/** Target tag requirements for the modifier to apply */
+	UPROPERTY(EditDefaultsOnly, Category=Execution)
+	FGameplayTagRequirements TargetTags;
+};
+
+/** 
+ * Struct representing the definition of a custom execution for a gameplay effect.
+ * Custom executions run special logic from an outside class each time the gameplay effect executes.
+ */
+USTRUCT()
+struct FGameplayEffectExecutionDefinition
+{
+	GENERATED_USTRUCT_BODY()
+
+	/**
+	 * Gathers and populates the specified array with the capture definitions that the execution would like in order
+	 * to perform its custom calculation. Up to the individual execution calculation to handle if some of them are missing
+	 * or not.
+	 * 
+	 * @param OutCaptureDefs	[OUT] Capture definitions requested by the execution
+	 */
+	void GetAttributeCaptureDefinitions(OUT TArray<FGameplayEffectAttributeCaptureDefinition>& OutCaptureDefs) const;
+
+	/** Custom execution calculation class to run when the gameplay effect executes */
+	UPROPERTY(EditDefaultsOnly, Category=Execution)
+	TSubclassOf<UGameplayEffectExecutionCalculation> CalculationClass;
+	
+	/** Modifiers that are applied "in place" during the execution calculation */
+	UPROPERTY(EditDefaultsOnly, Category = Execution)
+	TArray<FGameplayEffectExecutionScopedModifierInfo> CalculationModifiers;
 };
 
 /**
@@ -221,13 +326,14 @@ struct GAMEPLAYABILITIES_API FGameplayModifierInfo
 	UPROPERTY(EditDefaultsOnly, Category=GameplayModifier, meta=(DisplayName="Custom Ops"))
 	TArray<FGameplayModifierCallback> Callbacks;
 
-	/** How much this modifies what it is applied to */
-	UPROPERTY(EditDefaultsOnly, Category=GameplayModifier)
-	FScalableFloat Magnitude; // Not modified from defaults
+	// @todo: Remove this after content resave
+	/** Now "deprecated," though being handled in a custom manner to avoid engine version bump. */
+	UPROPERTY()
+	FScalableFloat Magnitude;
 
-	// @todo: Rename/fix-up/etc.
+	/** Magnitude of the modifier */
 	UPROPERTY(EditDefaultsOnly, Category=GameplayModifier)
-	FGameplayEffectMagnitude MagnitudeV2;
+	FGameplayEffectModifierMagnitude ModifierMagnitude;
 
 	UPROPERTY(EditDefaultsOnly, Category=GameplayModifier)
 	FGameplayTagRequirements	SourceTags;
@@ -398,11 +504,9 @@ public:
 	//		Tag pass: properties that make sense "right now" (remove this comment DaveR/BillyB)
 	// ----------------------------------------------------------------------
 	
-
 	/** The GameplayEffect's Tags: tags the the GE *has* and DOES NOT give to the actor. */
 	UPROPERTY(EditDefaultsOnly, BlueprintReadOnly, Category = Tags, meta=(DisplayName="GameplayEffectAssetTag"))
 	FGameplayTagContainer GameplayEffectTags;
-
 	
 	/** "These tags are applied to the actor I am applied to" */
 	UPROPERTY(EditDefaultsOnly, BlueprintReadOnly, Category = Tags)
@@ -418,30 +522,6 @@ public:
 
 };
 
-// GE_REMOVE this may go away? Or still useful for storing 'Final results of shit we calculated'?
-struct FGameplayModifierEvaluatedData
-{
-	FGameplayModifierEvaluatedData()
-		: Magnitude(0.f)
-		, IsValid(false)
-	{
-	}
-
-	FGameplayModifierEvaluatedData(float InMagnitude, FActiveGameplayEffectHandle InHandle = FActiveGameplayEffectHandle())
-		: Magnitude(InMagnitude)
-		, Handle(InHandle)
-		, IsValid(true)
-	{
-	}
-
-	FGameplayTagContainer	SourceTags;
-	FGameplayTagContainer	TargetTags;
-
-	float	Magnitude;
-	FActiveGameplayEffectHandle	Handle;	// Handle of the active gameplay effect that originated us. Will be invalid in many cases
-	bool IsValid;
-};
-
 /**
  * Modifier Specification
  *	-Const data (FGameplayModifierInfo) tells us what we modify, what we can modify
@@ -503,17 +583,67 @@ struct FGameplayEffectAttributeCaptureSpec
 	FGameplayEffectAttributeCaptureSpec();
 	FGameplayEffectAttributeCaptureSpec(const FGameplayEffectAttributeCaptureDefinition& InDefinition);
 
-	// @todo:
-	// API should provide way to get:
-		// Final modified value
-		// Base value
-		// Bonus value (Final - Base)
-		// Mods (optionally matching tag filter)
+	/**
+	 * Returns whether the spec actually has a valid capture yet or not
+	 * 
+	 * @return True if the spec has a valid attribute capture, false if it does not
+	 */
+	bool HasValidCapture() const;
+
+	/**
+	 * Attempts to calculate the magnitude of the captured attribute given the specified parameters. Can fail if the spec doesn't have
+	 * a valid capture yet.
+	 * 
+	 * @param InEvalParams	Parameters to evaluate the attribute under
+	 * @param OutMagnitude	[OUT] Computed magnitude, falls back to 0.f in the event of failure
+	 * 
+	 * @return True if the magnitude was successfully calculated, false if it was not
+	 */
+	bool AttemptCalculateAttributeMagnitude(const FAggregatorEvaluateParameters& InEvalParams, OUT float& OutMagnitude) const;
+
+	/**
+	 * Attempts to calculate the base value of the captured attribute given the specified parameters. Can fail if the spec doesn't have
+	 * a valid capture yet.
+	 * 
+	 * @param OutBaseValue	[OUT] Computed base value, falls back to 0.f in the event of failure
+	 * 
+	 * @return True if the base value was successfully calculated, false if it was not
+	 */
+	bool AttemptCalculateAttributeBaseValue(OUT float& OutBaseValue) const;
+
+	/**
+	 * Attempts to calculate the "bonus" magnitude (final - base value) of the captured attribute given the specified parameters. Can fail if the spec doesn't have
+	 * a valid capture yet.
+	 * 
+	 * @param InEvalParams		Parameters to evaluate the attribute under
+	 * @param OutBonusMagnitude	[OUT] Computed bonus magnitude, falls back to 0.f in the event of failure
+	 * 
+	 * @return True if the bonus magnitude was successfully calculated, false if it was not
+	 */
+	bool AttemptCalculateAttributeBonusMagnitude(const FAggregatorEvaluateParameters& InEvalParams, OUT float& OutBonusMagnitude) const;
+	
+	/**
+	 * Attempts to populate the specified aggregator with a snapshot of the backing captured aggregator. Can fail if the spec doesn't have
+	 * a valid capture yet.
+	 *
+	 * @param OutAggregatorSnapshot	[OUT] Snapshotted aggregator, if possible
+	 *
+ 	 * @return True if the aggregator was successfully snapshotted, false if it was not
+	 */
+	bool AttemptGetAttributeAggregatorSnapshot(OUT FAggregator& OutAggregatorSnapshot) const;
+
+	/**
+	 * Attempts to populate the specified aggregator with all of the mods of the backing captured aggregator. Can fail if the spec doesn't have
+	 * a valid capture yet.
+	 *
+	 * @param OutAggregatorToAddTo	[OUT] Aggregator with mods appended, if possible
+	 *
+ 	 * @return True if the aggregator had mods successfully added to it, false if it did not
+	 */
+	bool AttemptAddAggregatorModsToAggregator(OUT FAggregator& OutAggregatorToAddTo) const;
 	
 	/** Simple accessor to backing capture definition */
 	const FGameplayEffectAttributeCaptureDefinition& GetBackingDefinition() const;
-
-	const FAggregator& GetAggregator() const;
 		
 private:
 
@@ -549,10 +679,24 @@ public:
 	 */
 	void CaptureAttributes(class UAbilitySystemComponent* InAbilitySystemComponent, EGameplayEffectAttributeCaptureSource InCaptureSource);
 
-	// @todo:
-	// API should provide access to querying specs by definition; Wrapper around @todo API for specs above
-	
-	const FAggregator& GetAggregator(const FGameplayEffectAttributeCaptureDefinition& InDefinition) const;
+	/**
+	 * Find a capture spec within the container matching the specified capture definition, if possible.
+	 * 
+	 * @param InDefinition				Capture definition to use as the search basis
+	 * @param bOnlyIncludeValidCapture	If true, even if a spec is found, it won't be returned if it doesn't also have a valid capture already
+	 * 
+	 * @return The found attribute spec matching the specified search params, if any
+	 */
+	const FGameplayEffectAttributeCaptureSpec* FindCaptureSpecByDefinition(const FGameplayEffectAttributeCaptureDefinition& InDefinition, bool bOnlyIncludeValidCapture) const;
+
+	/**
+	 * Determines if the container has specs with valid captures for all of the specified definitions.
+	 * 
+	 * @param InCaptureDefsToCheck	Capture definitions to check for
+	 * 
+	 * @return True if the container has valid capture attributes for all of the specified definitions, false if it does not
+	 */
+	bool HasValidCapturedAttributes(const TArray<FGameplayEffectAttributeCaptureDefinition>& InCaptureDefsToCheck) const;
 
 	/** Returns whether the container has at least one spec w/o snapshotted attributes */
 	bool HasNonSnapshottedAttributes() const;
@@ -610,6 +754,15 @@ struct FGameplayEffectSpec
 	UPROPERTY(NotReplicated)
 	FGameplayEffectAttributeCaptureSpecContainer CapturedRelevantAttributes;
 
+	/**
+	 * Determines if the spec has capture specs with valid captures for all of the specified definitions.
+	 * 
+	 * @param InCaptureDefsToCheck	Capture definitions to check for
+	 * 
+	 * @return True if the container has valid capture attributes for all of the specified definitions, false if it does not
+	 */
+	bool HasValidCapturedAttributes(const TArray<FGameplayEffectAttributeCaptureDefinition>& InCaptureDefsToCheck) const;
+
 	/** Looks for an existing modified attribute struct, may return NULL */
 	const FGameplayEffectModifiedAttribute* GetModifiedAttribute(const FGameplayAttribute& Attribute) const;
 	FGameplayEffectModifiedAttribute* GetModifiedAttribute(const FGameplayAttribute& Attribute);
@@ -678,8 +831,6 @@ struct FGameplayEffectSpec
 	
 	TArray<FModifierSpec> Modifiers;
 	
-	TArray<FGameplayEffectExecutionSpec>	Executions;
-
 	// GE_REMOVE: this feels strange but need to start standardizing how people get this info.
 	float GetModifierMagnitude(const FModifierSpec& ModSpec) const;
 	void CalculateModifierMagnitudes();
@@ -698,6 +849,9 @@ struct FGameplayEffectSpec
 private:
 
 	void CaptureDataFromSource();
+
+	/** Helper function to initialize all of the capture definitions required by the spec */
+	void SetupAttributeCaptureDefinitions();
 	
 	UPROPERTY()
 	FGameplayEffectContextHandle EffectContext; // This tells us how we got here (who / what applied us)
@@ -867,6 +1021,24 @@ struct FActiveGameplayEffectsContainer : public FFastArraySerializer
 	// returns false if the handle points to an effect that is not in this container or is not applied because of the current stacking rules
 	bool IsGameplayEffectActive(FActiveGameplayEffectHandle Handle) const;
 
+	/**
+	 * Get the source tags from the gameplay spec represented by the specified handle, if possible
+	 * 
+	 * @param Handle	Handle of the gameplay effect to retrieve source tags from
+	 * 
+	 * @return Source tags from the gameplay spec represented by the handle, if possible
+	 */
+	const FGameplayTagContainer* GetGameplayEffectSourceTagsFromHandle(FActiveGameplayEffectHandle Handle) const;
+
+	/**
+	 * Get the target tags from the gameplay spec represented by the specified handle, if possible
+	 * 
+	 * @param Handle	Handle of the gameplay effect to retrieve target tags from
+	 * 
+	 * @return Target tags from the gameplay spec represented by the handle, if possible
+	 */
+	const FGameplayTagContainer* GetGameplayEffectTargetTagsFromHandle(FActiveGameplayEffectHandle Handle) const;
+
 	/**
 	 * Populate the specified capture spec with the data necessary to capture an attribute from the container
 	 * 
@@ -928,6 +1100,16 @@ private:
 	FTimerHandle StackHandle;
 
 	void InternalUpdateNumericalAttribute(FGameplayAttribute Attribute, float NewValue, const FGameplayEffectModCallbackData* ModData);
+	
+	/**
+	 * Helper function to execute a mod on owned attributes
+	 * 
+	 * @param Spec			Gameplay effect spec executing the mod
+	 * @param ModEvalData	Evaluated data for the mod
+	 * 
+	 * @return True if the mod successfully executed, false if it did not
+	 */
+	bool InternalExecuteMod(FGameplayEffectSpec& Spec, FGameplayModifierEvaluatedData& ModEvalData);
 
 	bool IsNetAuthority() const;
 
diff --git a/Engine/Source/Runtime/GameplayAbilities/Public/GameplayEffectAggregator.h b/Engine/Source/Runtime/GameplayAbilities/Public/GameplayEffectAggregator.h
--- a/Engine/Source/Runtime/GameplayAbilities/Public/GameplayEffectAggregator.h
+++ b/Engine/Source/Runtime/GameplayAbilities/Public/GameplayEffectAggregator.h
@@ -12,6 +12,12 @@ struct GAMEPLAYABILITIES_API FAggregatorEvaluateParameters
 
 	const FGameplayTagContainer*	SourceTags;
 	const FGameplayTagContainer*	TargetTags;
+
+	/** If any tags are specified in the filter, a mod's owning active gameplay effect's source tags must match ALL of them in order for the mod to count during evaluation */
+	FGameplayTagContainer	AppliedSourceTagFilter;
+
+	/** If any tags are specified in the filter, a mod's owning active gameplay effect's target tags must match ALL of them in order for the mod to count during evaluation */
+	FGameplayTagContainer	AppliedTargetTagFilter;
 };
 
 struct GAMEPLAYABILITIES_API FAggregatorMod
@@ -31,8 +37,10 @@ struct GAMEPLAYABILITIES_API FAggregator : public TSharedFromThis<FAggregator>
 {
 	DECLARE_MULTICAST_DELEGATE_OneParam(FOnAggregatorDirty, FAggregator*);
 
-	FAggregator(float InBaseValue=0.f) : CallbacksDisabled(false), BaseValue(InBaseValue) { }
+	FAggregator(float InBaseValue=0.f) : BaseValue(InBaseValue) { }
 
+	/** Simple accessor to base value */
+	float GetBaseValue() const;
 	void SetBaseValue(float NewBaseValue);
 	
 	void ExecModOnBaseValue(TEnumAsByte<EGameplayModOp::Type> ModifierOp, float EvaluatedMagnitude);
@@ -41,15 +49,15 @@ struct GAMEPLAYABILITIES_API FAggregator : public TSharedFromThis<FAggregator>
 	void AddMod(float EvaluatedData, TEnumAsByte<EGameplayModOp::Type> ModifierOp, const FGameplayTagRequirements*	SourceTagReqs, const FGameplayTagRequirements* TargetTagReqs, FActiveGameplayEffectHandle ActiveHandle = FActiveGameplayEffectHandle());
 	void RemoveMod(FActiveGameplayEffectHandle ActiveHandle);
 
-	void DisableCallbacks();
-	void EnabledCallbacks();
-
 	/** Evaluates the Aggregator with the internal base value and given parameters */
 	float Evaluate(const FAggregatorEvaluateParameters& Parameters) const;
 
 	/** Evaluates the Aggregator with an arbitrary base value */
 	float EvaluateWithBase(float InlineBaseValue, const FAggregatorEvaluateParameters& Parameters) const;
 
+	/** Evaluates the Aggregator to compute its "bonus" (final - base) value */
+	float EvaluateBonus(const FAggregatorEvaluateParameters& Parameters) const;
+
 	void TakeSnapshotOf(const FAggregator& AggToSnapshot);
 
 	FOnAggregatorDirty OnDirty;
@@ -60,8 +68,6 @@ private:
 
 	void BroadcastOnDirty();
 
-	bool CallbacksDisabled;
-
 	float	BaseValue;
 	TArray<FAggregatorMod>	Mods[EGameplayModOp::Max];
 
@@ -81,4 +87,4 @@ struct GAMEPLAYABILITIES_API FAggregatorRef
 	TSharedPtr<FAggregator>	Data;
 
 	void TakeSnapshotOf(const FAggregatorRef& RefToSnapshot);
-};
\ No newline at end of file
+};
diff --git a/Engine/Source/Runtime/GameplayAbilities/Public/GameplayEffectCalculation.h b/Engine/Source/Runtime/GameplayAbilities/Public/GameplayEffectCalculation.h
--- a/Engine/Source/Runtime/GameplayAbilities/Public/GameplayEffectCalculation.h
+++ b/Engine/Source/Runtime/GameplayAbilities/Public/GameplayEffectCalculation.h
@@ -2,34 +2,25 @@
 
 #pragma once
 
-
 #include "GameplayEffectCalculation.generated.h"
 
 struct FGameplayEffectSpec;
-struct FGameplayEffectAttributeCaptureDefinition;
-
 class UAbilitySystemComponent;
 
-UCLASS(abstract)
+/** Abstract base for specialized gameplay effect calculations; Capable of specifing attribute captures */
+UCLASS(BlueprintType, Blueprintable, Abstract)
 class GAMEPLAYABILITIES_API UGameplayEffectCalculation : public UObject
 {
 	GENERATED_UCLASS_BODY()
 
-	virtual void Execute(FGameplayEffectSpec& Spec, int32 ExecutionIdx, UAbilitySystemComponent* TargetAbilitySystemComponent) const;
+public:
 
-	virtual void GetCaptureDefinitions(OUT TArray<FGameplayEffectAttributeCaptureDefinition>& Definitions) const;
-};
+	/** Simple accessor to capture definitions for attributes */
+	virtual const TArray<FGameplayEffectAttributeCaptureDefinition>& GetAttributeCaptureDefinitions() const;
 
-// yuck?
+protected:
 
-#define GE_REQ_SRC(SET, PROP, SNAPSHOT) \
-{ \
-	static UProperty* ThisProperty = FindFieldChecked<UProperty>(SET::StaticClass(), GET_MEMBER_NAME_CHECKED(SET, PROP)); \
-	Definitions.Add(FGameplayEffectAttributeCaptureDefinition(FGameplayAttribute(ThisProperty), EGameplayEffectAttributeCaptureSource::Source, SNAPSHOT)); \
-}
-
-#define GE_REQ_TARGET(SET, PROP, SNAPSHOT) \
-{ \
-	static UProperty* ThisProperty = FindFieldChecked<UProperty>(SET::StaticClass(), GET_MEMBER_NAME_CHECKED(SET, PROP)); \
-	Definitions.Add(FGameplayEffectAttributeCaptureDefinition(FGameplayAttribute(ThisProperty), EGameplayEffectAttributeCaptureSource::Target, SNAPSHOT)); \
-}
+	/** Attributes to capture that are relevant to the calculation */
+	UPROPERTY(EditDefaultsOnly, BlueprintReadOnly, Category=Attributes)
+	TArray<FGameplayEffectAttributeCaptureDefinition> RelevantAttributesToCapture;
+};
diff --git a/Engine/Source/Runtime/GameplayAbilities/Public/GameplayEffectExecutionCalculation.h b/Engine/Source/Runtime/GameplayAbilities/Public/GameplayEffectExecutionCalculation.h
new file mode 100644
--- /dev/null
+++ b/Engine/Source/Runtime/GameplayAbilities/Public/GameplayEffectExecutionCalculation.h
@@ -0,0 +1,106 @@
+// Copyright 1998-2014 Epic Games, Inc. All Rights Reserved.
+
+#pragma once
+
+#include "GameplayEffectCalculation.h"
+#include "GameplayEffectAggregator.h"
+#include "GameplayEffectExecutionCalculation.generated.h"
+
+struct FGameplayEffectSpec;
+class UAbilitySystemComponent;
+
+/** Struct representing parameters for a custom gameplay effect execution. Should not be held onto via reference, used just for the scope of the execution */
+USTRUCT()
+struct GAMEPLAYABILITIES_API FGameplayEffectCustomExecutionParameters
+{
+	GENERATED_USTRUCT_BODY()
+
+public:
+
+	// Constructors
+	FGameplayEffectCustomExecutionParameters();
+	FGameplayEffectCustomExecutionParameters(const FGameplayEffectSpec& InOwningSpec, const TArray<FGameplayEffectExecutionScopedModifierInfo>& InScopedMods, UAbilitySystemComponent* InTargetAbilityComponent);
+
+	/** Simple accessor to owning gameplay spec */
+	const FGameplayEffectSpec& GetOwningSpec() const;
+
+	/** Simple accessor to target ability system component */
+	UAbilitySystemComponent* GetTargetAbilitySystemComponent() const;
+	
+	/**
+	 * Attempts to calculate the magnitude of a captured attribute given the specified parameters. Can fail if the gameplay spec doesn't have
+	 * a valid capture for the attribute.
+	 * 
+	 * @param InCaptureDef	Attribute definition to attempt to calculate the magnitude of
+	 * @param InEvalParams	Parameters to evaluate the attribute under
+	 * @param OutMagnitude	[OUT] Computed magnitude, falls back to 0.f in the event of failure
+	 * 
+	 * @return True if the magnitude was successfully calculated, false if it was not
+	 */
+	bool AttemptCalculateCapturedAttributeMagnitude(const FGameplayEffectAttributeCaptureDefinition& InCaptureDef, const FAggregatorEvaluateParameters& InEvalParams, OUT float& OutMagnitude) const;
+	
+	/**
+	 * Attempts to calculate the base value of a captured attribute given the specified parameters. Can fail if the gameplay spec doesn't have
+	 * a valid capture for the attribute.
+	 * 
+	 * @param InCaptureDef	Attribute definition to attempt to calculate the base value of
+	 * @param OutBaseValue	[OUT] Computed base value, falls back to 0.f in the event of failure
+	 * 
+	 * @return True if the base value was successfully calculated, false if it was not
+	 */
+	bool AttemptCalculateCapturedAttributeBaseValue(const FGameplayEffectAttributeCaptureDefinition& InCaptureDef, OUT float& OutBaseValue) const;
+
+	/**
+	 * Attempts to calculate the bonus magnitude of a captured attribute given the specified parameters. Can fail if the gameplay spec doesn't have
+	 * a valid capture for the attribute.
+	 * 
+	 * @param InCaptureDef		Attribute definition to attempt to calculate the bonus magnitude of
+	 * @param InEvalParams		Parameters to evaluate the attribute under
+	 * @param OutBonusMagnitude	[OUT] Computed bonus magnitude, falls back to 0.f in the event of failure
+	 * 
+	 * @return True if the bonus magnitude was successfully calculated, false if it was not
+	 */
+	bool AttemptCalculateCapturedAttributeBonusMagnitude(const FGameplayEffectAttributeCaptureDefinition& InCaptureDef, const FAggregatorEvaluateParameters& InEvalParams, OUT float& OutBonusMagnitude) const;
+	
+	/**
+	 * Attempts to populate the specified aggregator with a snapshot of a backing captured aggregator. Can fail if the gameplay spec doesn't have
+	 * a valid capture for the attribute.
+	 * 
+	 * @param InCaptureDef				Attribute definition to attempt to snapshot
+	 * @param OutSnapshottedAggregator	[OUT] Snapshotted aggregator, if possible
+	 * 
+	 * @return True if the aggregator was successfully snapshotted, false if it was not
+	 */
+	bool AttemptGetCapturedAttributeAggregatorSnapshot(const FGameplayEffectAttributeCaptureDefinition& InCaptureDef, OUT FAggregator& OutSnapshottedAggregator) const;
+
+private:
+
+	/** Mapping of capture definition to aggregator with scoped modifiers added in; Used to process scoped modifiers w/o modifying underlying aggregators in the capture */
+	TMap<FGameplayEffectAttributeCaptureDefinition, FAggregator> ScopedModifierAggregators;
+
+	/** Owning gameplay effect spec */
+	const FGameplayEffectSpec* OwningSpec;
+
+	/** Target ability system component of the execution */
+	TWeakObjectPtr<UAbilitySystemComponent> TargetAbilitySystemComponent;
+};
+
+UCLASS(BlueprintType, Blueprintable, Abstract)
+class GAMEPLAYABILITIES_API UGameplayEffectExecutionCalculation : public UGameplayEffectCalculation
+{
+	GENERATED_UCLASS_BODY()
+
+public:
+
+	/**
+	 * Called whenever the owning gameplay effect is executed. Allowed to do essentially whatever is desired, including generating new
+	 * modifiers to instantly execute as well.
+	 * 
+	 * @note: Native subclasses should override the auto-generated Execute_Implementation function and NOT this one.
+	 * 
+	 * @param ExecutionParams			Parameters for the custom execution calculation
+	 * @param OutAdditionalModifiers	[OUT] Additional modifiers the custom execution has generated and would like executed upon the target
+	 */
+	UFUNCTION(BlueprintNativeEvent, Category="Calculation")
+	void Execute(const FGameplayEffectCustomExecutionParameters& ExecutionParams, TArray<FGameplayModifierEvaluatedData>& OutAdditionalModifiers) const;
+};
diff --git a/Engine/Source/Runtime/GameplayAbilities/Public/GameplayEffectExtension.h b/Engine/Source/Runtime/GameplayAbilities/Public/GameplayEffectExtension.h
--- a/Engine/Source/Runtime/GameplayAbilities/Public/GameplayEffectExtension.h
+++ b/Engine/Source/Runtime/GameplayAbilities/Public/GameplayEffectExtension.h
@@ -10,9 +10,8 @@
 
 struct FGameplayEffectModCallbackData
 {
-	FGameplayEffectModCallbackData(const FGameplayEffectSpec& InEffectSpec, const FModifierSpec& InModifierSpec, FGameplayModifierEvaluatedData& InEvaluatedData, UAbilitySystemComponent& InTarget)
+	FGameplayEffectModCallbackData(const FGameplayEffectSpec& InEffectSpec, FGameplayModifierEvaluatedData& InEvaluatedData, UAbilitySystemComponent& InTarget)
 		: EffectSpec(InEffectSpec)
-		, ModifierSpec(InModifierSpec)
 		, EvaluatedData(InEvaluatedData)
 		, Target(InTarget)
 	{
@@ -20,7 +19,6 @@ struct FGameplayEffectModCallbackData
 	}
 
 	const struct FGameplayEffectSpec&		EffectSpec;		// The spec that the mod came from
-	const struct FModifierSpec&				ModifierSpec;	// The mod we are going to apply
 	struct FGameplayModifierEvaluatedData&	EvaluatedData;	// The 'flat'/computed data to be applied to the target
 
 	class UAbilitySystemComponent &Target;		// Target we intend to apply to
diff --git a/Engine/Source/Runtime/GameplayAbilities/Public/GameplayEffectTypes.h b/Engine/Source/Runtime/GameplayAbilities/Public/GameplayEffectTypes.h
--- a/Engine/Source/Runtime/GameplayAbilities/Public/GameplayEffectTypes.h
+++ b/Engine/Source/Runtime/GameplayAbilities/Public/GameplayEffectTypes.h
@@ -82,6 +82,13 @@ namespace EGameplayEffectStackingPolicy
 	};
 }
 
+/** Enumeration for options of where to capture gameplay attributes from for gameplay effects */
+UENUM()
+enum class EGameplayEffectAttributeCaptureSource : uint8
+{
+	Source,	// Source (caster) of the gameplay effect
+	Target	// Target (recipient) of the gameplay effect
+};
 
 /**
  * This handle is required for things outside of FActiveGameplayEffectsContainer to refer to a specific active GameplayEffect
@@ -113,6 +120,7 @@ struct FActiveGameplayEffectHandle
 	static FActiveGameplayEffectHandle GenerateNewHandle(UAbilitySystemComponent* OwningComponent);
 
 	UAbilitySystemComponent* GetOwningAbilitySystemComponent();
+	const UAbilitySystemComponent* GetOwningAbilitySystemComponent() const;
 
 	bool operator==(const FActiveGameplayEffectHandle& Other) const
 	{
@@ -140,6 +148,96 @@ private:
 	int32 Handle;
 };
 
+USTRUCT()
+struct FGameplayModifierEvaluatedData
+{
+	GENERATED_USTRUCT_BODY()
+
+	FGameplayModifierEvaluatedData()
+		: Attribute()
+		, ModifierOp(EGameplayModOp::Additive)
+		, Magnitude(0.f)
+		, IsValid(false)
+	{
+	}
+
+	FGameplayModifierEvaluatedData(const FGameplayAttribute& InAttribute, TEnumAsByte<EGameplayModOp::Type> InModOp, float InMagnitude, FActiveGameplayEffectHandle InHandle = FActiveGameplayEffectHandle())
+		: Attribute(InAttribute)
+		, ModifierOp(InModOp)
+		, Magnitude(InMagnitude)
+		, Handle(InHandle)
+		, IsValid(true)
+	{
+	}
+
+	UPROPERTY()
+	FGameplayAttribute Attribute;
+
+	/** The numeric operation of this modifier: Override, Add, Multiply, etc  */
+	UPROPERTY()
+	TEnumAsByte<EGameplayModOp::Type> ModifierOp;
+
+	UPROPERTY()
+	float Magnitude;
+
+	/** Handle of the active gameplay effect that originated us. Will be invalid in many cases */
+	UPROPERTY()
+	FActiveGameplayEffectHandle	Handle;
+
+	UPROPERTY()
+	bool IsValid;
+
+	FString ToSimpleString() const
+	{
+		return FString::Printf(TEXT("%s %s EvalMag: %f"), *Attribute.GetName(), *EGameplayModOpToString(ModifierOp), Magnitude);
+	}
+};
+
+/** Struct defining gameplay attribute capture options for gameplay effects */
+USTRUCT()
+struct FGameplayEffectAttributeCaptureDefinition
+{
+	GENERATED_USTRUCT_BODY()
+
+	FGameplayEffectAttributeCaptureDefinition()
+	{
+
+	}
+
+	FGameplayEffectAttributeCaptureDefinition(FGameplayAttribute InAttribute, EGameplayEffectAttributeCaptureSource InSource, bool InSnapshot)
+		: AttributeToCapture(InAttribute), AttributeSource(InSource), bSnapshot(InSnapshot)
+	{
+
+	}
+
+	/** Gameplay attribute to capture */
+	UPROPERTY(EditDefaultsOnly, Category=Capture)
+	FGameplayAttribute AttributeToCapture;
+
+	/** Source of the gameplay attribute */
+	UPROPERTY(EditDefaultsOnly, Category=Capture)
+	EGameplayEffectAttributeCaptureSource AttributeSource;
+
+	/** Whether the attribute should be snapshotted or not */
+	UPROPERTY(EditDefaultsOnly, Category=Capture)
+	bool bSnapshot;
+
+	/** Equality/Inequality operators */
+	bool operator==(const FGameplayEffectAttributeCaptureDefinition& Other) const;
+	bool operator!=(const FGameplayEffectAttributeCaptureDefinition& Other) const;
+
+	/**
+	 * Get type hash for the capture definition; Implemented to allow usage in TMap
+	 *
+	 * @param CaptureDef Capture definition to get the type hash of
+	 */
+	friend uint32 GetTypeHash(const FGameplayEffectAttributeCaptureDefinition& CaptureDef)
+	{
+		return FCrc::MemCrc32(&CaptureDef, sizeof(FGameplayEffectAttributeCaptureDefinition));
+	}
+
+	FString ToSimpleString() const;
+};
 
 /**
  * FGameplayEffectContext
diff --git a/Engine/Source/Runtime/GameplayAbilities/Public/GameplayModCalculation.h b/Engine/Source/Runtime/GameplayAbilities/Public/GameplayModCalculation.h
deleted file mode 100644
diff --git a/Engine/Source/Runtime/GameplayAbilities/Public/GameplayModMagnitudeCalculation.h b/Engine/Source/Runtime/GameplayAbilities/Public/GameplayModMagnitudeCalculation.h
new file mode 100644
--- /dev/null
+++ b/Engine/Source/Runtime/GameplayAbilities/Public/GameplayModMagnitudeCalculation.h
@@ -0,0 +1,29 @@
+// Copyright 1998-2014 Epic Games, Inc. All Rights Reserved.
+
+#pragma once
+
+#include "GameplayEffectCalculation.h"
+#include "GameplayEffect.h"
+#include "GameplayModMagnitudeCalculation.generated.h"
+
+// Forward declarations
+class UAbilitySystemComponent;
+
+/** Class used to perform custom gameplay effect modifier calculations, either via blueprint or native code */ 
+UCLASS(BlueprintType, Blueprintable, Abstract)
+class GAMEPLAYABILITIES_API UGameplayModMagnitudeCalculation : public UGameplayEffectCalculation
+{
+
+public:
+	GENERATED_UCLASS_BODY()
+
+	/**
+	 * Calculate the magnitude of the gameplay effect modifier, given the specified spec
+	 * 
+	 * @param Spec	Gameplay effect spec to use to calculate the magnitude with
+	 * 
+	 * @return Computed magnitude of the modifier
+	 */
+	UFUNCTION(BlueprintNativeEvent, Category="Calculation")
+	float CalculateMagnitude(const FGameplayEffectSpec& Spec) const;
+};
\ No newline at end of file
