- Non-unity build fixes

diff --git a/Engine/Source/Runtime/GameplayAbilities/Public/GameplayEffectCalculation.h b/Engine/Source/Runtime/GameplayAbilities/Public/GameplayEffectCalculation.h
--- a/Engine/Source/Runtime/GameplayAbilities/Public/GameplayEffectCalculation.h
+++ b/Engine/Source/Runtime/GameplayAbilities/Public/GameplayEffectCalculation.h
@@ -2,6 +2,7 @@
 
 #pragma once
 
+#include "GameplayEffectTypes.h"
 #include "GameplayEffectCalculation.generated.h"
 
 struct FGameplayEffectSpec;
diff --git a/Engine/Source/Runtime/GameplayAbilities/Public/GameplayEffectExecutionCalculation.h b/Engine/Source/Runtime/GameplayAbilities/Public/GameplayEffectExecutionCalculation.h
--- a/Engine/Source/Runtime/GameplayAbilities/Public/GameplayEffectExecutionCalculation.h
+++ b/Engine/Source/Runtime/GameplayAbilities/Public/GameplayEffectExecutionCalculation.h
@@ -4,6 +4,7 @@
 
 #include "GameplayEffectCalculation.h"
 #include "GameplayEffectAggregator.h"
+#include "GameplayEffect.h"
 #include "GameplayEffectExecutionCalculation.generated.h"
 
 struct FGameplayEffectSpec;
