FModifierSpec
-Now only holds EvaluatedMagnitude
-Calling code must correlate with Def->Modifier list
-Sets us up to properly replicate FModifierSpec::EvaluatedData

diff --git a/Engine/Source/Runtime/GameplayAbilities/Private/AbilitySystemComponent.cpp b/Engine/Source/Runtime/GameplayAbilities/Private/AbilitySystemComponent.cpp
--- a/Engine/Source/Runtime/GameplayAbilities/Private/AbilitySystemComponent.cpp
+++ b/Engine/Source/Runtime/GameplayAbilities/Private/AbilitySystemComponent.cpp
@@ -412,9 +412,9 @@ FActiveGameplayEffectHandle UAbilitySystemComponent::ApplyGameplayEffectSpecToSe
 	// Check AttributeSet requirements: do we have everything this GameplayEffectSpec expects?
 	// We may want to cache this off in some way to make the runtime check quicker.
 	// We also need to handle things in the execution list
-	for (const FModifierSpec& Mod : Spec.Modifiers)
+	for (const FGameplayModifierInfo& Mod : Spec.Def->Modifiers)
 	{
-		if (HasAttributeSetForAttribute(Mod.Info.Attribute) == false)
+		if (HasAttributeSetForAttribute(Mod.Attribute) == false)
 		{
 			return FActiveGameplayEffectHandle();
 		}
diff --git a/Engine/Source/Runtime/GameplayAbilities/Private/GameplayEffect.cpp b/Engine/Source/Runtime/GameplayAbilities/Private/GameplayEffect.cpp
--- a/Engine/Source/Runtime/GameplayAbilities/Private/GameplayEffect.cpp
+++ b/Engine/Source/Runtime/GameplayAbilities/Private/GameplayEffect.cpp
@@ -227,12 +227,7 @@ FGameplayEffectSpec::FGameplayEffectSpec(const UGameplayEffect* InDef, const FGa
 	SetContext(InEffectContext);
 
 	// Init our ModifierSpecs
-	Modifiers.Reserve(Def->Modifiers.Num());
-	for (const FGameplayModifierInfo &ModInfo : Def->Modifiers)
-	{
-		// This creates a new FModifierSpec that we own.
-		new (Modifiers)FModifierSpec(ModInfo);
-	}
+	Modifiers.SetNum(Def->Modifiers.Num());
 
 	// Prep the spec with all of the attribute captures it will need to perform
 	SetupAttributeCaptureDefinitions();
@@ -253,10 +248,14 @@ FGameplayEffectSpec::FGameplayEffectSpec(const UGameplayEffect* InDef, const FGa
 void FGameplayEffectSpec::SetupAttributeCaptureDefinitions()
 {
 	// Gather all capture definitions from modifiers
-	for (const FModifierSpec& CurModSpec : Modifiers)
+
+	for (int32 ModIdx = 0; ModIdx < Modifiers.Num(); ++ModIdx)
 	{
+		const FGameplayModifierInfo& ModDef = Def->Modifiers[ModIdx];
+		const FModifierSpec& ModSpec = Modifiers[ModIdx];
+
 		TArray<FGameplayEffectAttributeCaptureDefinition> CalculationCaptureDefs;
-		CurModSpec.Info.ModifierMagnitude.GetAttributeCaptureDefinitions(CalculationCaptureDefs);
+		ModDef.ModifierMagnitude.GetAttributeCaptureDefinitions(CalculationCaptureDefs);
 		
 		for (const FGameplayEffectAttributeCaptureDefinition& CurCaptureDef : CalculationCaptureDefs)
 		{
@@ -337,34 +336,37 @@ float FGameplayEffectSpec::GetChanceToExecuteOnGameplayEffect() const
 	return ChanceToExecuteOnGameplayEffect;
 }
 
-float FGameplayEffectSpec::GetModifierMagnitude(const FModifierSpec& ModSpec) const
-{
-	return ModSpec.EvaluatedMagnitude;
-}
-
 void FGameplayEffectSpec::CalculateModifierMagnitudes()
 {
-	for (FModifierSpec& Mod : Modifiers)
+	for(int32 ModIdx = 0; ModIdx < Modifiers.Num(); ++ModIdx)
 	{
-		Mod.CalculateMagnitude(*this);
+		const FGameplayModifierInfo& ModDef = Def->Modifiers[ModIdx];
+		FModifierSpec& ModSpec = Modifiers[ModIdx];
+
+		if (ModDef.ModifierMagnitude.AttemptCalculateMagnitude(*this, ModSpec.EvaluatedMagnitude) == false)
+		{
+			ModSpec.EvaluatedMagnitude = 0.f;
+			ABILITY_LOG(Warning, TEXT("Modifier on spec: %s was asked to CalculateMagnitude and failed, falling back to 0."), *ToSimpleString());
+		}
 	}
 }
 
 float FGameplayEffectSpec::GetMagnitude(const FGameplayAttribute &Attribute) const
 {
-	float CurrentMagnitude = 0.f;
-	for (const FModifierSpec &Mod : Modifiers)
-	{		
-		if (Mod.Info.Attribute != Attribute)
+	for (int32 ModIdx = 0; ModIdx < Modifiers.Num(); ++ModIdx)
+	{
+		const FGameplayModifierInfo& ModDef = Def->Modifiers[ModIdx];
+		const FModifierSpec& ModSpec = Modifiers[ModIdx];
+	
+		if (ModDef.Attribute != Attribute)
 		{
 			continue;
 		}
 
-		CurrentMagnitude = GetModifierMagnitude(Mod);
-		break;
+		return ModSpec.EvaluatedMagnitude;
 	}
 
-	return CurrentMagnitude;
+	return 0.f;
 }
 
 bool FGameplayEffectSpec::HasValidCapturedAttributes(const TArray<FGameplayEffectAttributeCaptureDefinition>& InCaptureDefsToCheck) const
@@ -442,31 +444,10 @@ void FGameplayEffectSpec::SetContext(FGameplayEffectContextHandle NewEffectConte
 
 // --------------------------------------------------------------------------------------------------------------------------------------------------------
 //
-//	FModifierSpec
+//	FGameplayEffectAttributeCaptureSpec
 //
 // --------------------------------------------------------------------------------------------------------------------------------------------------------
 
-FModifierSpec::FModifierSpec(const FGameplayModifierInfo& InInfo)
-	: Info(InInfo)
-	, EvaluatedMagnitude(0.f)
-{
-
-}
-
-float FModifierSpec::GetEvaluatedMagnitude() const
-{
-	return EvaluatedMagnitude;
-}
-
-void FModifierSpec::CalculateMagnitude(OUT FGameplayEffectSpec& OwnerSpec)
-{	
-	if (!Info.ModifierMagnitude.AttemptCalculateMagnitude(OwnerSpec, EvaluatedMagnitude))
-	{
-		EvaluatedMagnitude = 0.f;
-		ABILITY_LOG(Warning, TEXT("Modifier on spec: %s was asked to CalculateMagnitude and failed, falling back to 0."), *OwnerSpec.ToSimpleString());
-	}
-}
-
 FGameplayEffectAttributeCaptureSpec::FGameplayEffectAttributeCaptureSpec()
 {
 }
@@ -653,11 +634,7 @@ void FActiveGameplayEffect::CheckOngoingTagRequirements(const FGameplayTagContai
 				for (int32 ModIdx = 0; ModIdx < Spec.Modifiers.Num(); ++ModIdx)
 				{
 					const FModifierSpec &Mod = Spec.Modifiers[ModIdx];
-					
-					ABILITY_LOG_SCOPE(TEXT("Applying Attribute Mod %s to property"), *Mod.ToSimpleString());
-
-					FAggregator* Aggregator = OwningContainer.FindOrCreateAttributeAggregator(Mod.Info.Attribute).Get();
-
+					FAggregator* Aggregator = OwningContainer.FindOrCreateAttributeAggregator(Spec.Def->Modifiers[ModIdx].Attribute).Get();
 					Aggregator->RemoveMod(Handle);
 				}
 			}
@@ -670,15 +647,13 @@ void FActiveGameplayEffect::CheckOngoingTagRequirements(const FGameplayTagContai
 				for (int32 ModIdx = 0; ModIdx < Spec.Modifiers.Num(); ++ModIdx)
 				{
 					const FModifierSpec &Mod = Spec.Modifiers[ModIdx];
-					
-					ABILITY_LOG_SCOPE(TEXT("Applying Attribute Mod %s to property"), *Mod.ToSimpleString());
-
-					FAggregator* Aggregator = OwningContainer.FindOrCreateAttributeAggregator(Mod.Info.Attribute).Get();
+					const FGameplayModifierInfo &ModInfo = Spec.Def->Modifiers[ModIdx];
 
-					// GE_FIXME: Figure this out - when does EvaluatedData get calculated (here? Before and stored off?)
-					float EvaluatedData = Spec.GetModifierMagnitude(Mod);
-
-					Aggregator->AddMod(EvaluatedData, Mod.Info.ModifierOp, &Mod.Info.SourceTags, &Mod.Info.TargetTags, Handle);
+					// Note we assume the EvaluatedMagnitude is up to do. There is no case currently where we should recalculate magnitude based on
+					// Ongoing tags being met. We either calculate magnitude one time, or its done via OnDirty calls (or potentially a frequency timer one day)
+					
+					FAggregator* Aggregator = OwningContainer.FindOrCreateAttributeAggregator(Spec.Def->Modifiers[ModIdx].Attribute).Get();
+					Aggregator->AddMod(Mod.GetEvaluatedMagnitude(), ModInfo.ModifierOp, &ModInfo.SourceTags, &ModInfo.TargetTags, Handle);
 				}
 			}
 		}
@@ -809,9 +784,12 @@ void FActiveGameplayEffectsContainer::ExecuteActiveEffectsFrom(FGameplayEffectSp
 	//		These will modify the base value of attributes
 	// ------------------------------------------------------
 
-	for (const FModifierSpec &Mod : Spec.Modifiers)
+	for(int32 ModIdx = 0; ModIdx < Spec.Modifiers.Num(); ++ModIdx)
 	{
-		FGameplayModifierEvaluatedData EvalData(Mod.Info.Attribute, Mod.Info.ModifierOp, Mod.GetEvaluatedMagnitude());
+		const FGameplayModifierInfo& ModDef = Spec.Def->Modifiers[ModIdx];
+		const FModifierSpec& ModSpec = Spec.Modifiers[ModIdx];
+
+		FGameplayModifierEvaluatedData EvalData(ModDef.Attribute, ModDef.ModifierOp, ModSpec.GetEvaluatedMagnitude());
 		InvokeGameplayCueExecute |= InternalExecuteMod(Spec, EvalData);
 	}
 
@@ -957,11 +935,14 @@ float FActiveGameplayEffectsContainer::GetGameplayEffectMagnitude(FActiveGamepla
 	{
 		if (Effect.Handle == Handle)
 		{
-			for (const FModifierSpec &Mod : Effect.Spec.Modifiers)
+			for(int32 ModIdx = 0; ModIdx < Effect.Spec.Modifiers.Num(); ++ModIdx)
 			{
-				if (Mod.Info.Attribute == Attribute)
+				const FGameplayModifierInfo& ModDef = Effect.Spec.Def->Modifiers[ModIdx];
+				const FModifierSpec& ModSpec = Effect.Spec.Modifiers[ModIdx];
+			
+				if (ModDef.Attribute == Attribute)
 				{
-					return Effect.Spec.GetModifierMagnitude(Mod);
+					return ModSpec.GetEvaluatedMagnitude();
 				}
 			}
 		}
@@ -1507,8 +1488,8 @@ void FActiveGameplayEffectsContainer::RecalculateStacking()
 						{
 						case EGameplayEffectStackingPolicy::Highest:
 						{
-							float BestSpecMagnitude = StackedEffect->Spec.GetMagnitude(StackedEffect->Spec.Modifiers[0].Info.Attribute);
-							float CurrSpecMagnitude = Effect.Spec.GetMagnitude(Effect.Spec.Modifiers[0].Info.Attribute);
+							float BestSpecMagnitude = StackedEffect->Spec.GetMagnitude(StackedEffect->Spec.Def->Modifiers[0].Attribute);
+							float CurrSpecMagnitude = Effect.Spec.GetMagnitude(Effect.Spec.Def->Modifiers[0].Attribute);
 							if (BestSpecMagnitude < CurrSpecMagnitude)
 							{
 								StackedEffect = &Effect;
@@ -1517,8 +1498,8 @@ void FActiveGameplayEffectsContainer::RecalculateStacking()
 						}
 						case EGameplayEffectStackingPolicy::Lowest:
 						{
-							float BestSpecMagnitude = StackedEffect->Spec.GetMagnitude(StackedEffect->Spec.Modifiers[0].Info.Attribute);
-							float CurrSpecMagnitude = Effect.Spec.GetMagnitude(Effect.Spec.Modifiers[0].Info.Attribute);
+							float BestSpecMagnitude = StackedEffect->Spec.GetMagnitude(StackedEffect->Spec.Def->Modifiers[0].Attribute);
+							float CurrSpecMagnitude = Effect.Spec.GetMagnitude(Effect.Spec.Def->Modifiers[0].Attribute);
 							if (BestSpecMagnitude > CurrSpecMagnitude)
 							{
 								StackedEffect = &Effect;
@@ -1578,20 +1559,25 @@ void FActiveGameplayEffectsContainer::RecalculateStacking()
 	}
 }
 
-bool FActiveGameplayEffectsContainer::CanApplyAttributeModifiers(const UGameplayEffect *GameplayEffect, float Level, const FGameplayEffectContextHandle& EffectContext)
+bool FActiveGameplayEffectsContainer::CanApplyAttributeModifiers(const UGameplayEffect* GameplayEffect, float Level, const FGameplayEffectContextHandle& EffectContext)
 {
 	SCOPE_CYCLE_COUNTER(STAT_GameplayEffectsCanApplyAttributeModifiers);
 
 	FGameplayEffectSpec	Spec(GameplayEffect, EffectContext, Level);
+
+	Spec.CalculateModifierMagnitudes();
 	
-	for (const FModifierSpec& Mod : Spec.Modifiers)
+	for(int32 ModIdx = 0; ModIdx < Spec.Modifiers.Num(); ++ModIdx)
 	{
+		const FGameplayModifierInfo& ModDef = Spec.Def->Modifiers[ModIdx];
+		const FModifierSpec& ModSpec = Spec.Modifiers[ModIdx];
+	
 		// It only makes sense to check additive operators
-		if (Mod.Info.ModifierOp == EGameplayModOp::Additive)
+		if (ModDef.ModifierOp == EGameplayModOp::Additive)
 		{
-			const UAttributeSet* Set = Owner->GetAttributeSubobject(Mod.Info.Attribute.GetAttributeSetClass());
-			float CurrentValue = Mod.Info.Attribute.GetNumericValueChecked(Set);
-			float CostValue = Spec.GetModifierMagnitude(Mod);
+			const UAttributeSet* Set = Owner->GetAttributeSubobject(ModDef.Attribute.GetAttributeSetClass());
+			float CurrentValue = ModDef.Attribute.GetNumericValueChecked(Set);
+			float CostValue = ModSpec.GetEvaluatedMagnitude();
 
 			if (CurrentValue + CostValue < 0.f)
 			{
@@ -1751,9 +1737,13 @@ bool FActiveGameplayEffectQuery::Matches(const FActiveGameplayEffect& Effect) co
 	if (ModifyingAttribute.IsValid())
 	{
 		bool FailedModifyingAttributeCheck = true;
-		for (const FModifierSpec& Modifier : Effect.Spec.Modifiers)
+
+		for(int32 ModIdx = 0; ModIdx < Effect.Spec.Modifiers.Num(); ++ModIdx)
 		{
-			if (Modifier.Info.Attribute == ModifyingAttribute)
+			const FGameplayModifierInfo& ModDef = Effect.Spec.Def->Modifiers[ModIdx];
+			const FModifierSpec& ModSpec = Effect.Spec.Modifiers[ModIdx];
+
+			if (ModDef.Attribute == ModifyingAttribute)
 			{
 				FailedModifyingAttributeCheck = false;
 				break;
diff --git a/Engine/Source/Runtime/GameplayAbilities/Public/GameplayEffect.h b/Engine/Source/Runtime/GameplayAbilities/Public/GameplayEffect.h
--- a/Engine/Source/Runtime/GameplayAbilities/Public/GameplayEffect.h
+++ b/Engine/Source/Runtime/GameplayAbilities/Public/GameplayEffect.h
@@ -468,10 +468,6 @@ public:
 
 	void ValidateGameplayEffect();
 
-	/** Container of gameplay tags to be cleared upon effect application; Any active effects with these tags that can be cleared, will be */
-	UPROPERTY(EditDefaultsOnly, BlueprintReadOnly, Category=Tags)
-	FGameplayTagContainer ClearTagsContainer;
-
 	virtual void PostLoad() override;
 
 	// ----------------------------------------------
@@ -498,10 +494,8 @@ public:
 	UPROPERTY(EditDefaultsOnly, Category = Stacking)
 	TSubclassOf<class UGameplayEffectStackingExtension> StackingExtension;
 
-
-
 	// ----------------------------------------------------------------------
-	//		Tag pass: properties that make sense "right now" (remove this comment DaveR/BillyB)
+	//	Tag Containers
 	// ----------------------------------------------------------------------
 	
 	/** The GameplayEffect's Tags: tags the the GE *has* and DOES NOT give to the actor. */
@@ -520,36 +514,31 @@ public:
 	UPROPERTY(EditDefaultsOnly, BlueprintReadOnly, Category = Tags)
 	FGameplayTagRequirements ApplicationTagRequirements;
 
+	// Container of gameplay tags to be cleared upon effect application; Any active effects with these tags that can be cleared, will be.
+	/** CURRENTLY NOT IMPLEMENTED */
+	UPROPERTY(EditDefaultsOnly, BlueprintReadOnly, Category=Tags)
+	FGameplayTagContainer ClearTagsContainer;
+
 };
 
-/**
- * Modifier Specification
- *	-Const data (FGameplayModifierInfo) tells us what we modify, what we can modify
- *	-Mutable Aggregated data tells us how we modify (magnitude).
- *  
- * Modifiers can be modified. A modifier spec holds these modifications along with a reference to the const data about the modifier.
- * 
- */
+/** Holds evaluated magnitude from a GameplayEffect modifier */
+USTRUCT()
 struct FModifierSpec
 {
-	FModifierSpec(const FGameplayModifierInfo& InInfo);
+	GENERATED_USTRUCT_BODY()
+
+	FModifierSpec() : EvaluatedMagnitude(0.f) { }
 
-	// @todo: Need a way to indicate whether the magnitude was properly calculated or not. Asking for a value early (before target is available, etc.
-	// will just result in incorrect values)
-	float GetEvaluatedMagnitude() const;
-	void CalculateMagnitude(OUT struct FGameplayEffectSpec& OwnerSpec);
+	float GetEvaluatedMagnitude() const { return EvaluatedMagnitude; }
 
-	// Hard Ref to what we modify, this stuff is const and never changes
-	const FGameplayModifierInfo& Info;
+private:
 
 	// @todo: Probably need to make the modifier info private so people don't go gunking around in the magnitude
 	/** In the event that the modifier spec requires custom magnitude calculations, this is the authoritative, last evaluated value of the magnitude */
+	UPROPERTY()
 	float EvaluatedMagnitude;
 
-	FString ToSimpleString() const
-	{
-		return Info.ToSimpleString();
-	}
+	friend struct FGameplayEffectSpec;
 };
 
 /** Saves list of modified attributes, to use for gameplay cues or later processing */
@@ -628,7 +617,7 @@ struct FGameplayEffectAttributeCaptureSpec
 	 *
 	 * @param OutAggregatorSnapshot	[OUT] Snapshotted aggregator, if possible
 	 *
- 	 * @return True if the aggregator was successfully snapshotted, false if it was not
+	 * @return True if the aggregator was successfully snapshotted, false if it was not
 	 */
 	bool AttemptGetAttributeAggregatorSnapshot(OUT FAggregator& OutAggregatorSnapshot) const;
 
@@ -638,7 +627,7 @@ struct FGameplayEffectAttributeCaptureSpec
 	 *
 	 * @param OutAggregatorToAddTo	[OUT] Aggregator with mods appended, if possible
 	 *
- 	 * @return True if the aggregator had mods successfully added to it, false if it did not
+	 * @return True if the aggregator had mods successfully added to it, false if it did not
 	 */
 	bool AttemptAddAggregatorModsToAggregator(OUT FAggregator& OutAggregatorToAddTo) const;
 	
@@ -829,23 +818,22 @@ struct FGameplayEffectSpec
 	UPROPERTY()
 	FTagContainerAggregator	CapturedTargetTags;
 	
+	UPROPERTY()
 	TArray<FModifierSpec> Modifiers;
 	
-	// GE_REMOVE: this feels strange but need to start standardizing how people get this info.
-	float GetModifierMagnitude(const FModifierSpec& ModSpec) const;
 	void CalculateModifierMagnitudes();
 
-	FString ToSimpleString() const
-	{
-		return FString::Printf(TEXT("%s"), *Def->GetName());
-	}
-
 	void SetLevel(float InLevel);
 
 	float GetLevel() const;
 
 	void PrintAll() const;
 
+	FString ToSimpleString() const
+	{
+		return FString::Printf(TEXT("%s"), *Def->GetName());
+	}
+
 private:
 
 	void CaptureDataFromSource();
