added a level - 1 so its 0 based index like the read fixing everything starting at level 2 and not being able to hit max level

diff --git a/Engine/Source/Runtime/GameplayAbilities/Private/AttributeSet.cpp b/Engine/Source/Runtime/GameplayAbilities/Private/AttributeSet.cpp
--- a/Engine/Source/Runtime/GameplayAbilities/Private/AttributeSet.cpp
+++ b/Engine/Source/Runtime/GameplayAbilities/Private/AttributeSet.cpp
@@ -296,14 +296,14 @@ void FAttributeSetInitter::InitAttributeSetDefaults(UAbilitySystemComponent* Abi
 		}
 	}
 
-	if (!Collection->LevelData.IsValidIndex(Level))
+	if (!Collection->LevelData.IsValidIndex(Level - 1))
 	{
 		// We could eventually extrapolate values outside of the max defined levels
 		ABILITY_LOG(Warning, TEXT("Attribute defaults for Level %d are not defined! Skipping"), Level);
 		return;
 	}
 
-	const FAttributeSetDefaults& SetDefaults = Collection->LevelData[Level];
+	const FAttributeSetDefaults& SetDefaults = Collection->LevelData[Level - 1];
 	for (const UAttributeSet* Set : AbilitySystemComponent->SpawnedAttributes)
 	{
 		const FAttributeDefaultValueList* DefaultDataList = SetDefaults.DataMap.Find(Set->GetClass());
