#orion - Added BlueprintCallable version of ApplyGameplayEffectToTarget for minion use

diff --git a/Engine/Source/Runtime/GameplayAbilities/Private/AbilitySystemComponent.cpp b/Engine/Source/Runtime/GameplayAbilities/Private/AbilitySystemComponent.cpp
--- a/Engine/Source/Runtime/GameplayAbilities/Private/AbilitySystemComponent.cpp
+++ b/Engine/Source/Runtime/GameplayAbilities/Private/AbilitySystemComponent.cpp
@@ -185,6 +185,11 @@ FGameplayEffectContextHandle UAbilitySystemComponent::GetEffectContext() const
 	return Context;
 }
 
+FActiveGameplayEffectHandle UAbilitySystemComponent::K2_ApplyGameplayEffectToTarget(UGameplayEffect *GameplayEffect, UAbilitySystemComponent *Target, float Level, FGameplayEffectContextHandle Context)
+{
+	return ApplyGameplayEffectToTarget(GameplayEffect, Target, Level, Context);
+}
+
 /** This is a helper function used in automated testing, I'm not sure how useful it will be to gamecode or blueprints */
 FActiveGameplayEffectHandle UAbilitySystemComponent::ApplyGameplayEffectToTarget(UGameplayEffect *GameplayEffect, UAbilitySystemComponent *Target, float Level, FGameplayEffectContextHandle Context, FPredictionKey PredictionKey)
 {
diff --git a/Engine/Source/Runtime/GameplayAbilities/Public/AbilitySystemComponent.h b/Engine/Source/Runtime/GameplayAbilities/Public/AbilitySystemComponent.h
--- a/Engine/Source/Runtime/GameplayAbilities/Public/AbilitySystemComponent.h
+++ b/Engine/Source/Runtime/GameplayAbilities/Public/AbilitySystemComponent.h
@@ -296,6 +296,9 @@ class GAMEPLAYABILITIES_API UAbilitySystemComponent : public UActorComponent, pu
 	
 	FOnActiveGameplayEffectRemoved* OnGameplayEffectRemovedDelegate(FActiveGameplayEffectHandle Handle);
 
+	UFUNCTION(BlueprintCallable, Category = GameplayEffects, meta=(FriendlyName = "ApplyGameplayEffectToTarget"))
+	FActiveGameplayEffectHandle K2_ApplyGameplayEffectToTarget(UGameplayEffect *GameplayEffect, UAbilitySystemComponent *Target, float Level, FGameplayEffectContextHandle Context);
+
 	FActiveGameplayEffectHandle ApplyGameplayEffectToTarget(UGameplayEffect *GameplayEffect, UAbilitySystemComponent *Target, float Level = UGameplayEffect::INVALID_LEVEL, FGameplayEffectContextHandle Context = FGameplayEffectContextHandle(), FPredictionKey PredictionKey = FPredictionKey());
 
 	UFUNCTION(BlueprintCallable, Category = GameplayEffects, meta=(FriendlyName = "ApplyGameplayEffectToSelf"))
