removed FindChecked since a case was discovered where invalid tags could come through

diff --git a/Engine/Source/Runtime/GameplayAbilities/Private/GameplayCueManager.cpp b/Engine/Source/Runtime/GameplayAbilities/Private/GameplayCueManager.cpp
--- a/Engine/Source/Runtime/GameplayAbilities/Private/GameplayCueManager.cpp
+++ b/Engine/Source/Runtime/GameplayAbilities/Private/GameplayCueManager.cpp
@@ -30,13 +30,19 @@ void UGameplayCueManager::HandleGameplayCues(AActor* TargetActor, const FGamepla
 
 void UGameplayCueManager::HandleGameplayCue(AActor* TargetActor, FGameplayTag GameplayCueTag, EGameplayCueEvent::Type EventType, FGameplayCueParameters Parameters)
 {
-	int32 DataIdx = GameplayCueDataMap.FindChecked(GameplayCueTag);	// We should always find data here. We prepopulate the entire map with every gameplaycue tag when we init.
-	HandleGameplayCueNotify_Internal(TargetActor, DataIdx, EventType, Parameters);
-
-	IGameplayCueInterface* GameplayCueInterface = Cast<IGameplayCueInterface>(TargetActor);
-	if (GameplayCueInterface)
+	// GameplayCueTags could have been removed from the dictionary but not content. When the content is resaved the old tag will be cleaned up, but it could still come through here
+	// at runtime. Since we only populate the map with dictionary gameplaycue tags, we may not find it here.
+	int32* Ptr=GameplayCueDataMap.Find(GameplayCueTag);
+	if (Ptr)
 	{
-		GameplayCueInterface->HandleGameplayCue(TargetActor, GameplayCueTag, EventType, Parameters);
+		int32 DataIdx = *Ptr;
+		HandleGameplayCueNotify_Internal(TargetActor, DataIdx, EventType, Parameters);
+
+		IGameplayCueInterface* GameplayCueInterface = Cast<IGameplayCueInterface>(TargetActor);
+		if (GameplayCueInterface)
+		{
+			GameplayCueInterface->HandleGameplayCue(TargetActor, GameplayCueTag, EventType, Parameters);
+		}
 	}
 }
 
