GameplayCueNotify objects update correctly in-editor and are UObjects. _Blueprint class still exists, not folded up to base class yet.

diff --git a/Engine/Source/Runtime/Engine/Private/RepLayout.cpp b/Engine/Source/Runtime/Engine/Private/RepLayout.cpp
--- a/Engine/Source/Runtime/Engine/Private/RepLayout.cpp
+++ b/Engine/Source/Runtime/Engine/Private/RepLayout.cpp
@@ -24,8 +24,6 @@ FAutoConsoleVariable CVarDoReplicationContextString( TEXT( "net.ContextDebug" ),
 
 //#define ENABLE_SUPER_CHECKSUMS
 
-//#pragma optimize("", off)
-
 #ifdef USE_CUSTOM_COMPARE
 static FORCEINLINE bool CompareBool( const FRepLayoutCmd& Cmd, const void* A, const void* B )
 {
diff --git a/Engine/Source/Runtime/GameplayAbilities/Private/AbilitySystemGlobals.cpp b/Engine/Source/Runtime/GameplayAbilities/Private/AbilitySystemGlobals.cpp
--- a/Engine/Source/Runtime/GameplayAbilities/Private/AbilitySystemGlobals.cpp
+++ b/Engine/Source/Runtime/GameplayAbilities/Private/AbilitySystemGlobals.cpp
@@ -29,7 +29,7 @@ void UAbilitySystemGlobals::InitGlobalData()
 	GetGlobalCurveTable();
 	GetGlobalAttributeMetaDataTable();
 	
-	InitAtributeDefaults();
+	InitAttributeDefaults();
 
 	GetGameplayCueManager();
 }
@@ -184,7 +184,7 @@ FAttributeSetInitter* UAbilitySystemGlobals::GetAttributeSetInitter() const
 	return GlobalAttributeSetInitter.Get();
 }
 
-void UAbilitySystemGlobals::InitAtributeDefaults()
+void UAbilitySystemGlobals::InitAttributeDefaults()
 {
 	if (InternalGetLoadTable<UCurveTable>(GlobalAttributeDefaultsTable, GlobalAttributeSetDefaultsTableName))
 	{	
diff --git a/Engine/Source/Runtime/GameplayAbilities/Private/GameplayCueManager.cpp b/Engine/Source/Runtime/GameplayAbilities/Private/GameplayCueManager.cpp
--- a/Engine/Source/Runtime/GameplayAbilities/Private/GameplayCueManager.cpp
+++ b/Engine/Source/Runtime/GameplayAbilities/Private/GameplayCueManager.cpp
@@ -15,6 +15,7 @@ UGameplayCueManager::UGameplayCueManager(const class FObjectInitializer& PCIP)
 : Super(PCIP)
 {
 #if WITH_EDITOR
+	bAccelerationMapOutdated = true;
 	RegisteredEditorCallbacks = false;
 #endif
 }
@@ -161,23 +162,35 @@ void UGameplayCueManager::LoadObjectLibraryFromPaths(const TArray<FString>& InPa
 	bFullyLoad = InFullyLoad;
 
 	LoadObjectLibrary_Internal();
-
 #if WITH_EDITOR
+	bAccelerationMapOutdated = false;
 	if (!RegisteredEditorCallbacks)
 	{
 		FAssetRegistryModule& AssetRegistryModule = FModuleManager::LoadModuleChecked<FAssetRegistryModule>("AssetRegistry");
 		AssetRegistryModule.Get().OnInMemoryAssetCreated().AddUObject(this, &UGameplayCueManager::HandleAssetAdded);
 		AssetRegistryModule.Get().OnInMemoryAssetDeleted().AddUObject(this, &UGameplayCueManager::HandleAssetDeleted);
+		FWorldDelegates::OnPreWorldInitialization.AddUObject(this, &UGameplayCueManager::ReloadObjectLibrary);
 		RegisteredEditorCallbacks = true;
 	}
 #endif
 }
 
+#if WITH_EDITOR
+void UGameplayCueManager::ReloadObjectLibrary(UWorld* World, const UWorld::InitializationValues IVS)
+{
+	if (bAccelerationMapOutdated)
+	{
+		LoadObjectLibrary_Internal();
+	}
+}
+#endif
+
 void UGameplayCueManager::LoadObjectLibrary_Internal()
 {
 	DECLARE_SCOPE_CYCLE_COUNTER(TEXT("Loading Library"), STAT_ObjectLibrary, STATGROUP_LoadTime);
 
 #if WITH_EDITOR
+	bAccelerationMapOutdated = false;
 	FFormatNamedArguments Args;
 	FScopedSlowTask SlowTask(0, FText::Format(NSLOCTEXT("AbilitySystemEditor", "BeginLoadingGameplayCueNotify", "Loading GameplayCue Library"), Args));
 	SlowTask.MakeDialog();
@@ -186,7 +199,6 @@ void UGameplayCueManager::LoadObjectLibrary_Internal()
 	FScopeCycleCounterUObject PreloadScope(GameplayCueNotifyObjectLibrary);
 	
 	GameplayCueNotifyObjectLibrary->LoadBlueprintAssetDataFromPaths(LoadedPaths);
-	GameplayCueNotifyObjectLibrary->LoadAssetDataFromPaths(LoadedPaths);
 
 	if (bFullyLoad)
 	{
diff --git a/Engine/Source/Runtime/GameplayAbilities/Private/GameplayCueNotify.cpp b/Engine/Source/Runtime/GameplayAbilities/Private/GameplayCueNotify.cpp
--- a/Engine/Source/Runtime/GameplayAbilities/Private/GameplayCueNotify.cpp
+++ b/Engine/Source/Runtime/GameplayAbilities/Private/GameplayCueNotify.cpp
@@ -3,6 +3,8 @@
 #include "AbilitySystemPrivatePCH.h"
 #include "GameplayCueNotify.h"
 #include "GameplayTagsModule.h"
+#include "GameplayCueManager.h"
+#include "AbilitySystemComponent.h"
 
 
 
@@ -12,6 +14,13 @@ UGameplayCueNotify::UGameplayCueNotify(const class FObjectInitializer& PCIP)
 	IsOverride = true;
 }
 
+#if WITH_EDITOR
+void UGameplayCueNotify::PostEditChangeProperty(FPropertyChangedEvent& PropertyChangedEvent)
+{
+	UAbilitySystemGlobals::Get().GetGameplayCueManager()->bAccelerationMapOutdated = true;
+}
+#endif
+
 void UGameplayCueNotify::DeriveGameplayCueTagFromAssetName()
 {
 	// In the editor, attempt to infer GameplayCueTag from our asset name (if there is no valid GameplayCueTag already).
diff --git a/Engine/Source/Runtime/GameplayAbilities/Public/AbilitySystemGlobals.h b/Engine/Source/Runtime/GameplayAbilities/Public/AbilitySystemGlobals.h
--- a/Engine/Source/Runtime/GameplayAbilities/Public/AbilitySystemGlobals.h
+++ b/Engine/Source/Runtime/GameplayAbilities/Public/AbilitySystemGlobals.h
@@ -88,7 +88,7 @@ class GAMEPLAYABILITIES_API UAbilitySystemGlobals : public UObject
 
 protected:
 
-	virtual void InitAtributeDefaults();
+	virtual void InitAttributeDefaults();
 	virtual void AllocAttributeSetInitter();
 
 private:
diff --git a/Engine/Source/Runtime/GameplayAbilities/Public/GameplayCueManager.h b/Engine/Source/Runtime/GameplayAbilities/Public/GameplayCueManager.h
--- a/Engine/Source/Runtime/GameplayAbilities/Public/GameplayCueManager.h
+++ b/Engine/Source/Runtime/GameplayAbilities/Public/GameplayCueManager.h
@@ -134,12 +134,19 @@ class GAMEPLAYABILITIES_API UGameplayCueManager : public UDataAsset
 	void HandleAssetDeleted(UObject *Object);
 
 	bool RegisteredEditorCallbacks;
+
+	bool bAccelerationMapOutdated;
 #endif
 
 private:
 
 	virtual void HandleGameplayCueNotify_Internal(AActor* TargetActor, int32 DataIdx, EGameplayCueEvent::Type EventType, FGameplayCueParameters Parameters);
 
+	#if WITH_EDITOR
+		//This handles the case where GameplayCueNotifications have changed between sessions, which is possible in editor.
+		void ReloadObjectLibrary(UWorld* World, const UWorld::InitializationValues IVS);
+	#endif
+
 	void LoadObjectLibrary_Internal();
 
 	void AddGameplayCueData_Internal(FGameplayTag  GameplayCueTag, FStringAssetReference StringRef);
diff --git a/Engine/Source/Runtime/GameplayAbilities/Public/GameplayCueNotify.h b/Engine/Source/Runtime/GameplayAbilities/Public/GameplayCueNotify.h
--- a/Engine/Source/Runtime/GameplayAbilities/Public/GameplayCueNotify.h
+++ b/Engine/Source/Runtime/GameplayAbilities/Public/GameplayCueNotify.h
@@ -15,7 +15,7 @@
  */
 
 UCLASS(abstract, BlueprintType, hidecategories=Object, collapsecategories)
-class GAMEPLAYABILITIES_API UGameplayCueNotify : public UDataAsset
+class GAMEPLAYABILITIES_API UGameplayCueNotify : public UObject
 {
 	GENERATED_UCLASS_BODY()
 
@@ -33,6 +33,10 @@ class GAMEPLAYABILITIES_API UGameplayCueNotify : public UDataAsset
 	
 	virtual void HandleGameplayCue(AActor* MyTarget, EGameplayCueEvent::Type EventType, FGameplayCueParameters Parameters);
 
+	#if WITH_EDITOR
+		virtual void PostEditChangeProperty(FPropertyChangedEvent& PropertyChangedEvent) override;
+	#endif // WITH_EDITOR
+
 	UPROPERTY(EditDefaultsOnly, Category=GameplayCue, meta=(Categories="GameplayCue"))
 	FGameplayTag	GameplayCueTag;
 
