Add GAMEPLAYABILITIES_API to UAbilityTask_WaitTargetData;

* Fixes linker errors when deriving a class from UAbilityTask_WaitTargetData.

diff --git a/Engine/Source/Runtime/GameplayAbilities/Public/Abilities/Tasks/AbilityTask_WaitTargetData.h b/Engine/Source/Runtime/GameplayAbilities/Public/Abilities/Tasks/AbilityTask_WaitTargetData.h
--- a/Engine/Source/Runtime/GameplayAbilities/Public/Abilities/Tasks/AbilityTask_WaitTargetData.h
+++ b/Engine/Source/Runtime/GameplayAbilities/Public/Abilities/Tasks/AbilityTask_WaitTargetData.h
@@ -10,8 +10,8 @@ class AGameplayAbilityTargetActor;
 DECLARE_DYNAMIC_MULTICAST_DELEGATE_OneParam(FWaitTargetDataDelegate, FGameplayAbilityTargetDataHandle, Data);
 
 /** Wait for targeting actor (spawned from parameter) to provide data. Can be set not to end upon outputting data. Can be ended by task name. */
-UCLASS(MinimalAPI, notplaceable)
-class UAbilityTask_WaitTargetData: public UAbilityTask
+UCLASS(notplaceable)
+class GAMEPLAYABILITIES_API UAbilityTask_WaitTargetData: public UAbilityTask
 {
 	GENERATED_UCLASS_BODY()
 
@@ -33,7 +33,7 @@ class UAbilityTask_WaitTargetData: public UAbilityTask
 	UFUNCTION()
 	void OnTargetDataCancelledCallback(FGameplayAbilityTargetDataHandle Data);
 
-	/** Spawns Targeting actor and waits for it to return valid data or to be cancelled. */
+	/** Spawns Targeting actor and waits for it to return valid data or to be canceled. */
 	UFUNCTION(BlueprintCallable, meta=(HidePin = "WorldContextObject", DefaultToSelf = "WorldContextObject", BlueprintInternalUseOnly = "true", HideSpawnParms="Instigator"), Category="Ability|Tasks")
 	static UAbilityTask_WaitTargetData* WaitTargetData(UObject* WorldContextObject, FName TaskInstanceName, TEnumAsByte<EGameplayTargetingConfirmation::Type> ConfirmationType, TSubclassOf<AGameplayAbilityTargetActor> Class);
 
