replication fixes and suppport for UAbilitySystemComponent::DisplayDebug

diff --git a/Engine/Source/Runtime/GameplayAbilities/Private/AbilitySystemComponent.cpp b/Engine/Source/Runtime/GameplayAbilities/Private/AbilitySystemComponent.cpp
--- a/Engine/Source/Runtime/GameplayAbilities/Private/AbilitySystemComponent.cpp
+++ b/Engine/Source/Runtime/GameplayAbilities/Private/AbilitySystemComponent.cpp
@@ -880,4 +880,72 @@ void UAbilitySystemComponent::OnAttributeAggregatorDirty(FAggregator* Aggregator
 	ActiveGameplayEffects.OnAttributeAggregatorDirty(Aggregator, Attribute);
 }
 
+void UAbilitySystemComponent::DisplayDebug(class UCanvas* Canvas, const class FDebugDisplayInfo& DebugDisplay, float& YL, float& YPos)
+{
+	FGameplayTagContainer OwnerTags;
+	GetOwnedGameplayTags(OwnerTags);
+
+	Canvas->SetDrawColor(FColor::White);
+	YPos += Canvas->DrawText(GEngine->GetTinyFont(), FString::Printf(TEXT("Owned Tags: %s"), *OwnerTags.ToStringSimple()), 4.f, YPos);
+	YPos += YL;
+
+	TSet<FGameplayAttribute> DrawAttributes;
+
+	for (auto It = ActiveGameplayEffects.AttributeAggregatorMap.CreateConstIterator(); It; ++It)
+	{
+		FGameplayAttribute Attribute = It.Key();
+		const FAggregatorRef& AggregatorRef = It.Value();
+		if(AggregatorRef.Get())
+		{
+			FAggregator& Aggregator = *AggregatorRef.Get();
+
+			Canvas->SetDrawColor(FColor::White);
+			YPos += Canvas->DrawText(GEngine->GetTinyFont(), FString::Printf(TEXT("%s %.2f"), *Attribute.GetName(), GetNumericAttribute(Attribute) ), 4.f, YPos);
+
+			DrawAttributes.Add(Attribute);
+
+			for (int32 ModOpIdx = 0; ModOpIdx < ARRAY_COUNT(Aggregator.Mods); ++ModOpIdx)
+			{
+				for (const FAggregatorMod& Mod : Aggregator.Mods[ModOpIdx])
+				{
+					FAggregatorEvaluateParameters EmptyParams;
+					bool IsActivelyModifyingAttribute = Mod.Qualifies(EmptyParams);
+					Canvas->SetDrawColor(IsActivelyModifyingAttribute ? FColor::Yellow : FColor(128, 128, 128));
+
+					FActiveGameplayEffect* ActiveGE = ActiveGameplayEffects.GetActiveGameplayEffect(Mod.ActiveHandle);
+					FString SrcName = ActiveGE ? ActiveGE->Spec.Def->GetName() : FString(TEXT(""));
+
+					if (IsActivelyModifyingAttribute == false)
+					{
+						if (Mod.SourceTagReqs) SrcName += FString::Printf(TEXT(" SourceTags: [%s] "), *Mod.SourceTagReqs->ToString());
+						if (Mod.TargetTagReqs) SrcName += FString::Printf(TEXT("TargetTags: [%s]"), *Mod.TargetTagReqs->ToString());
+					}
+
+					YPos += Canvas->DrawText(GEngine->GetTinyFont(), FString::Printf(TEXT("   %s\t %.2f - %s"), *EGameplayModOpToString(ModOpIdx), Mod.EvaluatedMagnitude, *SrcName), 7.f, YPos);
+
+				}
+			}
+			YPos += YL;
+		}
+	}
+
+	Canvas->SetDrawColor(FColor::White);
+	for (UAttributeSet* Set : SpawnedAttributes)
+	{
+		for (TFieldIterator<UProperty> It(Set->GetClass()); It; ++It)
+		{
+			FGameplayAttribute	Attribute(*It);
+
+			if(DrawAttributes.Contains(Attribute))
+				continue;
+
+			if (Attribute.IsValid())
+			{
+				float Value = GetNumericAttribute(Attribute);
+				YPos += Canvas->DrawText(GEngine->GetTinyFont(), FString::Printf(TEXT("%s %.2f"), *Attribute.GetName(), Value ), 4.f, YPos);
+			}
+		}
+	}
+}
+
 #undef LOCTEXT_NAMESPACE
diff --git a/Engine/Source/Runtime/GameplayAbilities/Private/AttributeSet.cpp b/Engine/Source/Runtime/GameplayAbilities/Private/AttributeSet.cpp
--- a/Engine/Source/Runtime/GameplayAbilities/Private/AttributeSet.cpp
+++ b/Engine/Source/Runtime/GameplayAbilities/Private/AttributeSet.cpp
@@ -8,6 +8,12 @@
 
 #include "ComponentReregisterContext.h"	
 
+FGameplayAttribute::FGameplayAttribute(UProperty *NewProperty)
+{
+	// Only numeric properties are allowed right now
+	Attribute = Cast<UNumericProperty>(NewProperty);
+}
+
 void FGameplayAttribute::SetNumericValueChecked(const float NewValue, class UAttributeSet* Dest) const
 {
 	UNumericProperty *NumericProperty = CastChecked<UNumericProperty>(Attribute);
diff --git a/Engine/Source/Runtime/GameplayAbilities/Private/GameplayEffect.cpp b/Engine/Source/Runtime/GameplayAbilities/Private/GameplayEffect.cpp
--- a/Engine/Source/Runtime/GameplayAbilities/Private/GameplayEffect.cpp
+++ b/Engine/Source/Runtime/GameplayAbilities/Private/GameplayEffect.cpp
@@ -700,6 +700,9 @@ void FActiveGameplayEffect::PostReplicatedAdd(const struct FActiveGameplayEffect
 		}
 	}
 
+	// Handles are not replicated, so create a new one.
+	Handle = FActiveGameplayEffectHandle::GenerateNewHandle(InArray.Owner);
+
 	const_cast<FActiveGameplayEffectsContainer&>(InArray).InternalOnActiveGameplayEffectAdded(*this);	// Const cast is ok. It is there to prevent mutation of the GameplayEffects array, which this wont do.
 
 	static const int32 MAX_DELTA_TIME = 3;
@@ -1094,6 +1097,7 @@ void FActiveGameplayEffectsContainer::StacksNeedToRecalculate()
 	}
 }
 
+/**	This Creates a new ActiveGameplayEffect and is only called on the Server or in cases where a client predictively applies a GameplayEffect. */
 FActiveGameplayEffect& FActiveGameplayEffectsContainer::CreateNewActiveGameplayEffect(const FGameplayEffectSpec &Spec, FPredictionKey InPredictionKey)
 {
 	SCOPE_CYCLE_COUNTER(STAT_CreateNewActiveGameplayEffect);
@@ -1101,6 +1105,10 @@ FActiveGameplayEffect& FActiveGameplayEffectsContainer::CreateNewActiveGameplayE
 	FActiveGameplayEffectHandle NewHandle = FActiveGameplayEffectHandle::GenerateNewHandle(Owner);
 	FActiveGameplayEffect& NewEffect = *new (GameplayEffects)FActiveGameplayEffect(NewHandle, Spec, GetWorldTime(), GetGameStateTime(), InPredictionKey);
 
+	// Calc all of our modifier magnitudes now. Some may need to update later based on attributes changing, etc, but those should
+	// be done through delegate callbacks.
+	NewEffect.Spec.CalculateModifierMagnitudes();
+
 	// Calculate Duration mods if we have a real duration
 	float DurationBaseValue = NewEffect.Spec.GetDuration();
 	if (DurationBaseValue > 0.f)
@@ -1194,6 +1202,7 @@ FActiveGameplayEffect& FActiveGameplayEffectsContainer::CreateNewActiveGameplayE
 	return NewEffect;
 }
 
+/** This is called anytime a new ActiveGameplayEffect is added, on both client and server in all cases */
 void FActiveGameplayEffectsContainer::InternalOnActiveGameplayEffectAdded(FActiveGameplayEffect& Effect)
 {
 	// Add our ongoing tag requirements to the dependency map. We will actually check for these tags below.
@@ -1217,22 +1226,18 @@ void FActiveGameplayEffectsContainer::InternalOnActiveGameplayEffectAdded(FActiv
 	// Check if we should actually be turned on or not (this will turn us on for the first time)
 	FGameplayTagContainer OwnerTags;
 	Owner->GetOwnedGameplayTags(OwnerTags);
-
-	// Calc all of our modifier magnitudes now. Some may need to update later based on attributes changing, etc, but those should
-	// be done through delegate callbacks.
-	Effect.Spec.CalculateModifierMagnitudes();
-
-	// @todo: Probably do the initial mod spec magnitude calculations right here, but only on the server
-	// @todo: Shouldn't the rest of the stuff about dependencies only run on the server?
+	
+	Effect.IsInhibited = true; // Effect has to start inhibited, if it should be uninhibited, CheckOnGoingTagRequirements will handle that state change
 	Effect.CheckOngoingTagRequirements(OwnerTags, *this);
 }
 
+/** Called on server to remove a GameplayEffect */
 bool FActiveGameplayEffectsContainer::RemoveActiveGameplayEffect(FActiveGameplayEffectHandle Handle)
 {
 	// Could make this a map for quicker lookup
-	for (int32 Idx = 0; Idx < GameplayEffects.Num(); ++Idx)
+	for(int32 Idx = 0; Idx < GameplayEffects.Num(); ++Idx)
 	{
-		if (GameplayEffects[Idx].Handle == Handle)
+		if(GameplayEffects[Idx].Handle == Handle)
 		{
 			return InternalRemoveActiveGameplayEffect(Idx);
 		}
@@ -1242,6 +1247,7 @@ bool FActiveGameplayEffectsContainer::RemoveActiveGameplayEffect(FActiveGameplay
 	return false;
 }
 
+/** Called by server to actually remove a GameplayEffect */
 bool FActiveGameplayEffectsContainer::InternalRemoveActiveGameplayEffect(int32 Idx)
 {
 	SCOPE_CYCLE_COUNTER(STAT_RemoveActiveGameplayEffect);
@@ -1276,49 +1282,7 @@ bool FActiveGameplayEffectsContainer::InternalRemoveActiveGameplayEffect(int32 I
 			Owner->GetWorld()->GetTimerManager().ClearTimer(GameplayEffects[Idx].PeriodHandle);
 		}
 
-		// Remove our tag requirements from the dependancy map
-		for (const FGameplayTag& Tag : Effect.Spec.Def->OngoingTagRequirements.IgnoreTags)
-		{
-			auto Ptr = ActiveEffectTagDependencies.Find(Tag);
-			if (Ptr)
-			{
-				Ptr->Remove(Effect.Handle);
-				if (Ptr->Num() <= 0)
-				{
-					ActiveEffectTagDependencies.Remove(Tag);
-				}
-			}
-		}
-
-		for (const FGameplayTag& Tag : Effect.Spec.Def->OngoingTagRequirements.RequireTags)
-		{
-			auto Ptr = ActiveEffectTagDependencies.Find(Tag);
-			if (Ptr)
-			{
-				Ptr->Remove(Effect.Handle);
-				if (Ptr->Num() <= 0)
-				{
-					ActiveEffectTagDependencies.Remove(Tag);
-				}
-			}
-		}
-
-
-		// Fixme: Yuck?
-		for (const FGameplayModifierInfo& Mod : Effect.Spec.Def->Modifiers)
-		{
-			if (Mod.Attribute.IsValid())
-			{
-				FAggregatorRef* RefPtr = AttributeAggregatorMap.Find(Mod.Attribute);
-				if (RefPtr)
-				{
-					RefPtr->Get()->RemoveMod(Effect.Handle);
-				}
-			}
-		}
-
-
-
+		// Finally remove the ActiveGameplayEffect
 		GameplayEffects.RemoveAtSwap(Idx);
 		MarkArrayDirty();
 
@@ -1328,7 +1292,6 @@ bool FActiveGameplayEffectsContainer::InternalRemoveActiveGameplayEffect(int32 I
 		Owner->GetOwner()->ForceNetUpdate();
 
 		StacksNeedToRecalculate();
-
 		return true;
 	}
 
@@ -1336,11 +1299,28 @@ bool FActiveGameplayEffectsContainer::InternalRemoveActiveGameplayEffect(int32 I
 	return false;
 }
 
-/** This does cleanup that has to happen whether the effect is being removed locally or due to replication */
+/** Called by client and server: This does cleanup that has to happen whether the effect is being removed locally or due to replication */
 void FActiveGameplayEffectsContainer::InternalOnActiveGameplayEffectRemoved(const FActiveGameplayEffect& Effect)
 {
 	Effect.OnRemovedDelegate.Broadcast();
 
+	// Update AttributeAggregators: remove mods from this ActiveGE Handle
+	for(const FGameplayModifierInfo& Mod : Effect.Spec.Def->Modifiers)
+	{
+		if(Mod.Attribute.IsValid())
+		{
+			FAggregatorRef* RefPtr = AttributeAggregatorMap.Find(Mod.Attribute);
+			if(RefPtr)
+			{
+				RefPtr->Get()->RemoveMod(Effect.Handle);
+			}
+		}
+	}
+
+	// Remove our tag requirements from the dependancy map
+	RemoveActiveEffectTagDependancy(Effect.Spec.Def->OngoingTagRequirements.IgnoreTags, Effect.Handle);
+	RemoveActiveEffectTagDependancy(Effect.Spec.Def->OngoingTagRequirements.RequireTags, Effect.Handle);
+
 	if (Effect.Spec.Def)
 	{
 		// Update gameplaytag count and broadcast delegate if we are at 0
@@ -1358,6 +1338,22 @@ void FActiveGameplayEffectsContainer::InternalOnActiveGameplayEffectRemoved(cons
 	}
 }
 
+void FActiveGameplayEffectsContainer::RemoveActiveEffectTagDependancy(const FGameplayTagContainer& Tags, FActiveGameplayEffectHandle Handle)
+{
+	for(const FGameplayTag& Tag : Tags)
+	{
+		auto Ptr = ActiveEffectTagDependencies.Find(Tag);
+		if (Ptr)
+		{
+			Ptr->Remove(Handle);
+			if (Ptr->Num() <= 0)
+			{
+				ActiveEffectTagDependencies.Remove(Tag);
+			}
+		}
+	}
+}
+
 void FActiveGameplayEffectsContainer::OnOwnerTagChange(FGameplayTag TagChange, int32 NewCount)
 {
 	// It may be beneficial to do a scoped lock on attribute reeevalluation during this dunction
diff --git a/Engine/Source/Runtime/GameplayAbilities/Private/GameplayEffectTypes.cpp b/Engine/Source/Runtime/GameplayAbilities/Private/GameplayEffectTypes.cpp
--- a/Engine/Source/Runtime/GameplayAbilities/Private/GameplayEffectTypes.cpp
+++ b/Engine/Source/Runtime/GameplayAbilities/Private/GameplayEffectTypes.cpp
@@ -164,7 +164,9 @@ bool FGameplayEffectContextHandle::NetSerialize(FArchive& Ar, class UPackageMap*
 FString EGameplayModOpToString(int32 Type)
 {
 	static UEnum *e = FindObject<UEnum>(ANY_PACKAGE, TEXT("EGameplayModOp"));
-	return e->GetEnum(Type).ToString();
+	FString Right;
+	e->GetEnum(Type).ToString().Split(TEXT("::"), nullptr, &Right);
+	return Right;
 }
 
 FString EGameplayModEffectToString(int32 Type)
@@ -307,6 +309,22 @@ bool FGameplayTagRequirements::IsEmpty() const
 	return (RequireTags.Num() == 0 && IgnoreTags.Num() == 0);
 }
 
+FString FGameplayTagRequirements::ToString() const
+{
+	FString Str;
+
+	if (RequireTags.Num() > 0)
+	{
+		Str += FString::Printf(TEXT("require: %s "), *RequireTags.ToStringSimple());
+	}
+	if (IgnoreTags.Num() >0)
+	{
+		Str += FString::Printf(TEXT("ignore: %s "), *IgnoreTags.ToStringSimple());
+	}
+
+	return Str;
+}
+
 void FActiveGameplayEffectsContainer::PrintAllGameplayEffects() const
 {
 	ABILITY_LOG_SCOPE(TEXT("ActiveGameplayEffects. Num: %d"), GameplayEffects.Num());
diff --git a/Engine/Source/Runtime/GameplayAbilities/Public/AbilitySystemComponent.h b/Engine/Source/Runtime/GameplayAbilities/Public/AbilitySystemComponent.h
--- a/Engine/Source/Runtime/GameplayAbilities/Public/AbilitySystemComponent.h
+++ b/Engine/Source/Runtime/GameplayAbilities/Public/AbilitySystemComponent.h
@@ -110,6 +110,8 @@ class GAMEPLAYABILITIES_API UAbilitySystemComponent : public UActorComponent, pu
 	void SetNumericAttribute(const FGameplayAttribute &Attribute, float NewFloatValue);
 	float GetNumericAttribute(const FGameplayAttribute &Attribute);
 
+	virtual void DisplayDebug(class UCanvas* Canvas, const class FDebugDisplayInfo& DebugDisplay, float& YL, float& YPos);
+
 	// -- Replication -------------------------------------------------------------------------------------------------
 
 	virtual bool ReplicateSubobjects(class UActorChannel *Channel, class FOutBunch *Bunch, FReplicationFlags *RepFlags) override;
diff --git a/Engine/Source/Runtime/GameplayAbilities/Public/AttributeSet.h b/Engine/Source/Runtime/GameplayAbilities/Public/AttributeSet.h
--- a/Engine/Source/Runtime/GameplayAbilities/Public/AttributeSet.h
+++ b/Engine/Source/Runtime/GameplayAbilities/Public/AttributeSet.h
@@ -13,18 +13,15 @@ struct GAMEPLAYABILITIES_API FGameplayAttribute
 	GENERATED_USTRUCT_BODY()
 
 	FGameplayAttribute()
-		: Attribute(NULL)
+		: Attribute(nullptr)
 	{
 	}
 
-	FGameplayAttribute(UProperty *NewProperty)
-		: Attribute(NewProperty)
-	{
-	}
+	FGameplayAttribute(UProperty *NewProperty);
 
 	bool IsValid() const
 	{
-		return Attribute != NULL;
+		return Attribute != nullptr;
 	}
 
 	void SetUProperty(UProperty *NewProperty)
diff --git a/Engine/Source/Runtime/GameplayAbilities/Public/GameplayEffect.h b/Engine/Source/Runtime/GameplayAbilities/Public/GameplayEffect.h
--- a/Engine/Source/Runtime/GameplayAbilities/Public/GameplayEffect.h
+++ b/Engine/Source/Runtime/GameplayAbilities/Public/GameplayEffect.h
@@ -1108,6 +1108,9 @@ private:
 	void InternalOnActiveGameplayEffectAdded(FActiveGameplayEffect& Effect);
 	void InternalOnActiveGameplayEffectRemoved(const FActiveGameplayEffect& Effect);
 
+	/** Updates tag dependancy map when a GameplayEffect is removed */
+	void RemoveActiveEffectTagDependancy(const FGameplayTagContainer& Tags, FActiveGameplayEffectHandle Handle);
+
 	// -------------------------------------------------------------------------------------------
 
 	TMap<FGameplayAttribute, FAggregatorRef>		AttributeAggregatorMap;
diff --git a/Engine/Source/Runtime/GameplayAbilities/Public/GameplayEffectAggregator.h b/Engine/Source/Runtime/GameplayAbilities/Public/GameplayEffectAggregator.h
--- a/Engine/Source/Runtime/GameplayAbilities/Public/GameplayEffectAggregator.h
+++ b/Engine/Source/Runtime/GameplayAbilities/Public/GameplayEffectAggregator.h
@@ -75,6 +75,7 @@ private:
 	void RemoveModsWithActiveHandle(TArray<FAggregatorMod>& Mods, FActiveGameplayEffectHandle ActiveHandle);
 
 	friend struct FAggregator;
+	friend class UAbilitySystemComponent;	// Only needed for DisplayDebug()
 };
 
 struct GAMEPLAYABILITIES_API FAggregatorRef
diff --git a/Engine/Source/Runtime/GameplayAbilities/Public/GameplayEffectTypes.h b/Engine/Source/Runtime/GameplayAbilities/Public/GameplayEffectTypes.h
--- a/Engine/Source/Runtime/GameplayAbilities/Public/GameplayEffectTypes.h
+++ b/Engine/Source/Runtime/GameplayAbilities/Public/GameplayEffectTypes.h
@@ -719,6 +719,8 @@ struct FGameplayTagRequirements
 	bool	IsEmpty() const;
 
 	static FGetGameplayTags	SnapshotTags(FGetGameplayTags TagDelegate);
+
+	FString ToString() const;
 };
 
 USTRUCT()
diff --git a/Engine/Source/Runtime/GameplayTags/Classes/GameplayTagContainer.h b/Engine/Source/Runtime/GameplayTags/Classes/GameplayTagContainer.h
--- a/Engine/Source/Runtime/GameplayTags/Classes/GameplayTagContainer.h
+++ b/Engine/Source/Runtime/GameplayTags/Classes/GameplayTagContainer.h
@@ -221,6 +221,9 @@ struct GAMEPLAYTAGS_API FGameplayTagContainer
 	/** Returns human readable Tag list */
 	FString ToString() const;
 
+	/** Returns abbreviated human readable Tag list without parens or property names */
+	FString ToStringSimple() const;
+
 	/** Creates a const iterator for the contents of this array */
 	TArray<FGameplayTag>::TConstIterator CreateConstIterator() const
 	{
diff --git a/Engine/Source/Runtime/GameplayTags/Private/GameplayTagContainer.cpp b/Engine/Source/Runtime/GameplayTags/Private/GameplayTagContainer.cpp
--- a/Engine/Source/Runtime/GameplayTags/Private/GameplayTagContainer.cpp
+++ b/Engine/Source/Runtime/GameplayTags/Private/GameplayTagContainer.cpp
@@ -271,6 +271,22 @@ FString FGameplayTagContainer::ToString() const
 	return RetString;
 }
 
+FString FGameplayTagContainer::ToStringSimple() const
+{
+	FString RetString;
+	for (int i = 0; i < GameplayTags.Num(); ++i)
+	{
+		RetString += TEXT("\"");
+		RetString += GameplayTags[i].ToString();
+		RetString += TEXT("\"");
+		if (i < GameplayTags.Num() - 1)
+		{
+			RetString += TEXT(", ");
+		}
+	}
+	return RetString;
+}
+
 FGameplayTag::FGameplayTag()
 	: TagName(NAME_None)
 {
