additional displaydebug info for active abilities

diff --git a/Engine/Source/Runtime/GameplayAbilities/Private/AbilitySystemComponent.cpp b/Engine/Source/Runtime/GameplayAbilities/Private/AbilitySystemComponent.cpp
--- a/Engine/Source/Runtime/GameplayAbilities/Private/AbilitySystemComponent.cpp
+++ b/Engine/Source/Runtime/GameplayAbilities/Private/AbilitySystemComponent.cpp
@@ -12,6 +12,7 @@
 #include "MessageLog.h"
 #include "UObjectToken.h"
 #include "MapErrors.h"
+#include "DisplayDebugHelpers.h"
 
 DEFINE_LOG_CATEGORY(LogAbilitySystemComponent);
 
@@ -882,6 +883,18 @@ void UAbilitySystemComponent::OnAttributeAggregatorDirty(FAggregator* Aggregator
 
 void UAbilitySystemComponent::DisplayDebug(class UCanvas* Canvas, const class FDebugDisplayInfo& DebugDisplay, float& YL, float& YPos)
 {
+
+	bool bShowAttributes = true;
+	bool bShowGameplayEffects = true;
+	bool bShowAbilities = true;
+
+	if (DebugDisplay.IsDisplayOn(FName(TEXT("Ability"))))
+	{
+		bShowAbilities = true;
+		bShowAttributes = false;
+		bShowGameplayEffects = false;
+	}
+
 	FGameplayTagContainer OwnerTags;
 	GetOwnedGameplayTags(OwnerTags);
 
@@ -891,60 +904,114 @@ void UAbilitySystemComponent::DisplayDebug(class UCanvas* Canvas, const class FD
 
 	TSet<FGameplayAttribute> DrawAttributes;
 
-	for (auto It = ActiveGameplayEffects.AttributeAggregatorMap.CreateConstIterator(); It; ++It)
+	// -------------------------------------------------------------
+
+	if (bShowGameplayEffects)
 	{
-		FGameplayAttribute Attribute = It.Key();
-		const FAggregatorRef& AggregatorRef = It.Value();
-		if(AggregatorRef.Get())
+
+		for (auto It = ActiveGameplayEffects.AttributeAggregatorMap.CreateConstIterator(); It; ++It)
 		{
-			FAggregator& Aggregator = *AggregatorRef.Get();
+			FGameplayAttribute Attribute = It.Key();
+			const FAggregatorRef& AggregatorRef = It.Value();
+			if(AggregatorRef.Get())
+			{
+				FAggregator& Aggregator = *AggregatorRef.Get();
 
-			Canvas->SetDrawColor(FColor::White);
-			YPos += Canvas->DrawText(GEngine->GetTinyFont(), FString::Printf(TEXT("%s %.2f"), *Attribute.GetName(), GetNumericAttribute(Attribute) ), 4.f, YPos);
+				Canvas->SetDrawColor(FColor::White);
+				YPos += Canvas->DrawText(GEngine->GetTinyFont(), FString::Printf(TEXT("%s %.2f"), *Attribute.GetName(), GetNumericAttribute(Attribute) ), 4.f, YPos);
 
-			DrawAttributes.Add(Attribute);
+				DrawAttributes.Add(Attribute);
 
-			for (int32 ModOpIdx = 0; ModOpIdx < ARRAY_COUNT(Aggregator.Mods); ++ModOpIdx)
-			{
-				for (const FAggregatorMod& Mod : Aggregator.Mods[ModOpIdx])
+				for (int32 ModOpIdx = 0; ModOpIdx < ARRAY_COUNT(Aggregator.Mods); ++ModOpIdx)
 				{
-					FAggregatorEvaluateParameters EmptyParams;
-					bool IsActivelyModifyingAttribute = Mod.Qualifies(EmptyParams);
-					Canvas->SetDrawColor(IsActivelyModifyingAttribute ? FColor::Yellow : FColor(128, 128, 128));
+					for (const FAggregatorMod& Mod : Aggregator.Mods[ModOpIdx])
+					{
+						FAggregatorEvaluateParameters EmptyParams;
+						bool IsActivelyModifyingAttribute = Mod.Qualifies(EmptyParams);
+						Canvas->SetDrawColor(IsActivelyModifyingAttribute ? FColor::Yellow : FColor(128, 128, 128));
 
-					FActiveGameplayEffect* ActiveGE = ActiveGameplayEffects.GetActiveGameplayEffect(Mod.ActiveHandle);
-					FString SrcName = ActiveGE ? ActiveGE->Spec.Def->GetName() : FString(TEXT(""));
+						FActiveGameplayEffect* ActiveGE = ActiveGameplayEffects.GetActiveGameplayEffect(Mod.ActiveHandle);
+						FString SrcName = ActiveGE ? ActiveGE->Spec.Def->GetName() : FString(TEXT(""));
 
-					if (IsActivelyModifyingAttribute == false)
-					{
-						if (Mod.SourceTagReqs) SrcName += FString::Printf(TEXT(" SourceTags: [%s] "), *Mod.SourceTagReqs->ToString());
-						if (Mod.TargetTagReqs) SrcName += FString::Printf(TEXT("TargetTags: [%s]"), *Mod.TargetTagReqs->ToString());
-					}
+						if (IsActivelyModifyingAttribute == false)
+						{
+							if (Mod.SourceTagReqs) SrcName += FString::Printf(TEXT(" SourceTags: [%s] "), *Mod.SourceTagReqs->ToString());
+							if (Mod.TargetTagReqs) SrcName += FString::Printf(TEXT("TargetTags: [%s]"), *Mod.TargetTagReqs->ToString());
+						}
 
-					YPos += Canvas->DrawText(GEngine->GetTinyFont(), FString::Printf(TEXT("   %s\t %.2f - %s"), *EGameplayModOpToString(ModOpIdx), Mod.EvaluatedMagnitude, *SrcName), 7.f, YPos);
+						YPos += Canvas->DrawText(GEngine->GetTinyFont(), FString::Printf(TEXT("   %s\t %.2f - %s"), *EGameplayModOpToString(ModOpIdx), Mod.EvaluatedMagnitude, *SrcName), 7.f, YPos);
 
+					}
 				}
+				YPos += YL;
 			}
-			YPos += YL;
 		}
 	}
 
-	Canvas->SetDrawColor(FColor::White);
-	for (UAttributeSet* Set : SpawnedAttributes)
+	// -------------------------------------------------------------
+
+	if (bShowAttributes)
 	{
-		for (TFieldIterator<UProperty> It(Set->GetClass()); It; ++It)
+		Canvas->SetDrawColor(FColor::White);
+		for (UAttributeSet* Set : SpawnedAttributes)
 		{
-			FGameplayAttribute	Attribute(*It);
+			for (TFieldIterator<UProperty> It(Set->GetClass()); It; ++It)
+			{
+				FGameplayAttribute	Attribute(*It);
 
-			if(DrawAttributes.Contains(Attribute))
+				if(DrawAttributes.Contains(Attribute))
+					continue;
+
+				if (Attribute.IsValid())
+				{
+					float Value = GetNumericAttribute(Attribute);
+					YPos += Canvas->DrawText(GEngine->GetTinyFont(), FString::Printf(TEXT("%s %.2f"), *Attribute.GetName(), Value ), 4.f, YPos);
+				}
+			}
+		}
+		YPos += YL;
+	}
+
+	// -------------------------------------------------------------
+
+	if (bShowAbilities)
+	{
+		for (const FGameplayAbilitySpec& AbilitySpec : ActivatableAbilities)
+		{
+			if (AbilitySpec.Ability == nullptr)
 				continue;
 
-			if (Attribute.IsValid())
+			Canvas->SetDrawColor(AbilitySpec.IsActive() ? FColor::Yellow : FColor(128, 128, 128));
+			YPos += Canvas->DrawText(GEngine->GetTinyFont(), FString::Printf(TEXT("%s"), *GetNameSafe(AbilitySpec.Ability)), 4.f, YPos);
+	
+			if (AbilitySpec.IsActive())
 			{
-				float Value = GetNumericAttribute(Attribute);
-				YPos += Canvas->DrawText(GEngine->GetTinyFont(), FString::Printf(TEXT("%s %.2f"), *Attribute.GetName(), Value ), 4.f, YPos);
+				TArray<UGameplayAbility*> Instances = AbilitySpec.GetAbilityInstances();
+				for (int32 InstanceIdx=0; InstanceIdx < Instances.Num(); ++InstanceIdx)
+				{
+					UGameplayAbility* Instance = Instances[InstanceIdx];
+					if (!Instance)
+						continue;
+
+					Canvas->SetDrawColor(FColor::White);
+					for (auto TaskPtr : Instance->ActiveTasks)
+					{
+						UAbilityTask* Task = TaskPtr.Get();
+						if (Task)
+						{
+							YPos += Canvas->DrawText(GEngine->GetTinyFont(), FString::Printf(TEXT("%s"), *Task->GetDebugString()), 7.f, YPos);
+						}
+					}
+
+					if (InstanceIdx < Instances.Num() - 2)
+					{
+						Canvas->SetDrawColor(FColor(128, 128, 128));
+						YPos += Canvas->DrawText(GEngine->GetTinyFont(), FString::Printf(TEXT("--------")), 7.f, YPos);
+					}
+				}
 			}
 		}
+		YPos += YL;
 	}
 }
 
diff --git a/Engine/Source/Runtime/GameplayAbilities/Public/Abilities/GameplayAbilityTypes.h b/Engine/Source/Runtime/GameplayAbilities/Public/Abilities/GameplayAbilityTypes.h
--- a/Engine/Source/Runtime/GameplayAbilities/Public/Abilities/GameplayAbilityTypes.h
+++ b/Engine/Source/Runtime/GameplayAbilities/Public/Abilities/GameplayAbilityTypes.h
@@ -328,7 +328,7 @@ struct GAMEPLAYABILITIES_API FGameplayAbilitySpec
 	UPROPERTY()
 	TArray<UGameplayAbility*> ReplicatedInstances;
 
-	TArray<UGameplayAbility*> GetAbilityInstances()
+	TArray<UGameplayAbility*> GetAbilityInstances() const
 	{
 		TArray<UGameplayAbility*> Abilities;
 		Abilities.Append(ReplicatedInstances);
