Code for new GameplayCueNotify system, and updated (currently gated) HitscanWeaponEffect. Class rediretors included. GCN_Static functionality pending CDO static BP functions.

diff --git a/Engine/Config/BaseEngine.ini b/Engine/Config/BaseEngine.ini
--- a/Engine/Config/BaseEngine.ini
+++ b/Engine/Config/BaseEngine.ini
@@ -168,6 +168,8 @@ bUseStreamingPause=false
 
 ; Class Redirects for UE4  (e.g. all of the classes we have renamed that are in maps that we still want to utilize)
 ; Note: For class name redirects, the OldClassName must be the plain OldClassName, it cannot be OldPackage.OldClassName
++ActiveClassRedirects=(OldClassName="GameplayCueNotify",NewClassName="GameplayCueNotify_Static")
++ActiveClassRedirects=(OldClassName="GameplayCueNotify_Blueprint",NewClassName="GameplayCueNotify_Actor")
 +ActiveClassRedirects=(OldClassName="CylinderComponent",NewClassName="CapsuleComponent")
 +ActiveClassRedirects=(OldClassName="ComponentProperty",NewClassName="ObjectProperty")
 +ActiveClassRedirects=(OldClassName="WeakComponentProperty",NewClassName="WeakObjectProperty")
diff --git a/Engine/Source/Runtime/GameplayAbilities/Private/GameplayCueManager.cpp b/Engine/Source/Runtime/GameplayAbilities/Private/GameplayCueManager.cpp
--- a/Engine/Source/Runtime/GameplayAbilities/Private/GameplayCueManager.cpp
+++ b/Engine/Source/Runtime/GameplayAbilities/Private/GameplayCueManager.cpp
@@ -5,6 +5,7 @@
 #include "GameplayCueInterface.h"
 #include "GameplayCueManager.h"
 #include "GameplayTagsModule.h"
+#include "GameplayCueNotify_Static.h"
 
 #if WITH_EDITOR
 #include "UnrealEd.h"
@@ -76,8 +77,12 @@ void UGameplayCueManager::HandleGameplayCueNotify_Internal(AActor* TargetActor,
 			}
 			else
 			{
-				// Found it - did we load a data asset (UGameplayCueNotify) or a Blueprint (UBlueprint)
-				CueData.LoadedGameplayCueNotify = Cast<UGameplayCueNotify>(FoundObject);
+				// Found it - did we load a non-instance version (UGameplayCueNotify_Static) or an instanced version (AGameplayCueNotify_Actor)
+				CueData.LoadedGameplayCueNotify = Cast<AGameplayCueNotify_Actor>(FoundObject);
+				if (!CueData.LoadedGameplayCueNotify)
+				{
+					CueData.LoadedGameplayCueNotify = Cast<UGameplayCueNotify_Static>(FoundObject);
+				}
 				if (CueData.LoadedGameplayCueNotify == nullptr)
 				{
 					// Not a dataasset - maybe a blueprint
@@ -85,48 +90,75 @@ void UGameplayCueManager::HandleGameplayCueNotify_Internal(AActor* TargetActor,
 					UBlueprint* GameplayCueBlueprint = Cast<UBlueprint>(FoundObject);
 					if (GameplayCueBlueprint && GameplayCueBlueprint->GeneratedClass)
 					{
-						CueData.LoadedGameplayCueNotify = Cast<UGameplayCueNotify>(GameplayCueBlueprint->GeneratedClass->ClassDefaultObject);
+						CueData.LoadedGameplayCueNotify = Cast<AGameplayCueNotify_Actor>(GameplayCueBlueprint->GeneratedClass->ClassDefaultObject);
+						if (!CueData.LoadedGameplayCueNotify)
+						{
+							CueData.LoadedGameplayCueNotify = Cast<UGameplayCueNotify_Static>(GameplayCueBlueprint->GeneratedClass->ClassDefaultObject);
+						}
+						else
+						{
+							ABILITY_LOG(Warning, TEXT("GameplayCueNotify %s loaded blueprint object %s that is not a GameplayCueNotify"), *CueData.GameplayCueNotifyObj.ToString(), *FoundObject->GetName());
+							return;
+						}
 					}
 					else
 					{
-						ABILITY_LOG(Warning, TEXT("GameplayCueNotify %s loaded object %s that is not a UGameplayCueNotify or UBlueprint"), *CueData.GameplayCueNotifyObj.ToString(), *FoundObject->GetName());
+						ABILITY_LOG(Warning, TEXT("GameplayCueNotify %s loaded object %s that is not a GameplayCueNotify"), *CueData.GameplayCueNotifyObj.ToString(), *FoundObject->GetName());
 						return;
 					}
-					
 				}
 			}
 		}
 
 		// Handle the Notify if we found something
-		if (CueData.LoadedGameplayCueNotify)
+		if (UGameplayCueNotify_Static* NonInstancedCue = Cast<UGameplayCueNotify_Static>(CueData.LoadedGameplayCueNotify))
 		{
-			UGameplayCueNotify* SpawnedInstanced = nullptr;
-			if (auto InnerMap = NotifyMap.Find(TargetActor))
+			if (NonInstancedCue->HandlesEvent(EventType))
+			{
+				NonInstancedCue->HandleGameplayCue(TargetActor, EventType, Parameters);
+				if (!NonInstancedCue->IsOverride)
+				{
+					HandleGameplayCueNotify_Internal(TargetActor, CueData.ParentDataIdx, EventType, Parameters);
+				}
+			}
+			else
+			{
+				//Didn't even handle it, so IsOverride should not apply.
+				HandleGameplayCueNotify_Internal(TargetActor, CueData.ParentDataIdx, EventType, Parameters);
+			}
+		}
+		else if (AGameplayCueNotify_Actor* InstancedCue = Cast<AGameplayCueNotify_Actor>(CueData.LoadedGameplayCueNotify))
+		{
+			AGameplayCueNotify_Actor* SpawnedInstancedCue = nullptr;
+			if (auto InnerMap = NotifyMapActor.Find(TargetActor))
 			{
 				if (auto WeakPtrPtr = InnerMap->Find(DataIdx))
 				{
-					SpawnedInstanced = WeakPtrPtr->Get();
+					SpawnedInstancedCue = WeakPtrPtr->Get();
 				}
 			}
 
-			if (SpawnedInstanced == nullptr && CueData.LoadedGameplayCueNotify->NeedsInstanceForEvent(EventType))
+			//Get our instance. We should probably have a flag or something to determine if we want to reuse or stack instances. That would mean changing our map to have a list of active instances.
+			if (SpawnedInstancedCue == nullptr)
 			{
 				// We don't have an instance for this, and we need one, so make one
-				SpawnedInstanced = static_cast<UGameplayCueNotify*>(StaticDuplicateObject(CueData.LoadedGameplayCueNotify, this, TEXT("None"), ~RF_RootSet));
-
-				auto& InnerMap = NotifyMap.FindOrAdd(TargetActor);
-
-				InnerMap.Add(DataIdx) = SpawnedInstanced;
-
-				// Fixme: this will grow unbounded
-				InstantiatedObjects.Add(SpawnedInstanced);
+				//SpawnedInstancedCue = static_cast<AGameplayCueNotify_Actor*>(StaticDuplicateObject(InstancedCue, this, TEXT("None"), ~RF_RootSet));
+				SpawnedInstancedCue = TargetActor->GetWorld()->SpawnActor<AGameplayCueNotify_Actor>(InstancedCue->GetActorClass(), TargetActor->GetActorLocation(), TargetActor->GetActorRotation());
+				auto& InnerMap = NotifyMapActor.FindOrAdd(TargetActor);
+				InnerMap.Add(DataIdx) = SpawnedInstancedCue;
 			}
-			
-			UGameplayCueNotify* NotifyToExecute = SpawnedInstanced ? SpawnedInstanced : CueData.LoadedGameplayCueNotify;
-			NotifyToExecute->HandleGameplayCue(TargetActor, EventType, Parameters);
-
-			if (NotifyToExecute->IsOverride == false)
+			check(SpawnedInstancedCue);
+			if (SpawnedInstancedCue->HandlesEvent(EventType))
+			{
+				SpawnedInstancedCue->HandleGameplayCue(TargetActor, EventType, Parameters);
+				if (!SpawnedInstancedCue->IsOverride)
+				{
+					HandleGameplayCueNotify_Internal(TargetActor, CueData.ParentDataIdx, EventType, Parameters);
+				}
+			}
+			else
 			{
+				//Didn't even handle it, so IsOverride should not apply.
 				HandleGameplayCueNotify_Internal(TargetActor, CueData.ParentDataIdx, EventType, Parameters);
 			}
 		}
@@ -135,12 +167,12 @@ void UGameplayCueManager::HandleGameplayCueNotify_Internal(AActor* TargetActor,
 
 void UGameplayCueManager::EndGameplayCuesFor(AActor* TargetActor)
 {
-	TMap<int32, TWeakObjectPtr<UGameplayCueNotify> > FoundMap;	
-	if (NotifyMap.RemoveAndCopyValue(TargetActor, FoundMap))
+	TMap<int32, TWeakObjectPtr<AGameplayCueNotify_Actor>> FoundMapActor;
+	if (NotifyMapActor.RemoveAndCopyValue(TargetActor, FoundMapActor))
 	{
-		for (auto It = FoundMap.CreateConstIterator(); It; ++It)
+		for (auto It = FoundMapActor.CreateConstIterator(); It; ++It)
 		{
-			UGameplayCueNotify* InstancedCue =  It.Value().Get();
+			AGameplayCueNotify_Actor* InstancedCue =  It.Value().Get();
 			if (InstancedCue)
 			{
 				InstancedCue->OnOwnerDestroyed();
@@ -153,9 +185,13 @@ void UGameplayCueManager::EndGameplayCuesFor(AActor* TargetActor)
 
 void UGameplayCueManager::LoadObjectLibraryFromPaths(const TArray<FString>& InPaths, bool InFullyLoad)
 {
-	if (!GameplayCueNotifyObjectLibrary)
+	if (!GameplayCueNotifyActorObjectLibrary)
 	{
-		GameplayCueNotifyObjectLibrary = UObjectLibrary::CreateLibrary(UGameplayCueNotify::StaticClass(), true, GIsEditor);
+		GameplayCueNotifyActorObjectLibrary = UObjectLibrary::CreateLibrary(AGameplayCueNotify_Actor::StaticClass(), true, GIsEditor);
+	}
+	if (!GameplayCueNotifyStaticObjectLibrary)
+	{
+		GameplayCueNotifyStaticObjectLibrary = UObjectLibrary::CreateLibrary(UGameplayCueNotify_Static::StaticClass(), true, GIsEditor);
 	}
 
 	LoadedPaths = InPaths;
@@ -196,9 +232,9 @@ void UGameplayCueManager::LoadObjectLibrary_Internal()
 	SlowTask.MakeDialog();
 #endif
 
-	FScopeCycleCounterUObject PreloadScope(GameplayCueNotifyObjectLibrary);
-	
-	GameplayCueNotifyObjectLibrary->LoadBlueprintAssetDataFromPaths(LoadedPaths);
+	FScopeCycleCounterUObject PreloadScopeActor(GameplayCueNotifyActorObjectLibrary);
+	GameplayCueNotifyActorObjectLibrary->LoadBlueprintAssetDataFromPaths(LoadedPaths);
+	GameplayCueNotifyStaticObjectLibrary->LoadBlueprintAssetDataFromPaths(LoadedPaths);		//No separate cycle counter for this.
 
 	if (bFullyLoad)
 	{
@@ -206,24 +242,28 @@ void UGameplayCueManager::LoadObjectLibrary_Internal()
 		FString PerfMessage = FString::Printf(TEXT("Fully Loaded GameplayCueNotify object library"));
 		SCOPE_LOG_TIME_IN_SECONDS(*PerfMessage, nullptr)
 #endif
-		GameplayCueNotifyObjectLibrary->LoadAssetsFromAssetData();
+		GameplayCueNotifyActorObjectLibrary->LoadAssetsFromAssetData();
+		GameplayCueNotifyStaticObjectLibrary->LoadAssetsFromAssetData();
 	}
 
 	// ---------------------------------------------------------
 	// Look for GameplayCueNotifies that handle events
 	// ---------------------------------------------------------
 	
-	TArray<FAssetData> AssetDatas;
-	GameplayCueNotifyObjectLibrary->GetAssetDataList(AssetDatas);
+	TArray<FAssetData> ActorAssetDatas;
+	GameplayCueNotifyActorObjectLibrary->GetAssetDataList(ActorAssetDatas);
+
+	TArray<FAssetData> StaticAssetDatas;
+	GameplayCueNotifyStaticObjectLibrary->GetAssetDataList(StaticAssetDatas);
 
 	IGameplayTagsModule& GameplayTagsModule = IGameplayTagsModule::Get();
 
 	GameplayCueData.Empty();
 	GameplayCueDataMap.Empty();
 
-	for (FAssetData Data: AssetDatas)
+	for (FAssetData Data: ActorAssetDatas)
 	{
-		const FString* FoundGameplayTag = Data.TagsAndValues.Find(GET_MEMBER_NAME_CHECKED(UGameplayCueNotify, GameplayCueName));
+		const FString* FoundGameplayTag = Data.TagsAndValues.Find(GET_MEMBER_NAME_CHECKED(AGameplayCueNotify_Actor, GameplayCueName));
 		if (FoundGameplayTag && FoundGameplayTag->Equals(TEXT("None")) == false)
 		{
 			ABILITY_LOG(Warning, TEXT("Found: %s"), **FoundGameplayTag);
@@ -234,9 +274,28 @@ void UGameplayCueManager::LoadObjectLibrary_Internal()
 				// Add a new NotifyData entry to our flat list for this one
 				FStringAssetReference StringRef;
 				StringRef.AssetLongPathname = Data.ObjectPath.ToString();
+				AddGameplayCueData_Internal(GameplayCueTag, StringRef);				
+			}
+			else
+			{
+				ABILITY_LOG(Warning, TEXT("Found GameplayCue tag %s in asset %s but there is no corresponding tag in the GameplayTagMAnager."), **FoundGameplayTag, *Data.PackageName.ToString());
+			}
+		}
+	}
+	for (FAssetData Data : StaticAssetDatas)
+	{
+		const FString* FoundGameplayTag = Data.TagsAndValues.Find(GET_MEMBER_NAME_CHECKED(UGameplayCueNotify_Static, GameplayCueName));
+		if (FoundGameplayTag && FoundGameplayTag->Equals(TEXT("None")) == false)
+		{
+			ABILITY_LOG(Warning, TEXT("Found: %s"), **FoundGameplayTag);
 
+			FGameplayTag  GameplayCueTag = GameplayTagsModule.GetGameplayTagsManager().RequestGameplayTag(FName(**FoundGameplayTag), false);
+			if (GameplayCueTag.IsValid())
+			{
+				// Add a new NotifyData entry to our flat list for this one
+				FStringAssetReference StringRef;
+				StringRef.AssetLongPathname = Data.ObjectPath.ToString();
 				AddGameplayCueData_Internal(GameplayCueTag, StringRef);
-				
 			}
 			else
 			{
@@ -318,26 +377,23 @@ void UGameplayCueManager::HandleAssetAdded(UObject *Object)
 	UBlueprint* Blueprint = Cast<UBlueprint>(Object);
 	if (Blueprint && Blueprint->GeneratedClass)
 	{
-		UGameplayCueNotify* CDO = Cast<UGameplayCueNotify>(Blueprint->GeneratedClass->ClassDefaultObject);
-		if (CDO && CDO->GameplayCueTag.IsValid())
+		UGameplayCueNotify_Static* StaticCDO = Cast<UGameplayCueNotify_Static>(Blueprint->GeneratedClass->ClassDefaultObject);
+		FStringAssetReference StringRef;
+		StringRef.AssetLongPathname = Blueprint->GetPathName();
+		if (StaticCDO && StaticCDO->GameplayCueTag.IsValid())
 		{
-			FStringAssetReference StringRef;
-			StringRef.AssetLongPathname = Blueprint->GetPathName();
-
-			AddGameplayCueData_Internal(CDO->GameplayCueTag, StringRef);
-			BuildAccelerationMap_Internal();
+			AddGameplayCueData_Internal(StaticCDO->GameplayCueTag, StringRef);
+			BuildAccelerationMap_Internal();		//RICKH Could this just be done by setting the acceleration map rebuild bool to true?
+		}
+		else
+		{
+			AGameplayCueNotify_Actor* ActorCDO = Cast<AGameplayCueNotify_Actor>(Blueprint->GeneratedClass->ClassDefaultObject);
+			if (ActorCDO && ActorCDO->GameplayCueTag.IsValid())
+			{
+				AddGameplayCueData_Internal(ActorCDO->GameplayCueTag, StringRef);
+				BuildAccelerationMap_Internal();		//RICKH Could this just be done by setting the acceleration map rebuild bool to true?
+			}
 		}
-	}
-	
-	
-	UGameplayCueNotify* Notify = Cast<UGameplayCueNotify>(Object);
-	if (Notify && Notify->GameplayCueTag.IsValid())
-	{
-		FStringAssetReference StringRef;
-		StringRef.AssetLongPathname = Notify->GetPathName();
-
-		AddGameplayCueData_Internal(Notify->GameplayCueTag, StringRef);
-		BuildAccelerationMap_Internal();
 	}
 }
 
@@ -349,14 +405,14 @@ void UGameplayCueManager::HandleAssetDeleted(UObject *Object)
 	UBlueprint* Blueprint = Cast<UBlueprint>(Object);
 	if (Blueprint && Blueprint->GeneratedClass)
 	{
-		UGameplayCueNotify* CDO = Cast<UGameplayCueNotify>(Blueprint->GeneratedClass->ClassDefaultObject);
+		AGameplayCueNotify_Actor* CDO = Cast<AGameplayCueNotify_Actor>(Blueprint->GeneratedClass->ClassDefaultObject);
 		if (CDO && CDO->GameplayCueTag.IsValid())
 		{
 			TagtoRemove = CDO->GameplayCueTag;
 		}
 	}
 
-	UGameplayCueNotify* Notify = Cast<UGameplayCueNotify>(Object);
+	AGameplayCueNotify_Actor* Notify = Cast<AGameplayCueNotify_Actor>(Object);
 	if (Notify && Notify->GameplayCueTag.IsValid())
 	{
 		TagtoRemove = Notify->GameplayCueTag;
diff --git a/Engine/Source/Runtime/GameplayAbilities/Private/GameplayCueNotify.cpp b/Engine/Source/Runtime/GameplayAbilities/Private/GameplayCueNotify.cpp
deleted file mode 100644
diff --git a/Engine/Source/Runtime/GameplayAbilities/Private/GameplayCueNotify_Actor.cpp b/Engine/Source/Runtime/GameplayAbilities/Private/GameplayCueNotify_Actor.cpp
new file mode 100644
--- /dev/null
+++ b/Engine/Source/Runtime/GameplayAbilities/Private/GameplayCueNotify_Actor.cpp
@@ -0,0 +1,108 @@
+// Copyright 1998-2014 Epic Games, Inc. All Rights Reserved.
+
+#include "AbilitySystemPrivatePCH.h"
+#include "GameplayCueNotify_Actor.h"
+#include "GameplayTagsModule.h"
+#include "GameplayCueManager.h"
+#include "AbilitySystemComponent.h"
+
+AGameplayCueNotify_Actor::AGameplayCueNotify_Actor(const FObjectInitializer& ObjectInitializer)
+: Super(ObjectInitializer)
+{
+	IsOverride = true;
+}
+
+#if WITH_EDITOR
+void AGameplayCueNotify_Actor::PostEditChangeProperty(FPropertyChangedEvent& PropertyChangedEvent)
+{
+	UAbilitySystemGlobals::Get().GetGameplayCueManager()->bAccelerationMapOutdated = true;
+}
+#endif
+
+void AGameplayCueNotify_Actor::DeriveGameplayCueTagFromAssetName()
+{
+	// In the editor, attempt to infer GameplayCueTag from our asset name (if there is no valid GameplayCueTag already).
+	#if WITH_EDITOR
+	if (GIsEditor)
+	{
+		if (GameplayCueTag.IsValid() == false)
+		{
+			FString MyName = GetName();
+
+			MyName.RemoveFromStart(TEXT("Default__"));
+			MyName.RemoveFromEnd(TEXT("_c"));
+
+			MyName.ReplaceInline(TEXT("_"), TEXT("."));
+
+			if (!MyName.Contains(TEXT("GameplayCue")))
+			{
+				MyName = FString(TEXT("GameplayCue.")) + MyName;
+			}
+
+			IGameplayTagsModule& GameplayTagsModule = IGameplayTagsModule::Get();
+			GameplayCueTag = GameplayTagsModule.GetGameplayTagsManager().RequestGameplayTag(FName(*MyName), false);
+		}
+		GameplayCueName = GameplayCueTag.GetTagName();
+	}
+	#endif
+}
+
+void AGameplayCueNotify_Actor::Serialize(FArchive& Ar)
+{
+	if (Ar.IsSaving())
+	{
+		DeriveGameplayCueTagFromAssetName();
+	}
+
+	Super::Serialize(Ar);
+
+	if (Ar.IsLoading())
+	{
+		DeriveGameplayCueTagFromAssetName();
+	}
+}
+
+void AGameplayCueNotify_Actor::PostInitProperties()
+{
+	Super::PostInitProperties();
+	DeriveGameplayCueTagFromAssetName();
+}
+
+bool AGameplayCueNotify_Actor::HandlesEvent(EGameplayCueEvent::Type EventType) const
+{
+	return true;
+}
+
+void AGameplayCueNotify_Actor::HandleGameplayCue(AActor* MyTarget, EGameplayCueEvent::Type EventType, FGameplayCueParameters Parameters)
+{
+	if (MyTarget && !MyTarget->IsPendingKill())
+	{
+		K2_HandleGameplayCue(MyTarget, EventType, Parameters);
+
+		switch (EventType)
+		{
+		case EGameplayCueEvent::OnActive:
+			OnActive(MyTarget, Parameters);
+			break;
+
+		case EGameplayCueEvent::Executed:
+			OnExecute(MyTarget, Parameters);
+			break;
+
+		case EGameplayCueEvent::Removed:
+			OnRemove(MyTarget, Parameters);
+			break;
+		};
+	}
+	else
+	{
+		ABILITY_LOG(Warning, TEXT("Null Target"));
+	}
+}
+
+void AGameplayCueNotify_Actor::OnOwnerDestroyed()
+{
+	// May need to do extra cleanup in child classes
+	Destroy();
+	//MarkPendingKill();
+}
diff --git a/Engine/Source/Runtime/GameplayAbilities/Private/GameplayCueNotify_Blueprint.cpp b/Engine/Source/Runtime/GameplayAbilities/Private/GameplayCueNotify_Blueprint.cpp
deleted file mode 100644
diff --git a/Engine/Source/Runtime/GameplayAbilities/Private/GameplayCueNotify_HitImpact.cpp b/Engine/Source/Runtime/GameplayAbilities/Private/GameplayCueNotify_HitImpact.cpp
--- a/Engine/Source/Runtime/GameplayAbilities/Private/GameplayCueNotify_HitImpact.cpp
+++ b/Engine/Source/Runtime/GameplayAbilities/Private/GameplayCueNotify_HitImpact.cpp
@@ -14,11 +14,6 @@ bool UGameplayCueNotify_HitImpact::HandlesEvent(EGameplayCueEvent::Type EventTyp
 	return (EventType == EGameplayCueEvent::Executed);
 }
 
-bool UGameplayCueNotify_HitImpact::NeedsInstanceForEvent(EGameplayCueEvent::Type EventType) const
-{
-	return false;
-}
-
 void UGameplayCueNotify_HitImpact::HandleGameplayCue(AActor* Self, EGameplayCueEvent::Type EventType, FGameplayCueParameters Parameters)
 {
 	check(EventType == EGameplayCueEvent::Executed);
@@ -31,8 +26,6 @@ void UGameplayCueNotify_HitImpact::HandleGameplayCue(AActor* Self, EGameplayCueE
 		{
 			UGameplayStatics::SpawnEmitterAtLocation(Self, ParticleSystem, HitResult->ImpactPoint, HitResult->ImpactNormal.Rotation(), true);
 		}
-
-
 	}
 	else
 	{
diff --git a/Engine/Source/Runtime/GameplayAbilities/Private/GameplayCueNotify_Static.cpp b/Engine/Source/Runtime/GameplayAbilities/Private/GameplayCueNotify_Static.cpp
new file mode 100644
--- /dev/null
+++ b/Engine/Source/Runtime/GameplayAbilities/Private/GameplayCueNotify_Static.cpp
@@ -0,0 +1,105 @@
+// Copyright 1998-2014 Epic Games, Inc. All Rights Reserved.
+
+#include "AbilitySystemPrivatePCH.h"
+#include "GameplayCueNotify_Static.h"
+#include "GameplayTagsModule.h"
+#include "GameplayCueManager.h"
+#include "AbilitySystemComponent.h"
+
+UGameplayCueNotify_Static::UGameplayCueNotify_Static(const class FObjectInitializer& PCIP)
+: Super(PCIP)
+{
+	IsOverride = true;
+}
+
+#if WITH_EDITOR
+void UGameplayCueNotify_Static::PostEditChangeProperty(FPropertyChangedEvent& PropertyChangedEvent)
+{
+	UAbilitySystemGlobals::Get().GetGameplayCueManager()->bAccelerationMapOutdated = true;
+}
+#endif
+
+void UGameplayCueNotify_Static::DeriveGameplayCueTagFromAssetName()
+{
+	// In the editor, attempt to infer GameplayCueTag from our asset name (if there is no valid GameplayCueTag already).
+	#if WITH_EDITOR
+	if (GIsEditor)
+	{
+		if (GameplayCueTag.IsValid() == false)
+		{
+			FString MyName = GetName();
+
+			MyName.RemoveFromStart(TEXT("Default__"));
+			MyName.RemoveFromEnd(TEXT("_c"));
+
+			MyName.ReplaceInline(TEXT("_"), TEXT("."));
+
+			if (MyName.Contains(TEXT("GameplayCue")) == false)
+			{
+				MyName = FString(TEXT("GameplayCue.")) + MyName;
+			}
+
+			IGameplayTagsModule& GameplayTagsModule = IGameplayTagsModule::Get();
+			GameplayCueTag = GameplayTagsModule.GetGameplayTagsManager().RequestGameplayTag(FName(*MyName), false);
+		}
+		GameplayCueName = GameplayCueTag.GetTagName();
+	}
+	#endif
+}
+
+void UGameplayCueNotify_Static::Serialize(FArchive& Ar)
+{
+	if (Ar.IsSaving())
+	{
+		DeriveGameplayCueTagFromAssetName();
+	}
+
+	Super::Serialize(Ar);
+
+	if (Ar.IsLoading())
+	{
+		DeriveGameplayCueTagFromAssetName();
+	}
+}
+
+void UGameplayCueNotify_Static::PostInitProperties()
+{
+	Super::PostInitProperties();
+	DeriveGameplayCueTagFromAssetName();
+}
+
+bool UGameplayCueNotify_Static::HandlesEvent(EGameplayCueEvent::Type EventType) const
+{
+	return true;
+}
+
+void UGameplayCueNotify_Static::HandleGameplayCue(AActor* MyTarget, EGameplayCueEvent::Type EventType, FGameplayCueParameters Parameters)
+{
+	if (MyTarget && !MyTarget->IsPendingKill())
+	{
+		K2_HandleGameplayCue(MyTarget, EventType, Parameters);
+
+		switch (EventType)
+		{
+		case EGameplayCueEvent::OnActive:
+			OnActive(MyTarget, Parameters);
+			break;
+
+		case EGameplayCueEvent::Executed:
+			OnExecute(MyTarget, Parameters);
+			break;
+
+		case EGameplayCueEvent::Removed:
+			OnRemove(MyTarget, Parameters);
+			break;
+		};
+	}
+	else
+	{
+		ABILITY_LOG(Warning, TEXT("Null Target"));
+	}
+}
+
+void UGameplayCueNotify_Static::OnOwnerDestroyed()
+{
+}
diff --git a/Engine/Source/Runtime/GameplayAbilities/Public/GameplayCueManager.h b/Engine/Source/Runtime/GameplayAbilities/Public/GameplayCueManager.h
--- a/Engine/Source/Runtime/GameplayAbilities/Public/GameplayCueManager.h
+++ b/Engine/Source/Runtime/GameplayAbilities/Public/GameplayCueManager.h
@@ -4,7 +4,7 @@
 
 #include "GameplayTags.h"
 #include "GameplayEffect.h"
-#include "GameplayCueNotify.h"
+#include "GameplayCueNotify_Actor.h"
 #include "GameplayCueManager.generated.h"
 
 /**
@@ -60,7 +60,7 @@ struct FGameplayCueNotifyData
 	FStringAssetReference	GameplayCueNotifyObj;
 
 	UPROPERTY(transient)
-	UGameplayCueNotify*		LoadedGameplayCueNotify;
+	UObject*				LoadedGameplayCueNotify;
 
 	FGameplayTag			GameplayCueTag;
 
@@ -101,7 +101,10 @@ class GAMEPLAYABILITIES_API UGameplayCueManager : public UDataAsset
 	TMap<FGameplayTag, int32>	GameplayCueDataMap;
 	
 	UPROPERTY(transient)
-	UObjectLibrary* GameplayCueNotifyObjectLibrary;
+	UObjectLibrary* GameplayCueNotifyActorObjectLibrary;
+
+	UPROPERTY(transient)
+	UObjectLibrary* GameplayCueNotifyStaticObjectLibrary;
 
 	// -------------------------------------------------------------
 	// Preload GameplayCue tags that we think we will need:
@@ -113,14 +116,8 @@ class GAMEPLAYABILITIES_API UGameplayCueManager : public UDataAsset
 
 	FStreamableManager	StreamableManager;
 
-	UPROPERTY(transient)
-	TArray<UGameplayCueNotify*>	LoadedObjects;
-
-	UPROPERTY(transient)
-	TArray<UGameplayCueNotify*>	InstantiatedObjects;
-
 	// Fixme: we can combine the AActor* and the int32 into a single struct with a decent hash and avoid double map lookups
-	TMap< TWeakObjectPtr<AActor>, TMap<int32, TWeakObjectPtr<UGameplayCueNotify> > >	NotifyMap;
+	TMap<TWeakObjectPtr<AActor>, TMap<int32, TWeakObjectPtr<AGameplayCueNotify_Actor>>>		NotifyMapActor;
 
 	static FGameplayTag	BaseGameplayCueTag();
 
diff --git a/Engine/Source/Runtime/GameplayAbilities/Public/GameplayCueNotify.h b/Engine/Source/Runtime/GameplayAbilities/Public/GameplayCueNotify.h
deleted file mode 100644
diff --git a/Engine/Source/Runtime/GameplayAbilities/Public/GameplayCueNotify_Actor.h b/Engine/Source/Runtime/GameplayAbilities/Public/GameplayCueNotify_Actor.h
new file mode 100644
--- /dev/null
+++ b/Engine/Source/Runtime/GameplayAbilities/Public/GameplayCueNotify_Actor.h
@@ -0,0 +1,68 @@
+// Copyright 1998-2014 Epic Games, Inc. All Rights Reserved.
+
+#pragma once
+
+#include "GameplayTags.h"
+#include "GameplayEffect.h"
+#include "GameplayCueNotify_Actor.generated.h"
+
+
+/**
+ *	A self contained handler of a GameplayCue. These are similiar to AnimNotifies in implementation.
+ *	Instanced GameplayCueNotify which runs arbitrary blueprint code. (TODO: This should be the NON-instanced version!)
+ *	
+ *	
+ *	TODO/Fixme:
+ *		-Unsure: Leave K2_HandleGameplayCue in as generic function?
+ *		-OnExecute/Active/Remove are more clear, easy to use. Make it harder to share info between events.
+ *	
+ *	
+ */
+
+UCLASS(Blueprintable,meta=(ShowWorldContextPin))
+class GAMEPLAYABILITIES_API AGameplayCueNotify_Actor : public AActor
+{
+	GENERATED_UCLASS_BODY()
+
+	/** Does this GameplayCueNotify handle this type of GameplayCueEvent? */
+	virtual bool HandlesEvent(EGameplayCueEvent::Type EventType) const;
+
+	virtual void OnOwnerDestroyed();
+
+	virtual void PostInitProperties();
+
+	virtual void Serialize(FArchive& Ar);
+
+	virtual void HandleGameplayCue(AActor* MyTarget, EGameplayCueEvent::Type EventType, FGameplayCueParameters Parameters);
+
+#if WITH_EDITOR
+	virtual void PostEditChangeProperty(FPropertyChangedEvent& PropertyChangedEvent) override;
+#endif // WITH_EDITOR
+
+	/** Generic Event Graph event that will get called for every event type */
+	UFUNCTION(BlueprintImplementableEvent, Category = "GameplayCueNotify", FriendlyName = "HandleGameplayCue")
+	void K2_HandleGameplayCue(TWeakObjectPtr<AActor> MyTarget, EGameplayCueEvent::Type EventType, FGameplayCueParameters Parameters);
+
+	UFUNCTION(BlueprintImplementableEvent)
+	bool OnExecute(TWeakObjectPtr<AActor> MyTarget, FGameplayCueParameters Parameters) const;
+
+	UFUNCTION(BlueprintImplementableEvent)
+	bool OnActive(TWeakObjectPtr<AActor> MyTarget, FGameplayCueParameters Parameters);
+
+	UFUNCTION(BlueprintImplementableEvent)
+	bool OnRemove(TWeakObjectPtr<AActor> MyTarget, FGameplayCueParameters Parameters);
+
+	UPROPERTY(EditDefaultsOnly, Category = GameplayCue, meta = (Categories = "GameplayCue"))
+	FGameplayTag	GameplayCueTag;
+
+	/** Mirrors GameplayCueTag in order to be asset registry searchable */
+	UPROPERTY(AssetRegistrySearchable)
+	FName GameplayCueName;
+
+	/** Does this Cue override other cues, or is it called in addition to them? E.g., If this is Damage.Physical.Slash, we wont call Damage.Physical afer we run this cue. */
+	UPROPERTY(EditDefaultsOnly, Category = GameplayCue)
+	bool IsOverride;
+
+private:
+	void	DeriveGameplayCueTagFromAssetName();
+};
\ No newline at end of file
diff --git a/Engine/Source/Runtime/GameplayAbilities/Public/GameplayCueNotify_Blueprint.h b/Engine/Source/Runtime/GameplayAbilities/Public/GameplayCueNotify_Blueprint.h
deleted file mode 100644
diff --git a/Engine/Source/Runtime/GameplayAbilities/Public/GameplayCueNotify_HitImpact.h b/Engine/Source/Runtime/GameplayAbilities/Public/GameplayCueNotify_HitImpact.h
--- a/Engine/Source/Runtime/GameplayAbilities/Public/GameplayCueNotify_HitImpact.h
+++ b/Engine/Source/Runtime/GameplayAbilities/Public/GameplayCueNotify_HitImpact.h
@@ -4,7 +4,6 @@
 
 #include "GameplayTags.h"
 #include "GameplayEffect.h"
-#include "GameplayCueNotify.h"
 #include "GameplayCueNotify_HitImpact.generated.h"
 
 /**
@@ -13,15 +12,12 @@
  */
 
 UCLASS()
-class GAMEPLAYABILITIES_API UGameplayCueNotify_HitImpact : public UGameplayCueNotify
+class GAMEPLAYABILITIES_API UGameplayCueNotify_HitImpact : public UGameplayCueNotify_Static
 {
 	GENERATED_UCLASS_BODY()
 
 	/** Does this GameplayCueNotify handle this type of GameplayCueEvent? */
-	virtual bool HandlesEvent(EGameplayCueEvent::Type EventType) const;
-
-	/** Does this GameplayCueNotify need a per source actor instance for this event? */
-	virtual bool NeedsInstanceForEvent(EGameplayCueEvent::Type EventType) const;
+	virtual bool HandlesEvent(EGameplayCueEvent::Type EventType) const override;
 	
 	virtual void HandleGameplayCue(AActor* MyTarget, EGameplayCueEvent::Type EventType, FGameplayCueParameters Parameters) override;
 
diff --git a/Engine/Source/Runtime/GameplayAbilities/Public/GameplayCueNotify_Static.h b/Engine/Source/Runtime/GameplayAbilities/Public/GameplayCueNotify_Static.h
new file mode 100644
--- /dev/null
+++ b/Engine/Source/Runtime/GameplayAbilities/Public/GameplayCueNotify_Static.h
@@ -0,0 +1,68 @@
+// Copyright 1998-2014 Epic Games, Inc. All Rights Reserved.
+
+#pragma once
+
+#include "GameplayTags.h"
+#include "GameplayEffect.h"
+#include "GameplayCueNotify_Static.generated.h"
+
+
+/**
+ *	A self contained handler of a GameplayCue. These are similiar to AnimNotifies in implementation.
+ *	Instanced GameplayCueNotify which runs arbitrary blueprint code. (TODO: This should be the NON-instanced version!)
+ *	
+ *	
+ *	TODO/Fixme:
+ *		-Unsure: Leave K2_HandleGameplayCue in as generic function?
+ *		-OnExecute/Active/Remove are more clear, easy to use. Make it harder to share info between events.
+ *	
+ *	
+ */
+
+UCLASS(Blueprintable,meta=(ShowWorldContextPin))
+class GAMEPLAYABILITIES_API UGameplayCueNotify_Static : public UObject
+{
+	GENERATED_UCLASS_BODY()
+
+	/** Does this GameplayCueNotify handle this type of GameplayCueEvent? */
+	virtual bool HandlesEvent(EGameplayCueEvent::Type EventType) const;
+
+	virtual void OnOwnerDestroyed();
+
+	virtual void PostInitProperties();
+
+	virtual void Serialize(FArchive& Ar);
+
+	virtual void HandleGameplayCue(AActor* MyTarget, EGameplayCueEvent::Type EventType, FGameplayCueParameters Parameters);
+
+#if WITH_EDITOR
+	virtual void PostEditChangeProperty(FPropertyChangedEvent& PropertyChangedEvent) override;
+#endif // WITH_EDITOR
+
+	/** Generic Event Graph event that will get called for every event type */
+	UFUNCTION(BlueprintImplementableEvent, Category = "GameplayCueNotify", FriendlyName = "HandleGameplayCue")
+	void K2_HandleGameplayCue(TWeakObjectPtr<AActor> MyTarget, EGameplayCueEvent::Type EventType, FGameplayCueParameters Parameters);
+
+	UFUNCTION(BlueprintImplementableEvent, BlueprintPure, Category = "GameplayCueNotify")
+	static bool OnExecute(TWeakObjectPtr<AActor> MyTarget, FGameplayCueParameters Parameters);
+
+	UFUNCTION(BlueprintImplementableEvent, BlueprintPure, Category = "GameplayCueNotify")
+	static bool OnActive(TWeakObjectPtr<AActor> MyTarget, FGameplayCueParameters Parameters);
+
+	UFUNCTION(BlueprintImplementableEvent, BlueprintPure, Category = "GameplayCueNotify")
+	static bool OnRemove(TWeakObjectPtr<AActor> MyTarget, FGameplayCueParameters Parameters);
+
+	UPROPERTY(EditDefaultsOnly, Category = GameplayCue, meta = (Categories = "GameplayCue"))
+	FGameplayTag	GameplayCueTag;
+
+	/** Mirrors GameplayCueTag in order to be asset registry searchable */
+	UPROPERTY(AssetRegistrySearchable)
+	FName GameplayCueName;
+
+	/** Does this Cue override other cues, or is it called in addition to them? E.g., If this is Damage.Physical.Slash, we wont call Damage.Physical afer we run this cue. */
+	UPROPERTY(EditDefaultsOnly, Category = GameplayCue)
+	bool IsOverride;
+
+private:
+	void	DeriveGameplayCueTagFromAssetName();
+};
\ No newline at end of file
