- first data using the new base class for damage objectives that function like a tower
- setting it up so the world reticle class can have native function calls that can change material params

diff --git a/Engine/Source/Runtime/GameplayAbilities/Public/Abilities/GameplayAbilityWorldReticle.h b/Engine/Source/Runtime/GameplayAbilities/Public/Abilities/GameplayAbilityWorldReticle.h
--- a/Engine/Source/Runtime/GameplayAbilities/Public/Abilities/GameplayAbilityWorldReticle.h
+++ b/Engine/Source/Runtime/GameplayAbilities/Public/Abilities/GameplayAbilityWorldReticle.h
@@ -40,6 +40,12 @@ public:
 	UFUNCTION(BlueprintImplementableEvent, Category = Reticle)
 	void OnParametersInitialized();
 
+	UFUNCTION(BlueprintImplementableEvent, Category = Reticle)
+	void SetReticleMaterialParamFloat(FName ParamName, float value);
+
+	UFUNCTION(BlueprintImplementableEvent, Category = Reticle)
+	void SetReticleMaterialParamVector(FName ParamName, FVector value);
+
 	UPROPERTY(EditAnywhere, BlueprintReadWrite, meta = (ExposeOnSpawn = "true"), Category = "Reticle")
 	FWorldReticleParameters Parameters;
 
