Removing static BP function call until that code is checked in.

diff --git a/Engine/Source/Runtime/GameplayAbilities/Public/GameplayCueNotify_Static.h b/Engine/Source/Runtime/GameplayAbilities/Public/GameplayCueNotify_Static.h
--- a/Engine/Source/Runtime/GameplayAbilities/Public/GameplayCueNotify_Static.h
+++ b/Engine/Source/Runtime/GameplayAbilities/Public/GameplayCueNotify_Static.h
@@ -44,13 +44,13 @@ class GAMEPLAYABILITIES_API UGameplayCueNotify_Static : public UObject
 	void K2_HandleGameplayCue(TWeakObjectPtr<AActor> MyTarget, EGameplayCueEvent::Type EventType, FGameplayCueParameters Parameters);
 
 	UFUNCTION(BlueprintImplementableEvent, BlueprintPure, Category = "GameplayCueNotify")
-	static bool OnExecute(TWeakObjectPtr<AActor> MyTarget, FGameplayCueParameters Parameters);
+	/*static*/ bool OnExecute(TWeakObjectPtr<AActor> MyTarget, FGameplayCueParameters Parameters);
 
 	UFUNCTION(BlueprintImplementableEvent, BlueprintPure, Category = "GameplayCueNotify")
-	static bool OnActive(TWeakObjectPtr<AActor> MyTarget, FGameplayCueParameters Parameters);
+	/*static*/ bool OnActive(TWeakObjectPtr<AActor> MyTarget, FGameplayCueParameters Parameters);
 
 	UFUNCTION(BlueprintImplementableEvent, BlueprintPure, Category = "GameplayCueNotify")
-	static bool OnRemove(TWeakObjectPtr<AActor> MyTarget, FGameplayCueParameters Parameters);
+	/*static*/ bool OnRemove(TWeakObjectPtr<AActor> MyTarget, FGameplayCueParameters Parameters);
 
 	UPROPERTY(EditDefaultsOnly, Category = GameplayCue, meta = (Categories = "GameplayCue"))
 	FGameplayTag	GameplayCueTag;
