3D Reticles now have a BP function to face their source. Added the concept of valid target and actor target to it.

diff --git a/Engine/Source/Runtime/GameplayAbilities/Private/Abilities/GameplayAbilityTargetActor_GroundTrace.cpp b/Engine/Source/Runtime/GameplayAbilities/Private/Abilities/GameplayAbilityTargetActor_GroundTrace.cpp
--- a/Engine/Source/Runtime/GameplayAbilities/Private/Abilities/GameplayAbilityTargetActor_GroundTrace.cpp
+++ b/Engine/Source/Runtime/GameplayAbilities/Private/Abilities/GameplayAbilityTargetActor_GroundTrace.cpp
@@ -183,9 +183,9 @@ FHitResult AGameplayAbilityTargetActor_GroundTrace::PerformTrace(AActor* InSourc
 		}
 	}
 
-	if (AActor* LocalReticleActor = ReticleActor.Get())
+	if (AGameplayAbilityWorldReticle* LocalReticleActor = ReticleActor.Get())
 	{
-		//TODO Special (for now) functionality: We should tell the reticle to turn red or something, indicating this isn't a valid location
+		LocalReticleActor->SetIsTargetValid(bLastTraceWasGood);
 		LocalReticleActor->SetActorLocation(ReturnHitResult.Location);
 	}
 
diff --git a/Engine/Source/Runtime/GameplayAbilities/Private/Abilities/GameplayAbilityTargetActor_SingleLineTrace.cpp b/Engine/Source/Runtime/GameplayAbilities/Private/Abilities/GameplayAbilityTargetActor_SingleLineTrace.cpp
--- a/Engine/Source/Runtime/GameplayAbilities/Private/Abilities/GameplayAbilityTargetActor_SingleLineTrace.cpp
+++ b/Engine/Source/Runtime/GameplayAbilities/Private/Abilities/GameplayAbilityTargetActor_SingleLineTrace.cpp
@@ -43,9 +43,10 @@ FHitResult AGameplayAbilityTargetActor_SingleLineTrace::PerformTrace(AActor* InS
 	{
 		ReturnHitResult.Location = TraceEnd;
 	}
-	if (AActor* LocalReticleActor = ReticleActor.Get())
+	if (AGameplayAbilityWorldReticle* LocalReticleActor = ReticleActor.Get())
 	{
 		LocalReticleActor->SetActorLocation(ReturnHitResult.Location);
+		LocalReticleActor->SetIsTargetAnActor(ReturnHitResult.bBlockingHit && (ReturnHitResult.Actor != NULL));
 	}
 
 	if (bDebug)
diff --git a/Engine/Source/Runtime/GameplayAbilities/Private/Abilities/GameplayAbilityWorldReticle.cpp b/Engine/Source/Runtime/GameplayAbilities/Private/Abilities/GameplayAbilityWorldReticle.cpp
--- a/Engine/Source/Runtime/GameplayAbilities/Private/Abilities/GameplayAbilityWorldReticle.cpp
+++ b/Engine/Source/Runtime/GameplayAbilities/Private/Abilities/GameplayAbilityWorldReticle.cpp
@@ -15,6 +15,8 @@ AGameplayAbilityWorldReticle::AGameplayAbilityWorldReticle(const FObjectInitiali
 {
 	PrimaryActorTick.bCanEverTick = true;
 	PrimaryActorTick.TickGroup = TG_PrePhysics;
+	bIsTargetValid = true;
+	bIsTargetAnActor = false;
 }
 
 void AGameplayAbilityWorldReticle::InitializeReticle(AGameplayAbilityTargetActor* InTargetingActor, FWorldReticleParameters InParameters)
@@ -42,3 +44,27 @@ bool AGameplayAbilityWorldReticle::GetReplicates() const
 {
 	return bReplicates;
 }
+
+void AGameplayAbilityWorldReticle::SetIsTargetValid(bool bNewValue)
+{
+	if (bIsTargetValid != bNewValue)
+	{
+		bIsTargetValid = bNewValue;
+		OnValidTargetChanged(bNewValue);
+	}
+}
+
+void AGameplayAbilityWorldReticle::SetIsTargetAnActor(bool bNewValue)
+{
+	if (bIsTargetAnActor != bNewValue)
+	{
+		bIsTargetAnActor = bNewValue;
+		OnTargetingAnActor(bNewValue);
+	}
+}
+
+void AGameplayAbilityWorldReticle::FaceTowardSource()
+{
+	check(TargetingActor);
+	SetActorRotation((TargetingActor->StartLocation.GetTargetingTransform().GetLocation() - GetActorLocation()).Rotation());
+}
diff --git a/Engine/Source/Runtime/GameplayAbilities/Public/Abilities/GameplayAbilityWorldReticle.h b/Engine/Source/Runtime/GameplayAbilities/Public/Abilities/GameplayAbilityWorldReticle.h
--- a/Engine/Source/Runtime/GameplayAbilities/Public/Abilities/GameplayAbilityWorldReticle.h
+++ b/Engine/Source/Runtime/GameplayAbilities/Public/Abilities/GameplayAbilityWorldReticle.h
@@ -37,6 +37,17 @@ public:
 
 	void InitializeReticle(AGameplayAbilityTargetActor* InTargetingActor, FWorldReticleParameters InParameters);
 
+	void SetIsTargetValid(bool bNewValue);
+	void SetIsTargetAnActor(bool bNewValue);
+
+	/** Called whenever bIsTargetValid changes value. */
+	UFUNCTION(BlueprintImplementableEvent, Category = Reticle)
+	void OnValidTargetChanged(bool bNewValue);
+
+	/** Called whenever bIsTargetAnActor changes value. */
+	UFUNCTION(BlueprintImplementableEvent, Category = Reticle)
+	void OnTargetingAnActor(bool bNewValue);
+
 	UFUNCTION(BlueprintImplementableEvent, Category = Reticle)
 	void OnParametersInitialized();
 
@@ -46,10 +57,21 @@ public:
 	UFUNCTION(BlueprintImplementableEvent, Category = Reticle)
 	void SetReticleMaterialParamVector(FName ParamName, FVector value);
 
+	UFUNCTION(BlueprintCallable, Category = Reticle)
+	void FaceTowardSource();
+
 	UPROPERTY(EditAnywhere, BlueprintReadWrite, meta = (ExposeOnSpawn = "true"), Category = "Reticle")
 	FWorldReticleParameters Parameters;
 
 protected:
+	/** This indicates whether or not the targeting actor considers the current target to be valid. Defaults to true. */
+	UPROPERTY(BlueprintReadOnly, Category = "Network")
+	bool bIsTargetValid;
+
+	/** This indicates whether or not the targeting reticle is pointed at an actor. Defaults to false. */
+	UPROPERTY(BlueprintReadOnly, Category = "Network")
+	bool bIsTargetAnActor;
+
 	/** This is used in the process of determining whether we should replicate to a specific client. */
 	UPROPERTY(BlueprintReadOnly, Category = "Network")
 	APlayerController* MasterPC;
