GameplayEffect refactor work:
-ActiveGEs register with attribute aggregator's ondirty calls. Dynamic modifiers are updated when backing attributes are modified.
-Added guards against cyclic OnDirty calls
-Added batching mechanism for Aggregator OnDirty calls

diff --git a/Engine/Source/Runtime/GameplayAbilities/Private/Abilities/GameplayAbility.cpp b/Engine/Source/Runtime/GameplayAbilities/Private/Abilities/GameplayAbility.cpp
--- a/Engine/Source/Runtime/GameplayAbilities/Private/Abilities/GameplayAbility.cpp
+++ b/Engine/Source/Runtime/GameplayAbilities/Private/Abilities/GameplayAbility.cpp
@@ -374,7 +374,7 @@ void UGameplayAbility::ApplyCooldown(const FGameplayAbilitySpecHandle Handle, co
 {
 	if (CooldownGameplayEffect)
 	{
-		ApplyGameplayEffectToOwner(Handle, ActorInfo, ActivationInfo, CooldownGameplayEffect, 1.f);
+		ApplyGameplayEffectToOwner(Handle, ActorInfo, ActivationInfo, CooldownGameplayEffect, GetAbilityLevel(Handle, ActorInfo));
 	}
 }
 
@@ -383,7 +383,7 @@ bool UGameplayAbility::CheckCost(const FGameplayAbilitySpecHandle Handle, const
 	if (CostGameplayEffect)
 	{
 		check(ActorInfo->AbilitySystemComponent.IsValid());
-		return ActorInfo->AbilitySystemComponent->CanApplyAttributeModifiers(CostGameplayEffect, 1.f, GetEffectContext(ActorInfo));
+		return ActorInfo->AbilitySystemComponent->CanApplyAttributeModifiers(CostGameplayEffect, GetAbilityLevel(Handle, ActorInfo), GetEffectContext(ActorInfo));
 	}
 	return true;
 }
@@ -392,7 +392,7 @@ void UGameplayAbility::ApplyCost(const FGameplayAbilitySpecHandle Handle, const
 {
 	if (CostGameplayEffect)
 	{
-		ApplyGameplayEffectToOwner(Handle, ActorInfo, ActivationInfo, CostGameplayEffect, 1.f);
+		ApplyGameplayEffectToOwner(Handle, ActorInfo, ActivationInfo, CostGameplayEffect, GetAbilityLevel(Handle, ActorInfo));
 	}
 }
 
diff --git a/Engine/Source/Runtime/GameplayAbilities/Private/AbilitySystemComponent.cpp b/Engine/Source/Runtime/GameplayAbilities/Private/AbilitySystemComponent.cpp
--- a/Engine/Source/Runtime/GameplayAbilities/Private/AbilitySystemComponent.cpp
+++ b/Engine/Source/Runtime/GameplayAbilities/Private/AbilitySystemComponent.cpp
@@ -475,6 +475,7 @@ FActiveGameplayEffectHandle UAbilitySystemComponent::ApplyGameplayEffectSpecToSe
 		{
 			StackSpec = TSharedPtr<FGameplayEffectSpec>(new FGameplayEffectSpec(Spec));
 			OurCopyOfSpec = StackSpec.Get();
+			OurCopyOfSpec->CapturedRelevantAttributes.CaptureAttributes(this, EGameplayEffectAttributeCaptureSource::Target);
 		}
 
 		// if necessary add a modifier to OurCopyOfSpec to force it to have an infinite duration
@@ -486,7 +487,7 @@ FActiveGameplayEffectHandle UAbilitySystemComponent::ApplyGameplayEffectSpecToSe
 	}
 
 
-	OurCopyOfSpec->CapturedRelevantAttributes.CaptureAttributes(this, EGameplayEffectAttributeCaptureSource::Target);
+	
 	
 	// We still probably want to apply tags and stuff even if instant?
 	if (bInvokeGameplayCueApplied)
@@ -881,6 +882,11 @@ void UAbilitySystemComponent::OnAttributeAggregatorDirty(FAggregator* Aggregator
 	ActiveGameplayEffects.OnAttributeAggregatorDirty(Aggregator, Attribute);
 }
 
+void UAbilitySystemComponent::OnMagnitudeDependancyChange(FActiveGameplayEffectHandle Handle, const FAggregator* ChangedAggregator)
+{
+	ActiveGameplayEffects.OnMagnitudeDependancyChange(Handle, ChangedAggregator);
+}
+
 void UAbilitySystemComponent::DisplayDebug(class UCanvas* Canvas, const class FDebugDisplayInfo& DebugDisplay, float& YL, float& YPos)
 {
 
diff --git a/Engine/Source/Runtime/GameplayAbilities/Private/GameplayEffect.cpp b/Engine/Source/Runtime/GameplayAbilities/Private/GameplayEffect.cpp
--- a/Engine/Source/Runtime/GameplayAbilities/Private/GameplayEffect.cpp
+++ b/Engine/Source/Runtime/GameplayAbilities/Private/GameplayEffect.cpp
@@ -157,6 +157,27 @@ bool FGameplayEffectModifierMagnitude::AttemptCalculateMagnitude(const FGameplay
 	return bCanCalc;
 }
 
+bool FGameplayEffectModifierMagnitude::AttempteRecalculateMagnitudeFromDependantChange(const FGameplayEffectSpec& InRelevantSpec, OUT float& OutCalculatedMagnitude, const FAggregator* ChangedAggregator) const
+{
+	TArray<FGameplayEffectAttributeCaptureDefinition> ReqCaptureDefs;
+	GetAttributeCaptureDefinitions(ReqCaptureDefs);
+
+	// We could have many potential captures. If a single one matches our criteria, then we call AttemptCalculateMagnitude once and return.
+	for (const FGameplayEffectAttributeCaptureDefinition& CaptureDef : ReqCaptureDefs)
+	{
+		if (CaptureDef.bSnapshot == false)
+		{
+			const FGameplayEffectAttributeCaptureSpec* CapturedSpec = InRelevantSpec.CapturedRelevantAttributes.FindCaptureSpecByDefinition(CaptureDef, true);
+			if (CapturedSpec && CapturedSpec->ShouldRefreshLinkedAggregator(ChangedAggregator))
+			{
+				return AttemptCalculateMagnitude(InRelevantSpec, OutCalculatedMagnitude);
+			}
+		}
+	}
+
+	return false;
+}
+
 void FGameplayEffectModifierMagnitude::GetAttributeCaptureDefinitions(OUT TArray<FGameplayEffectAttributeCaptureDefinition>& OutCaptureDefs) const
 {
 	OutCaptureDefs.Empty();
@@ -528,6 +549,24 @@ bool FGameplayEffectAttributeCaptureSpec::AttemptAddAggregatorModsToAggregator(O
 	return false;
 }
 
+void FGameplayEffectAttributeCaptureSpec::RegisterLinkedAggregatorCallback(FActiveGameplayEffectHandle Handle) const
+{
+	if (BackingDefinition.bSnapshot == false)
+	{
+		// Its possible the linked Aggregator is already gone.
+		FAggregator* Agg = AttributeAggregator.Get();
+		if (Agg)
+		{
+			Agg->AddDependant(Handle);
+		}
+	}
+}
+
+bool FGameplayEffectAttributeCaptureSpec::ShouldRefreshLinkedAggregator(const FAggregator* ChangedAggregator) const
+{
+	return (BackingDefinition.bSnapshot == false && (ChangedAggregator == nullptr || AttributeAggregator.Get() == ChangedAggregator));
+}
+
 const FGameplayEffectAttributeCaptureDefinition& FGameplayEffectAttributeCaptureSpec::GetBackingDefinition() const
 {
 	return BackingDefinition;
@@ -611,7 +650,18 @@ bool FGameplayEffectAttributeCaptureSpecContainer::HasNonSnapshottedAttributes()
 	return bHasNonSnapshottedAttributes;
 }
 
+void FGameplayEffectAttributeCaptureSpecContainer::RegisterLinkedAggregatorCallbacks(FActiveGameplayEffectHandle Handle) const
+{
+	for (const FGameplayEffectAttributeCaptureSpec& CaptureSpec : SourceAttributes)
+	{
+		CaptureSpec.RegisterLinkedAggregatorCallback(Handle);
+	}
 
+	for (const FGameplayEffectAttributeCaptureSpec& CaptureSpec : TargetAttributes)
+	{
+		CaptureSpec.RegisterLinkedAggregatorCallback(Handle);
+	}
+}
 
 // --------------------------------------------------------------------------------------------------------------------------------------------------------
 //
@@ -626,10 +676,13 @@ void FActiveGameplayEffect::CheckOngoingTagRequirements(const FGameplayTagContai
 
 	if (IsInhibited != ShouldBeInhibited)
 	{
+		// All OnDirty callbacks must be inhibited until we update this entire GameplayEffect.
+		FScopedAggregatorOnDirtyBatch	AggregatorOnDirtyBatcher;
+
 		if (ShouldBeInhibited)
 		{
 			// Register this ActiveGameplayEffects modifiers with our Attribute Aggregators
-			if (Spec.GetPeriod() == UGameplayEffect::NO_PERIOD)
+			if (Spec.GetPeriod() <= UGameplayEffect::NO_PERIOD)
 			{
 				for (int32 ModIdx = 0; ModIdx < Spec.Modifiers.Num(); ++ModIdx)
 				{
@@ -642,7 +695,7 @@ void FActiveGameplayEffect::CheckOngoingTagRequirements(const FGameplayTagContai
 		else
 		{
 			// Register this ActiveGameplayEffects modifiers with our Attribute Aggregators
-			if (Spec.GetPeriod() == UGameplayEffect::NO_PERIOD)
+			if (Spec.GetPeriod() <= UGameplayEffect::NO_PERIOD)
 			{
 				for (int32 ModIdx = 0; ModIdx < Spec.Modifiers.Num(); ++ModIdx)
 				{
@@ -750,9 +803,6 @@ void FActiveGameplayEffectsContainer::RegisterWithOwner(UAbilitySystemComponent*
 		// Binding raw is ok here, since the owner is literally the UObject that owns us. If we are destroyed, its because that uobject is destroyed,
 		// and if that is destroyed, the delegate wont be able to fire.
 		Owner->RegisterGenericGameplayTagEvent().AddRaw(this, &FActiveGameplayEffectsContainer::OnOwnerTagChange);
-
-		// Add Aggregators for system attributes
-		//FindOrCreateAttributeAggregator(
 	}
 }
 
@@ -901,6 +951,65 @@ void FActiveGameplayEffectsContainer::OnAttributeAggregatorDirty(FAggregator* Ag
 	InternalUpdateNumericalAttribute(Attribute, NewValue, nullptr);
 }
 
+void FActiveGameplayEffectsContainer::OnMagnitudeDependancyChange(FActiveGameplayEffectHandle Handle, const FAggregator* ChangedAgg)
+{
+	if (Handle.IsValid())
+	{
+		FActiveGameplayEffect* ActiveEffect = GetActiveGameplayEffect(Handle);
+		if (ActiveEffect)
+		{
+			// This handle registered with the ChangedAgg to be notified when the aggregator changed.
+			// At this point we don't know what actually needs to be updated inside this active gameplay effect.
+			FGameplayEffectSpec& Spec = ActiveEffect->Spec;
+
+			// We must update attribute aggregators only if we are actually 'on' right now, and if we are non periodic (periodic effects do their thing on execute callbacks)
+			bool MustUpdateAttributeAggregators = (ActiveEffect->IsInhibited == false && (Spec.GetDuration() <= UGameplayEffect::NO_PERIOD));
+
+			// As we update our modifier magnitudes, we will update our owner's attribute aggregators. When we do this, we have to clear them first of all of our (Handle's) previous mods.
+			// Since we could potentially have two mods to the same attribute, one that gets updated, and one that doesnt - we need to do this in two passes.
+			TSet<FGameplayAttribute>	AttributesToUpdate;
+
+			// First pass: update magnitudes of our modifiers that changed
+			for(int32 ModIdx = 0; ModIdx < Spec.Modifiers.Num(); ++ModIdx)
+			{
+				const FGameplayModifierInfo& ModDef = Spec.Def->Modifiers[ModIdx];
+				FModifierSpec& ModSpec = Spec.Modifiers[ModIdx];
+
+				if (ModDef.ModifierMagnitude.AttempteRecalculateMagnitudeFromDependantChange(Spec, ModSpec.EvaluatedMagnitude, ChangedAgg))
+				{
+					// We changed, so we need to reapply/update our spot in the attribute aggregator map
+					if (MustUpdateAttributeAggregators)
+					{
+						AttributesToUpdate.Add(ModDef.Attribute);
+					}
+				}
+			}
+
+			// Second pass, update the aggregators that we need to
+			for (const FGameplayAttribute& Attribute : AttributesToUpdate)
+			{
+				FAggregator* Aggregator = FindOrCreateAttributeAggregator(Attribute).Get();
+				check(Aggregator);
+
+				// Remove ALL of our mods (this is all we can do! We don't keep track of aggregator mods <=> spec mod mappings)
+				Aggregator->RemoveMod(Handle);
+
+				// Now readd ALL of our mods (even if they weren't recalculated here)
+				for(int32 ModIdx = 0; ModIdx < Spec.Modifiers.Num(); ++ModIdx)
+				{
+					const FGameplayModifierInfo& ModDef = Spec.Def->Modifiers[ModIdx];
+					FModifierSpec& ModSpec = Spec.Modifiers[ModIdx];
+
+					if (ModDef.Attribute == Attribute)
+					{
+						Aggregator->AddMod(ModSpec.GetEvaluatedMagnitude(), ModDef.ModifierOp, &ModDef.SourceTags, &ModDef.TargetTags, Handle);
+					}
+				}
+			}
+		}
+	}
+}
+
 void FActiveGameplayEffectsContainer::SetBaseAttributeValueFromReplication(FGameplayAttribute Attribute, float BaseValue)
 {
 /*
@@ -1107,8 +1216,12 @@ FActiveGameplayEffect& FActiveGameplayEffectsContainer::CreateNewActiveGameplayE
 
 	// Calc all of our modifier magnitudes now. Some may need to update later based on attributes changing, etc, but those should
 	// be done through delegate callbacks.
+	NewEffect.Spec.CapturedRelevantAttributes.CaptureAttributes(Owner, EGameplayEffectAttributeCaptureSource::Target);
 	NewEffect.Spec.CalculateModifierMagnitudes();
 
+	// Register Source and Target non snapshot capture delegates here
+	NewEffect.Spec.CapturedRelevantAttributes.RegisterLinkedAggregatorCallbacks(NewHandle);
+
 	// Calculate Duration mods if we have a real duration
 	float DurationBaseValue = NewEffect.Spec.GetDuration();
 	if (DurationBaseValue > 0.f)
diff --git a/Engine/Source/Runtime/GameplayAbilities/Private/GameplayEffectAggregator.cpp b/Engine/Source/Runtime/GameplayAbilities/Private/GameplayEffectAggregator.cpp
--- a/Engine/Source/Runtime/GameplayAbilities/Private/GameplayEffectAggregator.cpp
+++ b/Engine/Source/Runtime/GameplayAbilities/Private/GameplayEffectAggregator.cpp
@@ -164,6 +164,11 @@ void FAggregator::AddModsFrom(const FAggregator& SourceAggregator)
 	}
 }
 
+void FAggregator::AddDependant(FActiveGameplayEffectHandle Handle)
+{
+	Dependants.Add(Handle);
+}
+
 void FAggregator::RemoveModsWithActiveHandle(TArray<FAggregatorMod>& Mods, FActiveGameplayEffectHandle ActiveHandle)
 {
 	check(ActiveHandle.IsValid());
@@ -188,7 +193,36 @@ void FAggregator::TakeSnapshotOf(const FAggregator& AggToSnapshot)
 
 void FAggregator::BroadcastOnDirty()
 {
+	// If we are batching on Dirty calls (and we actually have dependants registered with us) then early out.
+	if (FScopedAggregatorOnDirtyBatch::GlobalBatchCount > 0 && (Dependants.Num() > 0 || OnDirty.IsBound()))
+	{
+		FScopedAggregatorOnDirtyBatch::DirtyAggregators.Add(this);
+		return;
+	}
+
+	if (IsBroadcastingDirty)
+	{
+		// Apologies for the vague warning but its very hard from this spot to call out what data has caused this. If this frequently happens we should improve this.
+		ABILITY_LOG(Warning, TEXT("FAggregator detected cyclic attribute dependancies. We are skipping a recursive dirty call. Its possible the resulting attribute values are not what you expect!"));
+		return;
+	}
+
+	TGuardValue<bool>	Guard(IsBroadcastingDirty, true);
+	
 	OnDirty.Broadcast(this);
+
+	TArray<FActiveGameplayEffectHandle>	ValidDependants;
+	for (FActiveGameplayEffectHandle Handle : Dependants)
+	{
+		UAbilitySystemComponent* ASC = Handle.GetOwningAbilitySystemComponent();
+		if (ASC)
+		{
+			ASC->OnMagnitudeDependancyChange(Handle, this);
+			ValidDependants.Add(Handle);
+		}
+	}
+	Dependants = ValidDependants;
+
 }
 
 void FAggregatorRef::TakeSnapshotOf(const FAggregatorRef& RefToSnapshot)
@@ -205,3 +239,25 @@ void FAggregatorRef::TakeSnapshotOf(const FAggregatorRef& RefToSnapshot)
 		Data.Reset();
 	}
 }
+
+int32 FScopedAggregatorOnDirtyBatch::GlobalBatchCount = 0;
+TSet<FAggregator*> FScopedAggregatorOnDirtyBatch::DirtyAggregators;
+
+FScopedAggregatorOnDirtyBatch::FScopedAggregatorOnDirtyBatch()
+{
+	GlobalBatchCount++;
+}
+
+FScopedAggregatorOnDirtyBatch::~FScopedAggregatorOnDirtyBatch()
+{
+	GlobalBatchCount--;
+	if (GlobalBatchCount == 0)
+	{
+		for (FAggregator* Agg : DirtyAggregators)
+		{
+			Agg->BroadcastOnDirty();
+
+		}
+		DirtyAggregators.Empty();
+	}
+}
\ No newline at end of file
diff --git a/Engine/Source/Runtime/GameplayAbilities/Private/GameplayEffectExecutionCalculation.cpp b/Engine/Source/Runtime/GameplayAbilities/Private/GameplayEffectExecutionCalculation.cpp
--- a/Engine/Source/Runtime/GameplayAbilities/Private/GameplayEffectExecutionCalculation.cpp
+++ b/Engine/Source/Runtime/GameplayAbilities/Private/GameplayEffectExecutionCalculation.cpp
@@ -55,6 +55,12 @@ UAbilitySystemComponent* FGameplayEffectCustomExecutionParameters::GetTargetAbil
 	return TargetAbilitySystemComponent.Get();
 }
 
+UAbilitySystemComponent* FGameplayEffectCustomExecutionParameters::GetSourcebilitySystemComponent() const
+{
+	check(OwningSpec);
+	return OwningSpec->GetContext().GetInstigatorAbilitySystemComponent();
+}
+
 bool FGameplayEffectCustomExecutionParameters::AttemptCalculateCapturedAttributeMagnitude(const FGameplayEffectAttributeCaptureDefinition& InCaptureDef, const FAggregatorEvaluateParameters& InEvalParams, OUT float& OutMagnitude) const
 {
 	check(OwningSpec);
diff --git a/Engine/Source/Runtime/GameplayAbilities/Public/AbilitySystemComponent.h b/Engine/Source/Runtime/GameplayAbilities/Public/AbilitySystemComponent.h
--- a/Engine/Source/Runtime/GameplayAbilities/Public/AbilitySystemComponent.h
+++ b/Engine/Source/Runtime/GameplayAbilities/Public/AbilitySystemComponent.h
@@ -238,7 +238,7 @@ class GAMEPLAYABILITIES_API UAbilitySystemComponent : public UActorComponent, pu
 
 	void OnAttributeAggregatorDirty(FAggregator* Aggregator, FGameplayAttribute Attribute);
 
-	void OnOngoingTagRequirementChange(FActiveGameplayEffectHandle ActiveHandle);
+	void OnMagnitudeDependancyChange(FActiveGameplayEffectHandle Handle, const FAggregator* ChangedAggregator);
 
 	// --------------------------------------------
 	// Tags
diff --git a/Engine/Source/Runtime/GameplayAbilities/Public/GameplayEffect.h b/Engine/Source/Runtime/GameplayAbilities/Public/GameplayEffect.h
--- a/Engine/Source/Runtime/GameplayAbilities/Public/GameplayEffect.h
+++ b/Engine/Source/Runtime/GameplayAbilities/Public/GameplayEffect.h
@@ -201,6 +201,9 @@ public:
 	 */
 	bool AttemptCalculateMagnitude(const FGameplayEffectSpec& InRelevantSpec, OUT float& OutCalculatedMagnitude) const;
 
+	/** Attempts to recalculate the magnitude given a changed aggregator. This will only recalculate if we are a modifier that is linked (non snapshot) to the given aggregator. */
+	bool AttempteRecalculateMagnitudeFromDependantChange(const FGameplayEffectSpec& InRelevantSpec, OUT float& OutCalculatedMagnitude, const FAggregator* ChangedAggregator) const;
+
 	/**
 	 * Gather all of the attribute capture definitions necessary to compute the magnitude and place them into the provided array
 	 * 
@@ -539,7 +542,9 @@ private:
 	UPROPERTY()
 	float EvaluatedMagnitude;
 
+	/** These structures are the only ones that should internally be able to update the EvaluatedMagnitude. Any gamecode that gets its hands on FModifierSpec should never be setting EvaluatedMagnitude manually */
 	friend struct FGameplayEffectSpec;
+	friend struct FActiveGameplayEffectsContainer;
 };
 
 /** Saves list of modified attributes, to use for gameplay cues or later processing */
@@ -634,6 +639,12 @@ struct FGameplayEffectAttributeCaptureSpec
 	
 	/** Simple accessor to backing capture definition */
 	const FGameplayEffectAttributeCaptureDefinition& GetBackingDefinition() const;
+
+	/** Register this handle with linked aggregators */
+	void RegisterLinkedAggregatorCallback(FActiveGameplayEffectHandle Handle) const;
+	
+	/** Return true if this capture should be recalculated if the given aggregator has changed */
+	bool ShouldRefreshLinkedAggregator(const FAggregator* ChangedAggregator) const;
 		
 private:
 
@@ -691,6 +702,9 @@ public:
 	/** Returns whether the container has at least one spec w/o snapshotted attributes */
 	bool HasNonSnapshottedAttributes() const;
 
+	/** Registers any linked aggregators to notify this active handle if they are dirtied */
+	void RegisterLinkedAggregatorCallbacks(FActiveGameplayEffectHandle Handle) const;
+
 private:
 
 	/** Captured attributes from the source of a gameplay effect */
@@ -716,7 +730,7 @@ private:
  * is still distinct from an FActiveGameplayEffect which in an applied instance of an FGameplayEffectSpec.
  */
 USTRUCT()
-struct FGameplayEffectSpec
+struct GAMEPLAYABILITIES_API FGameplayEffectSpec
 {
 	GENERATED_USTRUCT_BODY()
 
@@ -1123,6 +1137,8 @@ private:
 	FAggregatorRef& FindOrCreateAttributeAggregator(FGameplayAttribute Attribute);
 
 	void OnAttributeAggregatorDirty(FAggregator* Aggregator, FGameplayAttribute Attribute);
+
+	void OnMagnitudeDependancyChange(FActiveGameplayEffectHandle Handle, const FAggregator* ChangedAgg);
 };
 
 template<>
diff --git a/Engine/Source/Runtime/GameplayAbilities/Public/GameplayEffectAggregator.h b/Engine/Source/Runtime/GameplayAbilities/Public/GameplayEffectAggregator.h
--- a/Engine/Source/Runtime/GameplayAbilities/Public/GameplayEffectAggregator.h
+++ b/Engine/Source/Runtime/GameplayAbilities/Public/GameplayEffectAggregator.h
@@ -4,8 +4,6 @@
 
 #include "GameplayEffectTypes.h"
 
-// #include "GameplayEffectAggregator.generated.h"
-
 struct GAMEPLAYABILITIES_API FAggregatorEvaluateParameters
 {
 	FAggregatorEvaluateParameters() : SourceTags(nullptr), TargetTags(nullptr) { }
@@ -37,7 +35,7 @@ struct GAMEPLAYABILITIES_API FAggregator : public TSharedFromThis<FAggregator>
 {
 	DECLARE_MULTICAST_DELEGATE_OneParam(FOnAggregatorDirty, FAggregator*);
 
-	FAggregator(float InBaseValue=0.f) : BaseValue(InBaseValue) { }
+	FAggregator(float InBaseValue=0.f) : BaseValue(InBaseValue), IsBroadcastingDirty(false) { }
 
 	/** Simple accessor to base value */
 	float GetBaseValue() const;
@@ -64,17 +62,23 @@ struct GAMEPLAYABILITIES_API FAggregator : public TSharedFromThis<FAggregator>
 
 	void AddModsFrom(const FAggregator& SourceAggregator);
 
+	void AddDependant(FActiveGameplayEffectHandle Handle);
+
 private:
 
 	void BroadcastOnDirty();
+	float SumMods(const TArray<FAggregatorMod> &Mods, float Bias, const FAggregatorEvaluateParameters& Parameters) const;
+	void RemoveModsWithActiveHandle(TArray<FAggregatorMod>& Mods, FActiveGameplayEffectHandle ActiveHandle);
 
 	float	BaseValue;
 	TArray<FAggregatorMod>	Mods[EGameplayModOp::Max];
 
-	float SumMods(const TArray<FAggregatorMod> &Mods, float Bias, const FAggregatorEvaluateParameters& Parameters) const;
-	void RemoveModsWithActiveHandle(TArray<FAggregatorMod>& Mods, FActiveGameplayEffectHandle ActiveHandle);
+	/** ActiveGE handles that we need to notify if we change. NOT copied over during snapshots. */
+	TArray<FActiveGameplayEffectHandle>	Dependants;
+	bool	IsBroadcastingDirty;
 
 	friend struct FAggregator;
+	friend struct FScopedAggregatorOnDirtyBatch;	// Only outside class that gets to call BroadcastOnDirty()
 	friend class UAbilitySystemComponent;	// Only needed for DisplayDebug()
 };
 
@@ -89,3 +93,19 @@ struct GAMEPLAYABILITIES_API FAggregatorRef
 
 	void TakeSnapshotOf(const FAggregatorRef& RefToSnapshot);
 };
+
+/**
+ *	Allows us to batch all aggregator OnDirty calls within a scope. That is, ALL OnDirty() callbacks are
+ *	delayed until FScopedAggregatorOnDirtyBatch goes out of scope.
+ *	
+ *	The only catch is that we store raw FAggregator*. This should only be used in scopes where aggreagtors
+ *	are not deleted. There is currently no place that does. If we find to, we could add additional safety checks.
+ */
+struct GAMEPLAYABILITIES_API FScopedAggregatorOnDirtyBatch
+{
+	FScopedAggregatorOnDirtyBatch();
+	~FScopedAggregatorOnDirtyBatch();
+
+	static int32	GlobalBatchCount;
+	static TSet<FAggregator*>	DirtyAggregators;
+};
\ No newline at end of file
diff --git a/Engine/Source/Runtime/GameplayAbilities/Public/GameplayEffectExecutionCalculation.h b/Engine/Source/Runtime/GameplayAbilities/Public/GameplayEffectExecutionCalculation.h
--- a/Engine/Source/Runtime/GameplayAbilities/Public/GameplayEffectExecutionCalculation.h
+++ b/Engine/Source/Runtime/GameplayAbilities/Public/GameplayEffectExecutionCalculation.h
@@ -27,6 +27,9 @@ public:
 
 	/** Simple accessor to target ability system component */
 	UAbilitySystemComponent* GetTargetAbilitySystemComponent() const;
+
+	/** Simple accessor to source ability system component (could be null!) */
+	UAbilitySystemComponent* GetSourcebilitySystemComponent() const;
 	
 	/**
 	 * Attempts to calculate the magnitude of a captured attribute given the specified parameters. Can fail if the gameplay spec doesn't have
