call IGameplayCueInterface even if cue isnt in the acceleration map

diff --git a/Engine/Source/Runtime/GameplayAbilities/Private/GameplayCueManager.cpp b/Engine/Source/Runtime/GameplayAbilities/Private/GameplayCueManager.cpp
--- a/Engine/Source/Runtime/GameplayAbilities/Private/GameplayCueManager.cpp
+++ b/Engine/Source/Runtime/GameplayAbilities/Private/GameplayCueManager.cpp
@@ -39,12 +39,12 @@ void UGameplayCueManager::HandleGameplayCue(AActor* TargetActor, FGameplayTag Ga
 	{
 		int32 DataIdx = *Ptr;
 		HandleGameplayCueNotify_Internal(TargetActor, DataIdx, EventType, Parameters);
+	}
 
-		IGameplayCueInterface* GameplayCueInterface = Cast<IGameplayCueInterface>(TargetActor);
-		if (GameplayCueInterface)
-		{
-			GameplayCueInterface->HandleGameplayCue(TargetActor, GameplayCueTag, EventType, Parameters);
-		}
+	IGameplayCueInterface* GameplayCueInterface = Cast<IGameplayCueInterface>(TargetActor);
+	if (GameplayCueInterface)
+	{
+		GameplayCueInterface->HandleGameplayCue(TargetActor, GameplayCueTag, EventType, Parameters);
 	}
 }
 
