Fix it so ActiveGameplayEffects only grant tags when ongoing tag reqs are met.
make it so captured tags are not replicated on the spec (no need for them after apply)
Added dynamic granted tags to GameplayEffectSpec
some prep work for Orion cooldown experiment

diff --git a/Engine/Source/Runtime/GameplayAbilities/Private/Abilities/GameplayAbility.cpp b/Engine/Source/Runtime/GameplayAbilities/Private/Abilities/GameplayAbility.cpp
--- a/Engine/Source/Runtime/GameplayAbilities/Private/Abilities/GameplayAbility.cpp
+++ b/Engine/Source/Runtime/GameplayAbilities/Private/Abilities/GameplayAbility.cpp
@@ -388,11 +388,11 @@ UGameplayEffect* UGameplayAbility::GetCostGameplayEffect() const
 
 bool UGameplayAbility::CheckCooldown(const FGameplayAbilitySpecHandle Handle, const FGameplayAbilityActorInfo* ActorInfo) const
 {
-	UGameplayEffect* CooldownGE = GetCooldownGameplayEffect();
-	if (CooldownGE)
+	FGameplayTagContainer* CooldownTags = GetCooldownTags();
+	if (CooldownTags)
 	{
 		check(ActorInfo->AbilitySystemComponent.IsValid());
-		if (CooldownGE->InheritableOwnedTagsContainer.CombinedTags.Num() > 0 && ActorInfo->AbilitySystemComponent->HasAnyMatchingGameplayTags(CooldownGE->InheritableOwnedTagsContainer.CombinedTags))
+		if (CooldownTags->Num() > 0 && ActorInfo->AbilitySystemComponent->HasAnyMatchingGameplayTags(*CooldownTags))
 		{
 			return false;
 		}
@@ -437,13 +437,13 @@ float UGameplayAbility::GetCooldownTimeRemaining(const FGameplayAbilityActorInfo
 	UGameplayEffect* CooldownGE = GetCooldownGameplayEffect();
 	if (CooldownGE)
 	{
-		TArray< float > Durations = ActorInfo->AbilitySystemComponent->GetActiveEffectsTimeRemaining(FActiveGameplayEffectQuery(&CooldownGE->InheritableOwnedTagsContainer.CombinedTags));
+		TArray< float > Durations = ActorInfo->AbilitySystemComponent->GetActiveEffectsTimeRemaining(FActiveGameplayEffectQuery(CooldownTags));
 		if (Durations.Num() > 0)
-			{
+		{
 			Durations.Sort();
 			return Durations[Durations.Num()-1];
-				}
-			}
+		}
+	}
 
 	return 0.f;
 }
@@ -456,14 +456,14 @@ void UGameplayAbility::GetCooldownTimeRemainingAndDuration(const FGameplayAbilit
 
 	TimeRemaining = 0.f;
 	CooldownDuration = 0.f;
-
-	UGameplayEffect* CooldownGE = GetCooldownGameplayEffect();
+	
+	FGameplayTagContainer* CooldownTags = GetCooldownTags();
 	if (CooldownGE)
 	{
-		TArray< float > DurationRemaining = ActorInfo->AbilitySystemComponent->GetActiveEffectsTimeRemaining(FActiveGameplayEffectQuery(&CooldownGE->InheritableOwnedTagsContainer.CombinedTags));
+		TArray< float > DurationRemaining = ActorInfo->AbilitySystemComponent->GetActiveEffectsTimeRemaining(FActiveGameplayEffectQuery(CooldownTags));
 		if (DurationRemaining.Num() > 0)
 		{
-			TArray< float > Durations = ActorInfo->AbilitySystemComponent->GetActiveEffectsDuration(FActiveGameplayEffectQuery(&CooldownGE->InheritableOwnedTagsContainer.CombinedTags));
+			TArray< float > Durations = ActorInfo->AbilitySystemComponent->GetActiveEffectsDuration(FActiveGameplayEffectQuery(CooldownTags));
 			check(Durations.Num() == DurationRemaining.Num());
 			int32 BestIdx = 0;
 			float LongestTime = DurationRemaining[0];
@@ -482,6 +482,12 @@ void UGameplayAbility::GetCooldownTimeRemainingAndDuration(const FGameplayAbilit
 	}
 }
 
+FGameplayTagContainer* UGameplayAbility::GetCooldownTags() const
+{
+	UGameplayEffect* CDGE =GetCooldownGameplayEffect();
+	return CDGE ? &CDGE->InheritableClearTagsContainer.CombinedTags : nullptr;
+}
+
 FGameplayAbilityActorInfo UGameplayAbility::GetActorInfo() const
 {
 	check(CurrentActorInfo);
diff --git a/Engine/Source/Runtime/GameplayAbilities/Private/GameplayEffect.cpp b/Engine/Source/Runtime/GameplayAbilities/Private/GameplayEffect.cpp
--- a/Engine/Source/Runtime/GameplayAbilities/Private/GameplayEffect.cpp
+++ b/Engine/Source/Runtime/GameplayAbilities/Private/GameplayEffect.cpp
@@ -16,6 +16,8 @@ const float UGameplayEffect::INSTANT_APPLICATION = 0.f;
 const float UGameplayEffect::NO_PERIOD = 0.f;
 const float UGameplayEffect::INVALID_LEVEL = -1.f;
 
+#pragma optimize( "", off )
+
 // --------------------------------------------------------------------------------------------------------------------------------------------------------
 //
 //	UGameplayEffect
@@ -746,6 +748,7 @@ void FGameplayEffectAttributeCaptureSpecContainer::RegisterLinkedAggregatorCallb
 // --------------------------------------------------------------------------------------------------------------------------------------------------------
 
 
+/** This is the core function that turns the ActiveGE 'on' or 'off */
 void FActiveGameplayEffect::CheckOngoingTagRequirements(const FGameplayTagContainer& OwnerTags, FActiveGameplayEffectsContainer& OwningContainer)
 {
 	bool ShouldBeInhibited = !Spec.Def->OngoingTagRequirements.RequirementsMet(OwnerTags);
@@ -757,34 +760,12 @@ void FActiveGameplayEffect::CheckOngoingTagRequirements(const FGameplayTagContai
 
 		if (ShouldBeInhibited)
 		{
-			// Register this ActiveGameplayEffects modifiers with our Attribute Aggregators
-			if (Spec.GetPeriod() <= UGameplayEffect::NO_PERIOD)
-			{
-				for (int32 ModIdx = 0; ModIdx < Spec.Modifiers.Num(); ++ModIdx)
-				{
-					const FModifierSpec &Mod = Spec.Modifiers[ModIdx];
-					FAggregator* Aggregator = OwningContainer.FindOrCreateAttributeAggregator(Spec.Def->Modifiers[ModIdx].Attribute).Get();
-					Aggregator->RemoveMod(Handle);
-				}
-			}
+			// Remove our ActiveGameplayEffects modifiers with our Attribute Aggregators
+			OwningContainer.RemoveActiveGameplayEffectGrantedTagsAndModifiers(*this);
 		}
 		else
 		{
-			// Register this ActiveGameplayEffects modifiers with our Attribute Aggregators
-			if (Spec.GetPeriod() <= UGameplayEffect::NO_PERIOD)
-			{
-				for (int32 ModIdx = 0; ModIdx < Spec.Modifiers.Num(); ++ModIdx)
-				{
-					const FModifierSpec &Mod = Spec.Modifiers[ModIdx];
-					const FGameplayModifierInfo &ModInfo = Spec.Def->Modifiers[ModIdx];
-
-					// Note we assume the EvaluatedMagnitude is up to do. There is no case currently where we should recalculate magnitude based on
-					// Ongoing tags being met. We either calculate magnitude one time, or its done via OnDirty calls (or potentially a frequency timer one day)
-					
-					FAggregator* Aggregator = OwningContainer.FindOrCreateAttributeAggregator(Spec.Def->Modifiers[ModIdx].Attribute).Get();
-					Aggregator->AddMod(Mod.GetEvaluatedMagnitude(), ModInfo.ModifierOp, &ModInfo.SourceTags, &ModInfo.TargetTags, Handle);
-				}
-			}
+			OwningContainer.AddActiveGameplayEffectGrantedTagsAndModifiers(*this);
 		}
 
 		IsInhibited = ShouldBeInhibited;
@@ -1414,13 +1395,6 @@ void FActiveGameplayEffectsContainer::InternalOnActiveGameplayEffectAdded(FActiv
 		ActiveEffectTagDependencies.FindOrAdd(Tag).Add(Effect.Handle);
 	}
 
-	// Update our owner with the tags this GameplayEffect grants them
-	Owner->UpdateTagMap(Effect.Spec.Def->InheritableOwnedTagsContainer.CombinedTags, 1);
-	for (const FGameplayEffectCue& Cue : Effect.Spec.Def->GameplayCues)
-	{
-		Owner->UpdateTagMap(Cue.GameplayCueTags, 1);
-	}
-
 	// Check if we should actually be turned on or not (this will turn us on for the first time)
 	FGameplayTagContainer OwnerTags;
 	Owner->GetOwnedGameplayTags(OwnerTags);
@@ -1429,6 +1403,35 @@ void FActiveGameplayEffectsContainer::InternalOnActiveGameplayEffectAdded(FActiv
 	Effect.CheckOngoingTagRequirements(OwnerTags, *this);
 }
 
+void FActiveGameplayEffectsContainer::AddActiveGameplayEffectGrantedTagsAndModifiers(FActiveGameplayEffect& Effect)
+{
+	// Register this ActiveGameplayEffects modifiers with our Attribute Aggregators
+	if (Effect.Spec.GetPeriod() <= UGameplayEffect::NO_PERIOD)
+	{
+		for (int32 ModIdx = 0; ModIdx < Effect.Spec.Modifiers.Num(); ++ModIdx)
+		{
+			const FModifierSpec &Mod = Effect.Spec.Modifiers[ModIdx];
+			const FGameplayModifierInfo &ModInfo = Effect.Spec.Def->Modifiers[ModIdx];
+
+			// Note we assume the EvaluatedMagnitude is up to do. There is no case currently where we should recalculate magnitude based on
+			// Ongoing tags being met. We either calculate magnitude one time, or its done via OnDirty calls (or potentially a frequency timer one day)
+					
+			FAggregator* Aggregator = FindOrCreateAttributeAggregator(Effect.Spec.Def->Modifiers[ModIdx].Attribute).Get();
+			Aggregator->AddMod(Mod.GetEvaluatedMagnitude(), ModInfo.ModifierOp, &ModInfo.SourceTags, &ModInfo.TargetTags, Effect.Handle);
+		}
+	}
+
+	// Update our owner with the tags this GameplayEffect grants them
+	Owner->UpdateTagMap(Effect.Spec.Def->InheritableOwnedTagsContainer.CombinedTags, 1);
+
+	Owner->UpdateTagMap(Effect.Spec.DynamicGrantedTags, 1);
+
+	for (const FGameplayEffectCue& Cue : Effect.Spec.Def->GameplayCues)
+	{
+		Owner->UpdateTagMap(Cue.GameplayCueTags, 1);
+	}
+}
+
 /** Called on server to remove a GameplayEffect */
 bool FActiveGameplayEffectsContainer::RemoveActiveGameplayEffect(FActiveGameplayEffectHandle Handle)
 {
@@ -1506,39 +1509,49 @@ bool FActiveGameplayEffectsContainer::InternalRemoveActiveGameplayEffect(int32 I
 /** Called by client and server: This does cleanup that has to happen whether the effect is being removed locally or due to replication */
 void FActiveGameplayEffectsContainer::InternalOnActiveGameplayEffectRemoved(const FActiveGameplayEffect& Effect)
 {
+	// Remove our tag requirements from the dependancy map
+	RemoveActiveEffectTagDependancy(Effect.Spec.Def->OngoingTagRequirements.IgnoreTags, Effect.Handle);
+	RemoveActiveEffectTagDependancy(Effect.Spec.Def->OngoingTagRequirements.RequireTags, Effect.Handle);
+
+	if (Effect.Spec.Def)
+	{
+		RemoveActiveGameplayEffectGrantedTagsAndModifiers(Effect);
+	}
+	else
+	{
+		ABILITY_LOG(Warning, TEXT("InternalOnActiveGameplayEffectRemoved called with no GameplayEffect: %s"), *Effect.Handle.ToString());
+	}
+
 	Effect.OnRemovedDelegate.Broadcast();
+}
 
+void FActiveGameplayEffectsContainer::RemoveActiveGameplayEffectGrantedTagsAndModifiers(const FActiveGameplayEffect& Effect)
+{
 	// Update AttributeAggregators: remove mods from this ActiveGE Handle
-	for(const FGameplayModifierInfo& Mod : Effect.Spec.Def->Modifiers)
+	if (Effect.Spec.GetPeriod() <= UGameplayEffect::NO_PERIOD)
 	{
-		if(Mod.Attribute.IsValid())
+		for(const FGameplayModifierInfo& Mod : Effect.Spec.Def->Modifiers)
 		{
-			FAggregatorRef* RefPtr = AttributeAggregatorMap.Find(Mod.Attribute);
-			if(RefPtr)
+			if(Mod.Attribute.IsValid())
 			{
-				RefPtr->Get()->RemoveMod(Effect.Handle);
+				FAggregatorRef* RefPtr = AttributeAggregatorMap.Find(Mod.Attribute);
+				if(RefPtr)
+				{
+					RefPtr->Get()->RemoveMod(Effect.Handle);
+				}
 			}
 		}
 	}
 
-	// Remove our tag requirements from the dependancy map
-	RemoveActiveEffectTagDependancy(Effect.Spec.Def->OngoingTagRequirements.IgnoreTags, Effect.Handle);
-	RemoveActiveEffectTagDependancy(Effect.Spec.Def->OngoingTagRequirements.RequireTags, Effect.Handle);
+	// Update gameplaytag count and broadcast delegate if we are at 0
+	IGameplayTagsModule& GameplayTagsModule = IGameplayTagsModule::Get();
+	Owner->UpdateTagMap(Effect.Spec.Def->InheritableOwnedTagsContainer.CombinedTags, -1);
 
-	if (Effect.Spec.Def)
-	{
-		// Update gameplaytag count and broadcast delegate if we are at 0
-		IGameplayTagsModule& GameplayTagsModule = IGameplayTagsModule::Get();
-		Owner->UpdateTagMap(Effect.Spec.Def->InheritableOwnedTagsContainer.CombinedTags, -1);
+	Owner->UpdateTagMap(Effect.Spec.DynamicGrantedTags, -1);
 
-		for (const FGameplayEffectCue& Cue : Effect.Spec.Def->GameplayCues)
-		{
-			Owner->UpdateTagMap(Cue.GameplayCueTags, -1);
-		}
-	}
-	else
+	for (const FGameplayEffectCue& Cue : Effect.Spec.Def->GameplayCues)
 	{
-		ABILITY_LOG(Warning, TEXT("InternalOnActiveGameplayEffectRemoved called with no GameplayEffect: %s"), *Effect.Handle.ToString());
+		Owner->UpdateTagMap(Cue.GameplayCueTags, -1);
 	}
 }
 
@@ -1928,7 +1941,8 @@ bool FActiveGameplayEffectQuery::Matches(const FActiveGameplayEffect& Effect) co
 {
 	if (TagContainer)
 	{
-		if (!Effect.Spec.Def->InheritableOwnedTagsContainer.CombinedTags.MatchesAny(*TagContainer, false))
+		if (!Effect.Spec.Def->InheritableOwnedTagsContainer.CombinedTags.MatchesAny(*TagContainer, false) &&
+			Effect.Spec.DynamicGrantedTags.MatchesAny(*TagContainer, false))
 		{
 			return false;
 		}
diff --git a/Engine/Source/Runtime/GameplayAbilities/Public/Abilities/GameplayAbilityTypes.h b/Engine/Source/Runtime/GameplayAbilities/Public/Abilities/GameplayAbilityTypes.h
--- a/Engine/Source/Runtime/GameplayAbilities/Public/Abilities/GameplayAbilityTypes.h
+++ b/Engine/Source/Runtime/GameplayAbilities/Public/Abilities/GameplayAbilityTypes.h
@@ -474,19 +474,6 @@ struct GAMEPLAYABILITIES_API FAttributeDefaults
 	class UDataTable*	DefaultStartingTable;
 };
 
-// Still WIP: not being used yet
-USTRUCT()
-struct GAMEPLAYABILITIES_API FGameplayAbilityCooldownInfo
-{
-	GENERATED_USTRUCT_BODY()
-
-	UPROPERTY(EditDefaultsOnly, Category = Cooldowns)
-	FScalableFloat	Cooldown;
-
-	UPROPERTY(EditDefaultsOnly, Category = Cooldowns)
-	FGameplayTagContainer CooldownTags;
-};
-
 /** Used for cleaning up predicted data on network clients */
 DECLARE_MULTICAST_DELEGATE(FAbilitySystemComponentPredictionKeyClear);
 
diff --git a/Engine/Source/Runtime/GameplayAbilities/Public/GameplayEffect.h b/Engine/Source/Runtime/GameplayAbilities/Public/GameplayEffect.h
--- a/Engine/Source/Runtime/GameplayAbilities/Public/GameplayEffect.h
+++ b/Engine/Source/Runtime/GameplayAbilities/Public/GameplayEffect.h
@@ -874,15 +874,16 @@ struct GAMEPLAYABILITIES_API FGameplayEffectSpec
 	bool bTopOfStack;
 
 	// Captured Source Tags on GameplayEffectSpec creation.	
-	UPROPERTY()
+	UPROPERTY(NotReplicated)
 	FTagContainerAggregator	CapturedSourceTags;
 
-	// Tags from the target, captured during execute?
-	// Not sure if this is the right place or not, or how 'contextual' tags get injected (headshot, crit, etc)
-	// Does this overlap with info in the GameplayEffectContextHandle? Should it just all live there?
-	// For now - putting it here so places have one spot to pull from. Need to figure this out long term.
-	UPROPERTY()
+	// Tags from the target, captured during execute	
+	UPROPERTY(NotReplicated)
 	FTagContainerAggregator	CapturedTargetTags;
+
+	/** Tags that are granted and that did not come from the UGameplayEffect def. These are replicated. */
+	UPROPERTY()
+	FGameplayTagContainer DynamicGrantedTags;
 	
 	UPROPERTY()
 	TArray<FModifierSpec> Modifiers;
@@ -1179,6 +1180,9 @@ private:
 	void InternalOnActiveGameplayEffectAdded(FActiveGameplayEffect& Effect);
 	void InternalOnActiveGameplayEffectRemoved(const FActiveGameplayEffect& Effect);
 
+	void RemoveActiveGameplayEffectGrantedTagsAndModifiers(const FActiveGameplayEffect& Effect);
+	void AddActiveGameplayEffectGrantedTagsAndModifiers(FActiveGameplayEffect& Effect);
+
 	/** Updates tag dependancy map when a GameplayEffect is removed */
 	void RemoveActiveEffectTagDependancy(const FGameplayTagContainer& Tags, FActiveGameplayEffectHandle Handle);
 
