missed file and const correctness

diff --git a/Engine/Source/Runtime/GameplayAbilities/Private/Abilities/GameplayAbility.cpp b/Engine/Source/Runtime/GameplayAbilities/Private/Abilities/GameplayAbility.cpp
--- a/Engine/Source/Runtime/GameplayAbilities/Private/Abilities/GameplayAbility.cpp
+++ b/Engine/Source/Runtime/GameplayAbilities/Private/Abilities/GameplayAbility.cpp
@@ -388,7 +388,7 @@ UGameplayEffect* UGameplayAbility::GetCostGameplayEffect() const
 
 bool UGameplayAbility::CheckCooldown(const FGameplayAbilitySpecHandle Handle, const FGameplayAbilityActorInfo* ActorInfo) const
 {
-	FGameplayTagContainer* CooldownTags = GetCooldownTags();
+	const FGameplayTagContainer* CooldownTags = GetCooldownTags();
 	if (CooldownTags)
 	{
 		check(ActorInfo->AbilitySystemComponent.IsValid());
@@ -434,8 +434,8 @@ float UGameplayAbility::GetCooldownTimeRemaining(const FGameplayAbilityActorInfo
 	SCOPE_CYCLE_COUNTER(STAT_GameplayAbilityGetCooldownTimeRemaining);
 
 	check(ActorInfo->AbilitySystemComponent.IsValid());
-	UGameplayEffect* CooldownGE = GetCooldownGameplayEffect();
-	if (CooldownGE)
+	const FGameplayTagContainer* CooldownTags = GetCooldownTags();
+	if (CooldownTags)
 	{
 		TArray< float > Durations = ActorInfo->AbilitySystemComponent->GetActiveEffectsTimeRemaining(FActiveGameplayEffectQuery(CooldownTags));
 		if (Durations.Num() > 0)
@@ -448,7 +448,7 @@ float UGameplayAbility::GetCooldownTimeRemaining(const FGameplayAbilityActorInfo
 	return 0.f;
 }
 
-void UGameplayAbility::GetCooldownTimeRemainingAndDuration(const FGameplayAbilityActorInfo* ActorInfo, float& TimeRemaining, float& CooldownDuration) const
+void UGameplayAbility::GetCooldownTimeRemainingAndDuration(FGameplayAbilitySpecHandle Handle, const FGameplayAbilityActorInfo* ActorInfo, float& TimeRemaining, float& CooldownDuration) const
 {
 	SCOPE_CYCLE_COUNTER(STAT_GameplayAbilityGetCooldownTimeRemainingAndDuration);
 
@@ -457,8 +457,8 @@ void UGameplayAbility::GetCooldownTimeRemainingAndDuration(const FGameplayAbilit
 	TimeRemaining = 0.f;
 	CooldownDuration = 0.f;
 	
-	FGameplayTagContainer* CooldownTags = GetCooldownTags();
-	if (CooldownGE)
+	const FGameplayTagContainer* CooldownTags = GetCooldownTags();
+	if (CooldownTags)
 	{
 		TArray< float > DurationRemaining = ActorInfo->AbilitySystemComponent->GetActiveEffectsTimeRemaining(FActiveGameplayEffectQuery(CooldownTags));
 		if (DurationRemaining.Num() > 0)
@@ -482,7 +482,7 @@ void UGameplayAbility::GetCooldownTimeRemainingAndDuration(const FGameplayAbilit
 	}
 }
 
-FGameplayTagContainer* UGameplayAbility::GetCooldownTags() const
+const FGameplayTagContainer* UGameplayAbility::GetCooldownTags() const
 {
 	UGameplayEffect* CDGE =GetCooldownGameplayEffect();
 	return CDGE ? &CDGE->InheritableClearTagsContainer.CombinedTags : nullptr;
diff --git a/Engine/Source/Runtime/GameplayAbilities/Public/Abilities/GameplayAbility.h b/Engine/Source/Runtime/GameplayAbilities/Public/Abilities/GameplayAbility.h
--- a/Engine/Source/Runtime/GameplayAbilities/Public/Abilities/GameplayAbility.h
+++ b/Engine/Source/Runtime/GameplayAbilities/Public/Abilities/GameplayAbility.h
@@ -128,7 +128,9 @@ public:
 	virtual float GetCooldownTimeRemaining(const FGameplayAbilityActorInfo* ActorInfo) const;
 
 	/** Returns the time in seconds remaining on the currently active cooldown and the original duration for this cooldown. */
-	virtual void GetCooldownTimeRemainingAndDuration(const FGameplayAbilityActorInfo* ActorInfo, float& TimeRemaining, float& CooldownDuration) const;
+	virtual void GetCooldownTimeRemainingAndDuration(FGameplayAbilitySpecHandle Handle, const FGameplayAbilityActorInfo* ActorInfo, float& TimeRemaining, float& CooldownDuration) const;
+
+	virtual const FGameplayTagContainer* GetCooldownTags() const;
 		
 	EGameplayAbilityInstancingPolicy::Type GetInstancingPolicy() const
 	{
@@ -503,15 +505,7 @@ protected:
 	//	Cooldowns
 	//
 	// ----------------------------------------------------------------------------------------------------------------
-		
-
-	/** Cooldown tags that we check for. If we have cooldowns with these tags, we cannot activate this ability */
-	UPROPERTY()
-	FGameplayTagContainer CheckedCooldownTags;
-
-	UPROPERTY()
-	TArray<FGameplayAbilityCooldownInfo> AppliedCooldowns;
-	
+			
 	/** Deprecated? This GameplayEffect represents the cooldown. It will be applied when the ability is committed and the ability cannot be used again until it is expired. */
 	UPROPERTY(EditDefaultsOnly, Category = Cooldowns)
 	TSubclassOf<class UGameplayEffect> CooldownGameplayEffectClass;
@@ -519,7 +513,6 @@ protected:
 	/** Deprecated. Use CooldownGameplayEffectClass instead */
 	UPROPERTY(VisibleDefaultsOnly, Category=Deprecated)
 	class UGameplayEffect* CooldownGameplayEffect;
-
 	
 	// ----------------------------------------------------------------------------------------------------------------
 	//
