Orion cooldown experiement
-Define cooldowns in abilities instead of via GEs

diff --git a/Engine/Source/Runtime/GameplayAbilities/Private/GameplayEffect.cpp b/Engine/Source/Runtime/GameplayAbilities/Private/GameplayEffect.cpp
--- a/Engine/Source/Runtime/GameplayAbilities/Private/GameplayEffect.cpp
+++ b/Engine/Source/Runtime/GameplayAbilities/Private/GameplayEffect.cpp
@@ -1942,7 +1942,7 @@ bool FActiveGameplayEffectQuery::Matches(const FActiveGameplayEffect& Effect) co
 	if (TagContainer)
 	{
 		if (!Effect.Spec.Def->InheritableOwnedTagsContainer.CombinedTags.MatchesAny(*TagContainer, false) &&
-			Effect.Spec.DynamicGrantedTags.MatchesAny(*TagContainer, false))
+			!Effect.Spec.DynamicGrantedTags.MatchesAny(*TagContainer, false))
 		{
 			return false;
 		}
