updated blueprint callable functions that take TSubclassOf to gameplay effect class instead of a raw pointer to gameplay effect objects

Adding a new version of ApplyGameplayEffectToTarget that takes a TSubclassOf<UGameplayEffect> instead of a pointer to the data structure. Deprecating the old blueprint callable function.

diff --git a/Engine/Source/Runtime/GameplayAbilities/Private/AbilitySystemComponent.cpp b/Engine/Source/Runtime/GameplayAbilities/Private/AbilitySystemComponent.cpp
--- a/Engine/Source/Runtime/GameplayAbilities/Private/AbilitySystemComponent.cpp
+++ b/Engine/Source/Runtime/GameplayAbilities/Private/AbilitySystemComponent.cpp
@@ -203,6 +203,17 @@ FGameplayEffectContextHandle UAbilitySystemComponent::GetEffectContext() const
 	return Context;
 }
 
+FActiveGameplayEffectHandle UAbilitySystemComponent::BP_ApplyGameplayEffectToTarget(TSubclassOf<UGameplayEffect> GameplayEffectClass, UAbilitySystemComponent *Target, float Level, FGameplayEffectContextHandle Context)
+{
+	if (GameplayEffectClass)
+	{
+		UGameplayEffect* GameplayEffect = GameplayEffectClass->GetDefaultObject<UGameplayEffect>();
+		return ApplyGameplayEffectToTarget(GameplayEffect, Target, Level, Context);
+	}
+
+	return FActiveGameplayEffectHandle();
+}
+
 FActiveGameplayEffectHandle UAbilitySystemComponent::K2_ApplyGameplayEffectToTarget(UGameplayEffect *GameplayEffect, UAbilitySystemComponent *Target, float Level, FGameplayEffectContextHandle Context)
 {
 	return ApplyGameplayEffectToTarget(GameplayEffect, Target, Level, Context);
diff --git a/Engine/Source/Runtime/GameplayAbilities/Public/AbilitySystemComponent.h b/Engine/Source/Runtime/GameplayAbilities/Public/AbilitySystemComponent.h
--- a/Engine/Source/Runtime/GameplayAbilities/Public/AbilitySystemComponent.h
+++ b/Engine/Source/Runtime/GameplayAbilities/Public/AbilitySystemComponent.h
@@ -307,6 +307,9 @@ class GAMEPLAYABILITIES_API UAbilitySystemComponent : public UActorComponent, pu
 	FOnActiveGameplayEffectRemoved* OnGameplayEffectRemovedDelegate(FActiveGameplayEffectHandle Handle);
 
 	UFUNCTION(BlueprintCallable, Category = GameplayEffects, meta=(FriendlyName = "ApplyGameplayEffectToTarget"))
+	FActiveGameplayEffectHandle BP_ApplyGameplayEffectToTarget(TSubclassOf<UGameplayEffect> GameplayEffectClass, UAbilitySystemComponent *Target, float Level, FGameplayEffectContextHandle Context);
+
+	UFUNCTION(BlueprintCallable, Category = GameplayEffects, meta=(FriendlyName = "ApplyGameplayEffectToTarget"), meta=(DeprecatedFunction, DeprecationMessage = "Use new ApplyGameplayEffectToSelf"))
 	FActiveGameplayEffectHandle K2_ApplyGameplayEffectToTarget(UGameplayEffect *GameplayEffect, UAbilitySystemComponent *Target, float Level, FGameplayEffectContextHandle Context);
 
 	FActiveGameplayEffectHandle ApplyGameplayEffectToTarget(UGameplayEffect *GameplayEffect, UAbilitySystemComponent *Target, float Level = UGameplayEffect::INVALID_LEVEL, FGameplayEffectContextHandle Context = FGameplayEffectContextHandle(), FPredictionKey PredictionKey = FPredictionKey());
