Additional debug info for active gameplayeffects

diff --git a/Engine/Source/Runtime/GameplayAbilities/Private/AbilitySystemComponent.cpp b/Engine/Source/Runtime/GameplayAbilities/Private/AbilitySystemComponent.cpp
--- a/Engine/Source/Runtime/GameplayAbilities/Private/AbilitySystemComponent.cpp
+++ b/Engine/Source/Runtime/GameplayAbilities/Private/AbilitySystemComponent.cpp
@@ -952,6 +952,13 @@ void UAbilitySystemComponent::OnMagnitudeDependancyChange(FActiveGameplayEffectH
 	ActiveGameplayEffects.OnMagnitudeDependancyChange(Handle, ChangedAggregator);
 }
 
+FString ASC_CleanupName(FString Str)
+{
+	Str.RemoveFromStart(TEXT("Default__"));
+	Str.RemoveFromEnd(TEXT("_c"));
+	return Str;
+}
+
 void UAbilitySystemComponent::DisplayDebug(class UCanvas* Canvas, const class FDebugDisplayInfo& DebugDisplay, float& YL, float& YPos)
 {
 
@@ -977,9 +984,10 @@ void UAbilitySystemComponent::DisplayDebug(class UCanvas* Canvas, const class FD
 
 	// -------------------------------------------------------------
 
-	if (bShowGameplayEffects)
-	{
 
+	if (bShowGameplayEffects || bShowAttributes)
+	{
+		// Draw the attribute aggregator map.
 		for (auto It = ActiveGameplayEffects.AttributeAggregatorMap.CreateConstIterator(); It; ++It)
 		{
 			FGameplayAttribute Attribute = It.Key();
@@ -1021,6 +1029,61 @@ void UAbilitySystemComponent::DisplayDebug(class UCanvas* Canvas, const class FD
 
 	// -------------------------------------------------------------
 
+	if (bShowGameplayEffects)
+	{
+		for (FActiveGameplayEffect& ActiveGE : ActiveGameplayEffects.GameplayEffects)
+		{
+			Canvas->SetDrawColor(FColor::White);
+
+			FString DurationStr = TEXT("Infinite Duration ");
+			if (ActiveGE.GetDuration() > 0.f)
+			{
+				DurationStr = FString::Printf(TEXT("Duration: %.2f. Remaining: %.2f "), ActiveGE.GetDuration(), ActiveGE.GetTimeRemaining(GetWorld()->GetTimeSeconds()));
+			}
+			if (ActiveGE.GetPeriod() > 0.f)
+			{
+				DurationStr += FString::Printf(TEXT("Period: %.2f"), ActiveGE.GetPeriod());
+			}
+
+			Canvas->SetDrawColor(ActiveGE.IsInhibited ? FColor(128, 128, 128): FColor::White );
+
+			YPos += Canvas->DrawText(GEngine->GetTinyFont(), FString::Printf(TEXT("%s %s"), *ASC_CleanupName(GetNameSafe(ActiveGE.Spec.Def)), *DurationStr ), 4.f, YPos);
+
+			FGameplayTagContainer GrantedTags;
+			ActiveGE.Spec.GetAllGrantedTags(GrantedTags);
+			if (GrantedTags.Num() > 0)
+			{
+				YPos += Canvas->DrawText(GEngine->GetTinyFont(), FString::Printf(TEXT("Granted Tags: %s"), *GrantedTags.ToStringSimple() ), 7.f, YPos);
+			}
+
+			for (int32 ModIdx=0; ModIdx < ActiveGE.Spec.Modifiers.Num(); ++ModIdx)
+			{
+				const FModifierSpec& ModSpec = ActiveGE.Spec.Modifiers[ModIdx];
+				const FGameplayModifierInfo& ModInfo = ActiveGE.Spec.Def->Modifiers[ModIdx];
+
+				FAggregatorMod TempMod;
+				TempMod.SourceTagReqs = &ModInfo.SourceTags;
+				TempMod.TargetTagReqs = &ModInfo.TargetTags;
+
+				FAggregatorEvaluateParameters EmptyParams;
+				bool IsActivelyModifyingAttribute = TempMod.Qualifies(EmptyParams);
+
+				if (IsActivelyModifyingAttribute == false)
+				{
+					Canvas->SetDrawColor(FColor(128, 128, 128) );
+				}
+
+				YPos += Canvas->DrawText(GEngine->GetTinyFont(), FString::Printf(TEXT("Mod: %s. %s. %.2f"), *ModInfo.Attribute.GetName(), *EGameplayModOpToString(ModInfo.ModifierOp), ModSpec.GetEvaluatedMagnitude() ), 7.f, YPos);
+
+				Canvas->SetDrawColor(ActiveGE.IsInhibited ? FColor(128, 128, 128): FColor::White );
+			}
+
+			YPos += YL;
+		}
+	}
+
+	// -------------------------------------------------------------
+
 	if (bShowAttributes)
 	{
 		Canvas->SetDrawColor(FColor::White);
@@ -1053,7 +1116,7 @@ void UAbilitySystemComponent::DisplayDebug(class UCanvas* Canvas, const class FD
 				continue;
 
 			Canvas->SetDrawColor(AbilitySpec.IsActive() ? FColor::Yellow : FColor(128, 128, 128));
-			YPos += Canvas->DrawText(GEngine->GetTinyFont(), FString::Printf(TEXT("%s"), *GetNameSafe(AbilitySpec.Ability)), 4.f, YPos);
+			YPos += Canvas->DrawText(GEngine->GetTinyFont(), FString::Printf(TEXT("%s"), *ASC_CleanupName(GetNameSafe(AbilitySpec.Ability))), 4.f, YPos);
 	
 			if (AbilitySpec.IsActive())
 			{
diff --git a/Engine/Source/Runtime/GameplayAbilities/Private/GameplayEffect.cpp b/Engine/Source/Runtime/GameplayAbilities/Private/GameplayEffect.cpp
--- a/Engine/Source/Runtime/GameplayAbilities/Private/GameplayEffect.cpp
+++ b/Engine/Source/Runtime/GameplayAbilities/Private/GameplayEffect.cpp
@@ -539,6 +539,12 @@ void FGameplayEffectSpec::SetContext(FGameplayEffectContextHandle NewEffectConte
 	}
 }
 
+void FGameplayEffectSpec::GetAllGrantedTags(OUT FGameplayTagContainer& Container) const
+{
+	Container.AppendTags(DynamicGrantedTags);
+	Container.AppendTags(Def->InheritableOwnedTagsContainer.CombinedTags);
+}
+
 // --------------------------------------------------------------------------------------------------------------------------------------------------------
 //
 //	FGameplayEffectAttributeCaptureSpec
diff --git a/Engine/Source/Runtime/GameplayAbilities/Public/GameplayEffect.h b/Engine/Source/Runtime/GameplayAbilities/Public/GameplayEffect.h
--- a/Engine/Source/Runtime/GameplayAbilities/Public/GameplayEffect.h
+++ b/Engine/Source/Runtime/GameplayAbilities/Public/GameplayEffect.h
@@ -850,6 +850,8 @@ struct GAMEPLAYABILITIES_API FGameplayEffectSpec
 		return EffectContext;
 	}
 
+	void GetAllGrantedTags(OUT FGameplayTagContainer& Container) const;
+
 	// The duration in seconds of this effect
 	// instantaneous effects should have a duration of UGameplayEffect::INSTANT_APPLICATION
 	// effects that last forever should have a duration of UGameplayEffect::INFINITE_DURATION
