NULL check

diff --git a/Engine/Source/Runtime/GameplayAbilities/Private/Abilities/GameplayAbility.cpp b/Engine/Source/Runtime/GameplayAbilities/Private/Abilities/GameplayAbility.cpp
--- a/Engine/Source/Runtime/GameplayAbilities/Private/Abilities/GameplayAbility.cpp
+++ b/Engine/Source/Runtime/GameplayAbilities/Private/Abilities/GameplayAbility.cpp
@@ -265,8 +265,11 @@ void UGameplayAbility::EndAbility(const FGameplayAbilitySpecHandle Handle, const
 
 	// Stop any timers or latent actions for the ability
 	UWorld* MyWorld = ActorInfo->AbilitySystemComponent->GetOwner()->GetWorld();
-	MyWorld->GetLatentActionManager().RemoveActionsForObject(this);
-	MyWorld->GetTimerManager().ClearAllTimersForObject(this);
+	if (MyWorld)
+	{
+		MyWorld->GetLatentActionManager().RemoveActionsForObject(this);
+		MyWorld->GetTimerManager().ClearAllTimersForObject(this);
+	}
 
 	// Execute our delegate and unbind it, as we are no longer active and listeners can re-register when we become active again.
 	OnGameplayAbilityEnded.ExecuteIfBound(this);
