GameplayEffect - support for SetByCaller modifier magnitudes

diff --git a/Engine/Source/Runtime/GameplayAbilities/Private/GameplayEffect.cpp b/Engine/Source/Runtime/GameplayAbilities/Private/GameplayEffect.cpp
--- a/Engine/Source/Runtime/GameplayAbilities/Private/GameplayEffect.cpp
+++ b/Engine/Source/Runtime/GameplayAbilities/Private/GameplayEffect.cpp
@@ -227,6 +227,10 @@ bool FGameplayEffectModifierMagnitude::AttemptCalculateMagnitude(const FGameplay
 				OutCalculatedMagnitude = CalcCDO->CalculateMagnitude(InRelevantSpec);
 			}
 			break;
+
+			case EGameplayEffectMagnitudeCalculation::SetByCaller:
+				// Add check that we've been set?
+				break;
 		}
 	}
 
@@ -545,6 +549,24 @@ void FGameplayEffectSpec::GetAllGrantedTags(OUT FGameplayTagContainer& Container
 	Container.AppendTags(Def->InheritableOwnedTagsContainer.CombinedTags);
 }
 
+void FGameplayEffectSpec::SetModifierMagnitude(int32 ModIdx, float EvaluatedMagnitude)
+{
+	if (Def->Modifiers.IsValidIndex(EvaluatedMagnitude) == false)
+	{
+		ABILITY_LOG(Error, TEXT("FGameplayEffectSpec::SetModifierMagnitude called on invalid index %d for Def: %s"), ModIdx, *Def->GetName());
+		return;
+	}
+
+	if (Def->Modifiers[ModIdx].ModifierMagnitude.GetMagnitudeCalculationType() != EGameplayEffectMagnitudeCalculation::SetByCaller)
+	{
+		ABILITY_LOG(Error, TEXT("FGameplayEffectSpec::SetModifierMagnitude called on index %d for Def: %s - this modifier not SetByCaller!"), ModIdx, *Def->GetName());
+		return;
+	}
+	
+	Modifiers[ModIdx].EvaluatedMagnitude = EvaluatedMagnitude;
+}
+
+
 // --------------------------------------------------------------------------------------------------------------------------------------------------------
 //
 //	FGameplayEffectAttributeCaptureSpec
diff --git a/Engine/Source/Runtime/GameplayAbilities/Public/GameplayEffect.h b/Engine/Source/Runtime/GameplayAbilities/Public/GameplayEffect.h
--- a/Engine/Source/Runtime/GameplayAbilities/Public/GameplayEffect.h
+++ b/Engine/Source/Runtime/GameplayAbilities/Public/GameplayEffect.h
@@ -89,9 +89,14 @@ struct FGameplayEffectStackingCallbacks
 UENUM()
 enum class EGameplayEffectMagnitudeCalculation : uint8
 {
-	ScalableFloat,			// Use a simple, scalable float for the calculation
-	AttributeBased,			// Perform a calculation based upon an attribute
-	CustomCalculationClass	// Perform a custom calculation, capable of capturing and acting on multiple attributes, in either BP or native
+	/** Use a simple, scalable float for the calculation */
+	ScalableFloat,
+	/** Perform a calculation based upon an attribute */
+	AttributeBased,
+	/** Perform a custom calculation, capable of capturing and acting on multiple attributes, in either BP or native */
+	CustomCalculationClass,	
+	/** This magnitude will be set explicity by that code/blueprint that creates the spec. */
+	SetByCaller,
 };
 
 /** Enumeration outlining the possible attribute based float calculation policies */
@@ -211,6 +216,8 @@ public:
 	 */
 	void GetAttributeCaptureDefinitions(OUT TArray<FGameplayEffectAttributeCaptureDefinition>& OutCaptureDefs) const;
 
+	EGameplayEffectMagnitudeCalculation GetMagnitudeCalculationType() const { return MagnitudeCalculationType; }
+
 protected:
 
 	/** Type of calculation to perform to derive the magnitude */
@@ -852,6 +859,9 @@ struct GAMEPLAYABILITIES_API FGameplayEffectSpec
 
 	void GetAllGrantedTags(OUT FGameplayTagContainer& Container) const;
 
+	/** Sets the magnitude of one of our modifiers. This must be done on modifiers which are EGameplayEffectMagnitudeCalculation::SetByCaller  */
+	void SetModifierMagnitude(int32 ModIdx, float EvaluatedMagnitude);
+
 	// The duration in seconds of this effect
 	// instantaneous effects should have a duration of UGameplayEffect::INSTANT_APPLICATION
 	// effects that last forever should have a duration of UGameplayEffect::INFINITE_DURATION
