Fix for bug with compile button breaking GCN lookup.

diff --git a/Engine/Source/Runtime/GameplayAbilities/Private/GameplayCueManager.cpp b/Engine/Source/Runtime/GameplayAbilities/Private/GameplayCueManager.cpp
--- a/Engine/Source/Runtime/GameplayAbilities/Private/GameplayCueManager.cpp
+++ b/Engine/Source/Runtime/GameplayAbilities/Private/GameplayCueManager.cpp
@@ -57,7 +57,7 @@ void UGameplayCueManager::HandleGameplayCueNotify_Internal(AActor* TargetActor,
 		FGameplayCueNotifyData& CueData = GameplayCueData[DataIdx];
 
 		// If object is not loaded yet
-		if (CueData.LoadedGameplayCueNotify == nullptr)
+		if (CueData.LoadedGameplayCueClass == nullptr)
 		{
 			// Ignore removed events if this wasn't already loaded (only call Removed if we handled OnActive/WhileActive)
 			if (EventType == EGameplayCueEvent::Removed)
@@ -77,29 +77,20 @@ void UGameplayCueManager::HandleGameplayCueNotify_Internal(AActor* TargetActor,
 			}
 			else
 			{
-				// Found it - did we load a non-instance version (UGameplayCueNotify_Static) or an instanced version (AGameplayCueNotify_Actor)
-				CueData.LoadedGameplayCueNotify = Cast<AGameplayCueNotify_Actor>(FoundObject);
-				if (!CueData.LoadedGameplayCueNotify)
+				// Found it - did we load a non-instanced version (UGameplayCueNotify_Static) or an instanced version (AGameplayCueNotify_Actor)
+				UObject* LoadedGameplayCueNotify = Cast<AGameplayCueNotify_Actor>(FoundObject);
+				if (!LoadedGameplayCueNotify)
 				{
-					CueData.LoadedGameplayCueNotify = Cast<UGameplayCueNotify_Static>(FoundObject);
+					//Try the other class
+					LoadedGameplayCueNotify = Cast<UGameplayCueNotify_Static>(FoundObject);
 				}
-				if (CueData.LoadedGameplayCueNotify == nullptr)
+				if (!LoadedGameplayCueNotify)
 				{
 					// Not a dataasset - maybe a blueprint
-					
 					UBlueprint* GameplayCueBlueprint = Cast<UBlueprint>(FoundObject);
 					if (GameplayCueBlueprint && GameplayCueBlueprint->GeneratedClass)
 					{
-						CueData.LoadedGameplayCueNotify = Cast<AGameplayCueNotify_Actor>(GameplayCueBlueprint->GeneratedClass->ClassDefaultObject);
-						if (!CueData.LoadedGameplayCueNotify)
-						{
-							CueData.LoadedGameplayCueNotify = Cast<UGameplayCueNotify_Static>(GameplayCueBlueprint->GeneratedClass->ClassDefaultObject);
-						}
-						else
-						{
-							ABILITY_LOG(Warning, TEXT("GameplayCueNotify %s loaded blueprint object %s that is not a GameplayCueNotify"), *CueData.GameplayCueNotifyObj.ToString(), *FoundObject->GetName());
-							return;
-						}
+						CueData.LoadedGameplayCueClass = GameplayCueBlueprint->GeneratedClass;
 					}
 					else
 					{
@@ -111,7 +102,7 @@ void UGameplayCueManager::HandleGameplayCueNotify_Internal(AActor* TargetActor,
 		}
 
 		// Handle the Notify if we found something
-		if (UGameplayCueNotify_Static* NonInstancedCue = Cast<UGameplayCueNotify_Static>(CueData.LoadedGameplayCueNotify))
+		if (UGameplayCueNotify_Static* NonInstancedCue = Cast<UGameplayCueNotify_Static>(CueData.LoadedGameplayCueClass->ClassDefaultObject))
 		{
 			if (NonInstancedCue->HandlesEvent(EventType))
 			{
@@ -127,7 +118,7 @@ void UGameplayCueManager::HandleGameplayCueNotify_Internal(AActor* TargetActor,
 				HandleGameplayCueNotify_Internal(TargetActor, CueData.ParentDataIdx, EventType, Parameters);
 			}
 		}
-		else if (AGameplayCueNotify_Actor* InstancedCue = Cast<AGameplayCueNotify_Actor>(CueData.LoadedGameplayCueNotify))
+		else if (AGameplayCueNotify_Actor* InstancedCue = Cast<AGameplayCueNotify_Actor>(CueData.LoadedGameplayCueClass->ClassDefaultObject))
 		{
 			AGameplayCueNotify_Actor* SpawnedInstancedCue = nullptr;
 			if (auto InnerMap = NotifyMapActor.Find(TargetActor))
diff --git a/Engine/Source/Runtime/GameplayAbilities/Public/GameplayCueManager.h b/Engine/Source/Runtime/GameplayAbilities/Public/GameplayCueManager.h
--- a/Engine/Source/Runtime/GameplayAbilities/Public/GameplayCueManager.h
+++ b/Engine/Source/Runtime/GameplayAbilities/Public/GameplayCueManager.h
@@ -51,7 +51,7 @@ struct FGameplayCueNotifyData
 	GENERATED_USTRUCT_BODY()
 
 	FGameplayCueNotifyData()
-	: LoadedGameplayCueNotify(nullptr)
+	: LoadedGameplayCueClass(nullptr)
 	, ParentDataIdx( INDEX_NONE )
 	{
 	}
@@ -59,8 +59,8 @@ struct FGameplayCueNotifyData
 	UPROPERTY(VisibleDefaultsOnly, Category=GameplayCue, meta=(AllowedClasses="GameplayCueNotify"))
 	FStringAssetReference	GameplayCueNotifyObj;
 
-	UPROPERTY(transient)
-	UObject*				LoadedGameplayCueNotify;
+	UPROPERTY()
+	UClass*					LoadedGameplayCueClass;
 
 	FGameplayTag			GameplayCueTag;
 
