missed checkin

diff --git a/Engine/Source/Runtime/GameplayAbilities/Private/GameplayEffect.cpp b/Engine/Source/Runtime/GameplayAbilities/Private/GameplayEffect.cpp
--- a/Engine/Source/Runtime/GameplayAbilities/Private/GameplayEffect.cpp
+++ b/Engine/Source/Runtime/GameplayAbilities/Private/GameplayEffect.cpp
@@ -551,7 +551,7 @@ void FGameplayEffectSpec::GetAllGrantedTags(OUT FGameplayTagContainer& Container
 
 void FGameplayEffectSpec::SetModifierMagnitude(int32 ModIdx, float EvaluatedMagnitude)
 {
-	if (Def->Modifiers.IsValidIndex(EvaluatedMagnitude) == false)
+	if (Def->Modifiers.IsValidIndex(ModIdx) == false)
 	{
 		ABILITY_LOG(Error, TEXT("FGameplayEffectSpec::SetModifierMagnitude called on invalid index %d for Def: %s"), ModIdx, *Def->GetName());
 		return;
