One more fix for energy consumption

diff --git a/Engine/Source/Runtime/GameplayAbilities/Private/GameplayEffect.cpp b/Engine/Source/Runtime/GameplayAbilities/Private/GameplayEffect.cpp
--- a/Engine/Source/Runtime/GameplayAbilities/Private/GameplayEffect.cpp
+++ b/Engine/Source/Runtime/GameplayAbilities/Private/GameplayEffect.cpp
@@ -200,8 +200,6 @@ bool FGameplayEffectModifierMagnitude::CanCalculateMagnitude(const FGameplayEffe
 
 bool FGameplayEffectModifierMagnitude::AttemptCalculateMagnitude(const FGameplayEffectSpec& InRelevantSpec, OUT float& OutCalculatedMagnitude) const
 {
-	OutCalculatedMagnitude = 0.f;
-	
 	const bool bCanCalc = CanCalculateMagnitude(InRelevantSpec);
 	if (bCanCalc)
 	{
@@ -231,8 +229,16 @@ bool FGameplayEffectModifierMagnitude::AttemptCalculateMagnitude(const FGameplay
 			case EGameplayEffectMagnitudeCalculation::SetByCaller:
 				// Add check that we've been set?
 				break;
+			default:
+				ABILITY_LOG(Error, TEXT("Unknown MagnitudeCalculationType %d in AttemptCalculateMagnitude"), (int32)MagnitudeCalculationType);
+				OutCalculatedMagnitude = 0.f;
+				break;
 		}
 	}
+	else
+	{
+		OutCalculatedMagnitude = 0.f;
+	}
 
 	return bCanCalc;
 }
