Proper checking of hit results in line trace targeting.

diff --git a/Engine/Source/Runtime/GameplayAbilities/Private/Abilities/GameplayAbilityTargetActor_Trace.cpp b/Engine/Source/Runtime/GameplayAbilities/Private/Abilities/GameplayAbilityTargetActor_Trace.cpp
--- a/Engine/Source/Runtime/GameplayAbilities/Private/Abilities/GameplayAbilityTargetActor_Trace.cpp
+++ b/Engine/Source/Runtime/GameplayAbilities/Private/Abilities/GameplayAbilityTargetActor_Trace.cpp
@@ -40,7 +40,7 @@ void AGameplayAbilityTargetActor_Trace::LineTraceWithFilter(FHitResult& ReturnHi
 	while (true)
 	{
 		InWorld->LineTraceSingle(ReturnHitResult, InTraceStart, InTraceEnd, Channel, LocalParams);
-		if (ReturnHitResult.Actor.IsValid() && !InFilterHandle.FilterPassesForActor(ReturnHitResult.Actor))
+		if (ReturnHitResult.bBlockingHit && ReturnHitResult.Actor.IsValid() && !InFilterHandle.FilterPassesForActor(ReturnHitResult.Actor))
 		{
 			LocalParams.AddIgnoredActor(ReturnHitResult.Actor.Get());
 			continue;
