Improved fix, not reusing traces.

diff --git a/Engine/Source/Runtime/GameplayAbilities/Private/Abilities/GameplayAbilityTargetActor_Trace.cpp b/Engine/Source/Runtime/GameplayAbilities/Private/Abilities/GameplayAbilityTargetActor_Trace.cpp
--- a/Engine/Source/Runtime/GameplayAbilities/Private/Abilities/GameplayAbilityTargetActor_Trace.cpp
+++ b/Engine/Source/Runtime/GameplayAbilities/Private/Abilities/GameplayAbilityTargetActor_Trace.cpp
@@ -39,14 +39,16 @@ void AGameplayAbilityTargetActor_Trace::LineTraceWithFilter(FHitResult& ReturnHi
 	FCollisionQueryParams LocalParams = Params;
 	while (true)
 	{
-		InWorld->LineTraceSingle(ReturnHitResult, InTraceStart, InTraceEnd, Channel, LocalParams);
-		if (ReturnHitResult.bBlockingHit && ReturnHitResult.Actor.IsValid() && !InFilterHandle.FilterPassesForActor(ReturnHitResult.Actor))
+		FHitResult TempHitResult;
+		InWorld->LineTraceSingle(TempHitResult, InTraceStart, InTraceEnd, Channel, LocalParams);
+		if (TempHitResult.bBlockingHit && TempHitResult.Actor.IsValid() && !InFilterHandle.FilterPassesForActor(TempHitResult.Actor))
 		{
-			LocalParams.AddIgnoredActor(ReturnHitResult.Actor.Get());
+			LocalParams.AddIgnoredActor(TempHitResult.Actor.Get());
 			continue;
 		}
 		//Either hit something we're not ignoring, or didn't hit anything.
-		return;
+		ReturnHitResult = TempHitResult;
+		break;
 	};
 }
 
@@ -56,14 +58,16 @@ void AGameplayAbilityTargetActor_Trace::SweepWithFilter(FHitResult& ReturnHitRes
 	FCollisionQueryParams LocalParams = Params;
 	while (true)
 	{
-		InWorld->SweepSingle(ReturnHitResult, InTraceStart, InTraceEnd, InRotation, Channel, CollisionShape, LocalParams);
-		if (ReturnHitResult.bBlockingHit && ReturnHitResult.Actor.IsValid() && !InFilterHandle.FilterPassesForActor(ReturnHitResult.Actor))
+		FHitResult TempHitResult;
+		InWorld->SweepSingle(TempHitResult, InTraceStart, InTraceEnd, InRotation, Channel, CollisionShape, LocalParams);
+		if (TempHitResult.bBlockingHit && TempHitResult.Actor.IsValid() && !InFilterHandle.FilterPassesForActor(TempHitResult.Actor))
 		{
 			LocalParams.AddIgnoredActor(ReturnHitResult.Actor.Get());
 			continue;
 		}
 		//Either hit something we're not ignoring, or didn't hit anything.
-		return;
+		ReturnHitResult = TempHitResult;
+		break;
 	};
 }
 
