Prediction pass
-fixup client prediction of attribute mods

diff --git a/Engine/Source/Runtime/GameplayAbilities/Private/GameplayEffect.cpp b/Engine/Source/Runtime/GameplayAbilities/Private/GameplayEffect.cpp
--- a/Engine/Source/Runtime/GameplayAbilities/Private/GameplayEffect.cpp
+++ b/Engine/Source/Runtime/GameplayAbilities/Private/GameplayEffect.cpp
@@ -1036,8 +1036,29 @@ void FActiveGameplayEffectsContainer::OnAttributeAggregatorDirty(FAggregator* Ag
 	// aggregators when that happens.
 	
 	FAggregatorEvaluateParameters EvaluationParameters;
+
+	if (Owner->IsNetSimulating())
+	{
+		// We are a client. The current value of this attribute is the replicated server's "final" value. We dont actually know what the 
+		// server's base value is. But we can calculate it with ReverseEvaluate(). Then, we can call Evaluate with IncludePredictiveMods=true
+		// to apply our mods and get an accurate predicted value.
+		
+		float FinalValue = Owner->GetNumericAttribute(Attribute);
+		float BaseValue = Aggregator->ReverseEvaluate(FinalValue, EvaluationParameters);
+		Aggregator->SetBaseValue(BaseValue, false);
+
+		EvaluationParameters.IncludePredictiveMods = true;
+		ABILITY_LOG(Log, TEXT("Reverse Evaluated %s. FinalValue: %.2f  BaseValue: %.2f "), *Attribute.GetName(), FinalValue, BaseValue);
+	}
+	
 	
 	float NewValue = Aggregator->Evaluate(EvaluationParameters);
+
+	if (EvaluationParameters.IncludePredictiveMods)
+	{
+		ABILITY_LOG(Log, TEXT("After Prediction, FinalValue: %.2f"), NewValue);
+	}
+
 	InternalUpdateNumericalAttribute(Attribute, NewValue, nullptr);
 }
 
@@ -1092,7 +1113,7 @@ void FActiveGameplayEffectsContainer::OnMagnitudeDependancyChange(FActiveGamepla
 
 					if (ModDef.Attribute == Attribute)
 					{
-						Aggregator->AddMod(ModSpec.GetEvaluatedMagnitude(), ModDef.ModifierOp, &ModDef.SourceTags, &ModDef.TargetTags, Handle);
+						Aggregator->AddMod(ModSpec.GetEvaluatedMagnitude(), ModDef.ModifierOp, &ModDef.SourceTags, &ModDef.TargetTags, ActiveEffect->PredictionKey.WasLocallyGenerated(), Handle);
 					}
 				}
 			}
@@ -1100,19 +1121,17 @@ void FActiveGameplayEffectsContainer::OnMagnitudeDependancyChange(FActiveGamepla
 	}
 }
 
-void FActiveGameplayEffectsContainer::SetBaseAttributeValueFromReplication(FGameplayAttribute Attribute, float BaseValue)
+void FActiveGameplayEffectsContainer::SetBaseAttributeValueFromReplication(FGameplayAttribute Attribute, float ServerValue)
 {
-/*
-	FIXME: the approach for client side attribute prediction will need to be slightly rethought. Rather than touching base values we should be able
-	to just apply Mods... this function may just go away?
-
 	FAggregatorRef* RefPtr = AttributeAggregatorMap.Find(Attribute);
 	if (RefPtr && RefPtr->Get())
 	{
-		RefPtr->Get()->SetBaseValue(Attribute);
-		RefPtr->Get()->MarkDirty();
+		FAggregator* Aggregator =  RefPtr->Get();
+		if (Aggregator->HasPredictedMods())
+		{
+			OnAttributeAggregatorDirty(Aggregator, Attribute);
+		}
 	}
-*/
 }
 
 float FActiveGameplayEffectsContainer::GetGameplayEffectDuration(FActiveGameplayEffectHandle Handle) const
@@ -1449,7 +1468,7 @@ void FActiveGameplayEffectsContainer::AddActiveGameplayEffectGrantedTagsAndModif
 			// Ongoing tags being met. We either calculate magnitude one time, or its done via OnDirty calls (or potentially a frequency timer one day)
 					
 			FAggregator* Aggregator = FindOrCreateAttributeAggregator(Effect.Spec.Def->Modifiers[ModIdx].Attribute).Get();
-			Aggregator->AddMod(Mod.GetEvaluatedMagnitude(), ModInfo.ModifierOp, &ModInfo.SourceTags, &ModInfo.TargetTags, Effect.Handle);
+			Aggregator->AddMod(Mod.GetEvaluatedMagnitude(), ModInfo.ModifierOp, &ModInfo.SourceTags, &ModInfo.TargetTags, Effect.PredictionKey.WasLocallyGenerated(), Effect.Handle);
 		}
 	}
 
diff --git a/Engine/Source/Runtime/GameplayAbilities/Private/GameplayEffectAggregator.cpp b/Engine/Source/Runtime/GameplayAbilities/Private/GameplayEffectAggregator.cpp
--- a/Engine/Source/Runtime/GameplayAbilities/Private/GameplayEffectAggregator.cpp
+++ b/Engine/Source/Runtime/GameplayAbilities/Private/GameplayEffectAggregator.cpp
@@ -11,6 +11,11 @@ bool FAggregatorMod::Qualifies(const FAggregatorEvaluateParameters& Parameters)
 
 	bool bSourceFilterMet = (Parameters.AppliedSourceTagFilter.Num() == 0);
 	bool bTargetFilterMet = (Parameters.AppliedTargetTagFilter.Num() == 0);
+
+	if (Parameters.IncludePredictiveMods == false && IsPredicted)
+	{
+		return false;
+	}
 	
 	const UAbilitySystemComponent* HandleComponent = ActiveHandle.GetOwningAbilitySystemComponent();
 	if (HandleComponent)
@@ -59,6 +64,36 @@ float FAggregator::EvaluateWithBase(float InlineBaseValue, const FAggregatorEval
 	return ((InlineBaseValue + Additive) * Multiplicitive) / Division;
 }
 
+float FAggregator::ReverseEvaluate(float FinalValue, const FAggregatorEvaluateParameters& Parameters) const
+{
+	for (const FAggregatorMod& Mod : Mods[EGameplayModOp::Override])
+	{
+		if (Mod.Qualifies(Parameters))
+		{
+			// This is the case we can't really handle due to lack of information.
+			return FinalValue;
+		}
+	}
+
+	float Additive = SumMods(Mods[EGameplayModOp::Additive], 0.f, Parameters);
+	float Multiplicitive = SumMods(Mods[EGameplayModOp::Multiplicitive], 1.f, Parameters);
+	float Division = SumMods(Mods[EGameplayModOp::Division], 1.f, Parameters);
+
+	if (FMath::IsNearlyZero(Division))
+	{
+		ABILITY_LOG(Warning, TEXT("Division summation was 0.0f in FAggregator."));
+		Division = 1.f;
+	}
+
+	if (Multiplicitive <= SMALL_NUMBER)
+	{
+		return FinalValue;
+	}
+
+	float CalculatedBaseValue = (FinalValue * Division / Multiplicitive) - Additive;
+	return CalculatedBaseValue;
+}
+
 float FAggregator::EvaluateBonus(const FAggregatorEvaluateParameters& Parameters) const
 {
 	return (Evaluate(Parameters) - GetBaseValue());
@@ -69,10 +104,13 @@ float FAggregator::GetBaseValue() const
 	return BaseValue;
 }
 
-void FAggregator::SetBaseValue(float NewBaseValue)
+void FAggregator::SetBaseValue(float NewBaseValue, bool BroadcastDirtyEvent)
 {
 	BaseValue = NewBaseValue;
-	BroadcastOnDirty();
+	if (BroadcastDirtyEvent)
+	{
+		BroadcastOnDirty();
+	}
 }
 
 float FAggregator::StaticExecModOnBaseValue(float BaseValue, TEnumAsByte<EGameplayModOp::Type> ModifierOp, float EvaluatedMagnitude)
@@ -128,7 +166,7 @@ float FAggregator::SumMods(const TArray<FAggregatorMod> &Mods, float Bias, const
 	return Sum;
 }
 
-void FAggregator::AddMod(float EvaluatedMagnitude, TEnumAsByte<EGameplayModOp::Type> ModifierOp, const FGameplayTagRequirements* SourceTagReqs, const FGameplayTagRequirements* TargetTagReqs, FActiveGameplayEffectHandle ActiveHandle)
+void FAggregator::AddMod(float EvaluatedMagnitude, TEnumAsByte<EGameplayModOp::Type> ModifierOp, const FGameplayTagRequirements* SourceTagReqs, const FGameplayTagRequirements* TargetTagReqs, bool IsPredicted, FActiveGameplayEffectHandle ActiveHandle)
 {
 	TArray<FAggregatorMod> &ModList = Mods[ModifierOp];
 
@@ -139,6 +177,10 @@ void FAggregator::AddMod(float EvaluatedMagnitude, TEnumAsByte<EGameplayModOp::T
 	NewMod.TargetTagReqs = TargetTagReqs;
 	NewMod.EvaluatedMagnitude = EvaluatedMagnitude;
 	NewMod.ActiveHandle = ActiveHandle;
+	NewMod.IsPredicted = IsPredicted;
+
+	if (IsPredicted)
+		NumPredictiveMods++;
 
 	BroadcastOnDirty();
 }
@@ -169,6 +211,11 @@ void FAggregator::AddDependant(FActiveGameplayEffectHandle Handle)
 	Dependants.Add(Handle);
 }
 
+bool FAggregator::HasPredictedMods() const
+{
+	return NumPredictiveMods > 0;
+}
+
 void FAggregator::RemoveModsWithActiveHandle(TArray<FAggregatorMod>& Mods, FActiveGameplayEffectHandle ActiveHandle)
 {
 	check(ActiveHandle.IsValid());
@@ -177,6 +224,9 @@ void FAggregator::RemoveModsWithActiveHandle(TArray<FAggregatorMod>& Mods, FActi
 	{
 		if (Mods[idx].ActiveHandle == ActiveHandle)
 		{
+			if (Mods[idx].IsPredicted)
+				NumPredictiveMods--;
+
 			Mods.RemoveAtSwap(idx, 1, false);
 		}
 	}
@@ -211,6 +261,7 @@ void FAggregator::BroadcastOnDirty()
 	
 	OnDirty.Broadcast(this);
 
+
 	TArray<FActiveGameplayEffectHandle>	ValidDependants;
 	for (FActiveGameplayEffectHandle Handle : Dependants)
 	{
diff --git a/Engine/Source/Runtime/GameplayAbilities/Private/GameplayEffectExecutionCalculation.cpp b/Engine/Source/Runtime/GameplayAbilities/Private/GameplayEffectExecutionCalculation.cpp
--- a/Engine/Source/Runtime/GameplayAbilities/Private/GameplayEffectExecutionCalculation.cpp
+++ b/Engine/Source/Runtime/GameplayAbilities/Private/GameplayEffectExecutionCalculation.cpp
@@ -35,7 +35,7 @@ FGameplayEffectCustomExecutionParameters::FGameplayEffectCustomExecutionParamete
 		float ModEvalValue = 0.f;
 		if (ScopedAggregator && CurScopedMod.ModifierMagnitude.AttemptCalculateMagnitude(InOwningSpec, ModEvalValue))
 		{
-			ScopedAggregator->AddMod(ModEvalValue, CurScopedMod.ModifierOp, &CurScopedMod.SourceTags, &CurScopedMod.TargetTags, ModifierHandle);
+			ScopedAggregator->AddMod(ModEvalValue, CurScopedMod.ModifierOp, &CurScopedMod.SourceTags, &CurScopedMod.TargetTags, false, ModifierHandle);
 		}
 		else
 		{
diff --git a/Engine/Source/Runtime/GameplayAbilities/Public/GameplayEffectAggregator.h b/Engine/Source/Runtime/GameplayAbilities/Public/GameplayEffectAggregator.h
--- a/Engine/Source/Runtime/GameplayAbilities/Public/GameplayEffectAggregator.h
+++ b/Engine/Source/Runtime/GameplayAbilities/Public/GameplayEffectAggregator.h
@@ -6,7 +6,7 @@
 
 struct GAMEPLAYABILITIES_API FAggregatorEvaluateParameters
 {
-	FAggregatorEvaluateParameters() : SourceTags(nullptr), TargetTags(nullptr) { }
+	FAggregatorEvaluateParameters() : SourceTags(nullptr), TargetTags(nullptr), IncludePredictiveMods(false) { }
 
 	const FGameplayTagContainer*	SourceTags;
 	const FGameplayTagContainer*	TargetTags;
@@ -16,6 +16,8 @@ struct GAMEPLAYABILITIES_API FAggregatorEvaluateParameters
 
 	/** If any tags are specified in the filter, a mod's owning active gameplay effect's target tags must match ALL of them in order for the mod to count during evaluation */
 	FGameplayTagContainer	AppliedTargetTagFilter;
+
+	bool IncludePredictiveMods;
 };
 
 struct GAMEPLAYABILITIES_API FAggregatorMod
@@ -26,6 +28,7 @@ struct GAMEPLAYABILITIES_API FAggregatorMod
 	float EvaluatedMagnitude;		// Magnitude this mod was last evaluated at
 
 	FActiveGameplayEffectHandle ActiveHandle;	// Handle of the active GameplayEffect we are tied to (if any)
+	bool IsPredicted;
 
 	bool Qualifies(const FAggregatorEvaluateParameters& Parameters) const;
 };
@@ -35,21 +38,23 @@ struct GAMEPLAYABILITIES_API FAggregator : public TSharedFromThis<FAggregator>
 {
 	DECLARE_MULTICAST_DELEGATE_OneParam(FOnAggregatorDirty, FAggregator*);
 
-	FAggregator(float InBaseValue=0.f) : BaseValue(InBaseValue), IsBroadcastingDirty(false) { }
+	FAggregator(float InBaseValue=0.f) : BaseValue(InBaseValue), IsBroadcastingDirty(false), NumPredictiveMods(0) { }
 
 	/** Simple accessor to base value */
 	float GetBaseValue() const;
-	void SetBaseValue(float NewBaseValue);
+	void SetBaseValue(float NewBaseValue, bool BroadcastDirtyEvent = true);
 	
 	void ExecModOnBaseValue(TEnumAsByte<EGameplayModOp::Type> ModifierOp, float EvaluatedMagnitude);
 	static float StaticExecModOnBaseValue(float BaseValue, TEnumAsByte<EGameplayModOp::Type> ModifierOp, float EvaluatedMagnitude);
 
-	void AddMod(float EvaluatedData, TEnumAsByte<EGameplayModOp::Type> ModifierOp, const FGameplayTagRequirements*	SourceTagReqs, const FGameplayTagRequirements* TargetTagReqs, FActiveGameplayEffectHandle ActiveHandle = FActiveGameplayEffectHandle());
+	void AddMod(float EvaluatedData, TEnumAsByte<EGameplayModOp::Type> ModifierOp, const FGameplayTagRequirements*	SourceTagReqs, const FGameplayTagRequirements* TargetTagReqs, bool IsPredicted, FActiveGameplayEffectHandle ActiveHandle = FActiveGameplayEffectHandle() );
 	void RemoveMod(FActiveGameplayEffectHandle ActiveHandle);
 
 	/** Evaluates the Aggregator with the internal base value and given parameters */
 	float Evaluate(const FAggregatorEvaluateParameters& Parameters) const;
 
+	float ReverseEvaluate(float FinalValue, const FAggregatorEvaluateParameters& Parameters) const;
+
 	/** Evaluates the Aggregator with an arbitrary base value */
 	float EvaluateWithBase(float InlineBaseValue, const FAggregatorEvaluateParameters& Parameters) const;
 
@@ -64,6 +69,8 @@ struct GAMEPLAYABILITIES_API FAggregator : public TSharedFromThis<FAggregator>
 
 	void AddDependant(FActiveGameplayEffectHandle Handle);
 
+	bool HasPredictedMods() const;
+
 private:
 
 	void BroadcastOnDirty();
@@ -76,6 +83,7 @@ private:
 	/** ActiveGE handles that we need to notify if we change. NOT copied over during snapshots. */
 	TArray<FActiveGameplayEffectHandle>	Dependants;
 	bool	IsBroadcastingDirty;
+	int32	NumPredictiveMods;
 
 	friend struct FAggregator;
 	friend struct FScopedAggregatorOnDirtyBatch;	// Only outside class that gets to call BroadcastOnDirty()
diff --git a/Engine/Source/Runtime/GameplayAbilities/Public/GameplayPrediction.h b/Engine/Source/Runtime/GameplayAbilities/Public/GameplayPrediction.h
--- a/Engine/Source/Runtime/GameplayAbilities/Public/GameplayPrediction.h
+++ b/Engine/Source/Runtime/GameplayAbilities/Public/GameplayPrediction.h
@@ -337,6 +337,11 @@ struct GAMEPLAYABILITIES_API FPredictionKey
 		return PredictiveConnection != nullptr;
 	}
 
+	bool WasLocallyGenerated() const
+	{
+		return Current > 0 && PredictiveConnection == nullptr;
+	}
+
 	bool DependsOn(KeyType Key)
 	{
 		return (Current == Key || Base == Key);
