PreApplyGameplayEffectSpec global function
Allows projects to add additional rules (tags, linking gameplay effects, etc) at the start of the apply stage of GE

diff --git a/Engine/Source/Runtime/GameplayAbilities/Private/AbilitySystemComponent.cpp b/Engine/Source/Runtime/GameplayAbilities/Private/AbilitySystemComponent.cpp
--- a/Engine/Source/Runtime/GameplayAbilities/Private/AbilitySystemComponent.cpp
+++ b/Engine/Source/Runtime/GameplayAbilities/Private/AbilitySystemComponent.cpp
@@ -519,6 +519,7 @@ FActiveGameplayEffectHandle UAbilitySystemComponent::ApplyGameplayEffectSpecToSe
 		{
 			StackSpec = TSharedPtr<FGameplayEffectSpec>(new FGameplayEffectSpec(Spec));
 			OurCopyOfSpec = StackSpec.Get();
+			UAbilitySystemGlobals::Get().GlobalPreGameplayEffectSpecApply(*OurCopyOfSpec, this);
 			OurCopyOfSpec->CapturedRelevantAttributes.CaptureAttributes(this, EGameplayEffectAttributeCaptureSource::Target);
 		}
 
diff --git a/Engine/Source/Runtime/GameplayAbilities/Private/AbilitySystemGlobals.cpp b/Engine/Source/Runtime/GameplayAbilities/Private/AbilitySystemGlobals.cpp
--- a/Engine/Source/Runtime/GameplayAbilities/Private/AbilitySystemGlobals.cpp
+++ b/Engine/Source/Runtime/GameplayAbilities/Private/AbilitySystemGlobals.cpp
@@ -208,3 +208,8 @@ UGameplayCueManager* UAbilitySystemGlobals::GetGameplayCueManager()
 
 	return GlobalGameplayCueManager;
 }
+
+void UAbilitySystemGlobals::GlobalPreGameplayEffectSpecApply(FGameplayEffectSpec& Spec, UAbilitySystemComponent* AbilitySystemComponent)
+{
+
+}
\ No newline at end of file
diff --git a/Engine/Source/Runtime/GameplayAbilities/Private/GameplayEffect.cpp b/Engine/Source/Runtime/GameplayAbilities/Private/GameplayEffect.cpp
--- a/Engine/Source/Runtime/GameplayAbilities/Private/GameplayEffect.cpp
+++ b/Engine/Source/Runtime/GameplayAbilities/Private/GameplayEffect.cpp
@@ -1331,6 +1331,8 @@ FActiveGameplayEffect& FActiveGameplayEffectsContainer::CreateNewActiveGameplayE
 	FActiveGameplayEffectHandle NewHandle = FActiveGameplayEffectHandle::GenerateNewHandle(Owner);
 	FActiveGameplayEffect& NewEffect = *new (GameplayEffects)FActiveGameplayEffect(NewHandle, Spec, GetWorldTime(), GetGameStateTime(), InPredictionKey);
 
+	UAbilitySystemGlobals::Get().GlobalPreGameplayEffectSpecApply(NewEffect.Spec, Owner);
+
 	// Calc all of our modifier magnitudes now. Some may need to update later based on attributes changing, etc, but those should
 	// be done through delegate callbacks.
 	NewEffect.Spec.CapturedRelevantAttributes.CaptureAttributes(Owner, EGameplayEffectAttributeCaptureSource::Target);
diff --git a/Engine/Source/Runtime/GameplayAbilities/Public/AbilitySystemGlobals.h b/Engine/Source/Runtime/GameplayAbilities/Public/AbilitySystemGlobals.h
--- a/Engine/Source/Runtime/GameplayAbilities/Public/AbilitySystemGlobals.h
+++ b/Engine/Source/Runtime/GameplayAbilities/Public/AbilitySystemGlobals.h
@@ -14,6 +14,7 @@ struct FGameplayAbilityActorInfo;
 struct FGameplayEffectContext;
 struct FGameplayTag;
 struct FAttributeSetInitter;
+struct FGameplayEffectSpec;
 
 /** Holds global data for the skill system. Can be configured per project via config file */
 UCLASS(config=Game)
@@ -86,6 +87,8 @@ class GAMEPLAYABILITIES_API UAbilitySystemGlobals : public UObject
 
 	virtual UGameplayCueManager* GetGameplayCueManager();
 
+	virtual void GlobalPreGameplayEffectSpecApply(FGameplayEffectSpec& Spec, UAbilitySystemComponent* AbilitySystemComponent);
+
 protected:
 
 	virtual void InitAttributeDefaults();
