When hit result is added to an effect context, use that for the world origin if one wasn't explicitly set

diff --git a/Engine/Source/Runtime/GameplayAbilities/Private/GameplayEffectTypes.cpp b/Engine/Source/Runtime/GameplayAbilities/Private/GameplayEffectTypes.cpp
--- a/Engine/Source/Runtime/GameplayAbilities/Private/GameplayEffectTypes.cpp
+++ b/Engine/Source/Runtime/GameplayAbilities/Private/GameplayEffectTypes.cpp
@@ -45,10 +45,15 @@ void FGameplayEffectContext::AddHitResult(const FHitResult InHitResult, bool bRe
 	if (bReset && HitResult.IsValid())
 	{
 		HitResult.Reset();
+		bHasWorldOrigin = false;
 	}
 
 	check(!HitResult.IsValid());
 	HitResult = TSharedPtr<FHitResult>(new FHitResult(InHitResult));
+	if (bHasWorldOrigin == false)
+	{
+		AddOrigin(InHitResult.TraceStart);
+	}
 }
 
 bool FGameplayEffectContext::NetSerialize(FArchive& Ar, class UPackageMap* Map, bool& bOutSuccess)
