Change owned tag container's display name to GrantedTags

diff --git a/Engine/Source/Runtime/GameplayAbilities/Public/GameplayEffect.h b/Engine/Source/Runtime/GameplayAbilities/Public/GameplayEffect.h
--- a/Engine/Source/Runtime/GameplayAbilities/Public/GameplayEffect.h
+++ b/Engine/Source/Runtime/GameplayAbilities/Public/GameplayEffect.h
@@ -556,7 +556,7 @@ public:
 	FGameplayTagContainer GameplayEffectTags;
 	
 	/** "These tags are applied to the actor I am applied to" */
-	UPROPERTY(EditDefaultsOnly, BlueprintReadOnly, Category = Tags)
+	UPROPERTY(EditDefaultsOnly, BlueprintReadOnly, Category = Tags, Deprecated, meta=(DisplayName="GrantedTags"))
 	FInheritedTagContainer InheritableOwnedTagsContainer;
 
 	/** "These tags are applied to the actor I am applied to" */
