Fixed AGameplayAbilityTargetActor_GroundTrace to return the correct start location.

diff --git a/Engine/Source/Runtime/GameplayAbilities/Private/Abilities/GameplayAbilityTargetActor_GroundTrace.cpp b/Engine/Source/Runtime/GameplayAbilities/Private/Abilities/GameplayAbilityTargetActor_GroundTrace.cpp
--- a/Engine/Source/Runtime/GameplayAbilities/Private/Abilities/GameplayAbilityTargetActor_GroundTrace.cpp
+++ b/Engine/Source/Runtime/GameplayAbilities/Private/Abilities/GameplayAbilityTargetActor_GroundTrace.cpp
@@ -149,8 +149,6 @@ FHitResult AGameplayAbilityTargetActor_GroundTrace::PerformTrace(AActor* InSourc
 	// ------------------------------------------------------
 
 	FHitResult ReturnHitResult;
-	FVector LineTraceStart = TraceStart;
-	FVector LineTraceEnd = TraceEnd;
 	//Use a line trace initially to see where the player is actually pointing
 	LineTraceWithFilter(ReturnHitResult, InSourceActor->GetWorld(), Filter, TraceStart, TraceEnd, TraceChannel, Params);
 	//Default to end of trace line if we don't hit anything.
@@ -189,6 +187,9 @@ FHitResult AGameplayAbilityTargetActor_GroundTrace::PerformTrace(AActor* InSourc
 		LocalReticleActor->SetActorLocation(ReturnHitResult.Location);
 	}
 
+	// Reset the trace start so the target data uses the correct origin
+	ReturnHitResult.TraceStart = StartLocation.GetTargetingTransform().GetLocation();
+
 	return ReturnHitResult;
 }
 
