forward decl

diff --git a/Engine/Source/Runtime/GameplayAbilities/Public/GameplayModMagnitudeCalculation.h b/Engine/Source/Runtime/GameplayAbilities/Public/GameplayModMagnitudeCalculation.h
--- a/Engine/Source/Runtime/GameplayAbilities/Public/GameplayModMagnitudeCalculation.h
+++ b/Engine/Source/Runtime/GameplayAbilities/Public/GameplayModMagnitudeCalculation.h
@@ -7,6 +7,7 @@
 #include "GameplayModMagnitudeCalculation.generated.h"
 
 // Forward declarations
+struct FGameplayEffectSpec;
 class UAbilitySystemComponent;
 
 /** Class used to perform custom gameplay effect modifier calculations, either via blueprint or native code */ 
