ApplicationImmunity

diff --git a/Engine/Source/Runtime/GameplayAbilities/Private/Abilities/GameplayAbility.cpp b/Engine/Source/Runtime/GameplayAbilities/Private/Abilities/GameplayAbility.cpp
--- a/Engine/Source/Runtime/GameplayAbilities/Private/Abilities/GameplayAbility.cpp
+++ b/Engine/Source/Runtime/GameplayAbilities/Private/Abilities/GameplayAbility.cpp
@@ -866,9 +866,8 @@ FActiveGameplayEffectHandle UGameplayAbility::BP_ApplyGameplayEffectToTarget(FGa
 		ABILITY_LOG(Error, TEXT("BP_ApplyGameplayEffectToTarget called on ability %s with no GameplayEffectClass."), *GetName());
 		return FActiveGameplayEffectHandle();
 	}
-
-	const UGameplayEffect* GameplayEffect = GameplayEffectClass->GetDefaultObject<UGameplayEffect>();
-	return ApplyGameplayEffectToTarget(CurrentSpecHandle, CurrentActorInfo, CurrentActivationInfo, Target, GameplayEffect, GameplayEffectLevel);
+	
+	return ApplyGameplayEffectToTarget(CurrentSpecHandle, CurrentActorInfo, CurrentActivationInfo, Target, GameplayEffectClass, GameplayEffectLevel);
 }
 
 FActiveGameplayEffectHandle UGameplayAbility::K2_ApplyGameplayEffectToTarget(FGameplayAbilityTargetDataHandle Target, const UGameplayEffect* GameplayEffect, int32 GameplayEffectLevel)
@@ -882,9 +881,9 @@ FActiveGameplayEffectHandle UGameplayAbility::K2_ApplyGameplayEffectToTarget(FGa
 	return BP_ApplyGameplayEffectToTarget(Target, GameplayEffect->GetClass(), GameplayEffectLevel);
 }
 
-FActiveGameplayEffectHandle UGameplayAbility::ApplyGameplayEffectToTarget(const FGameplayAbilitySpecHandle Handle, const FGameplayAbilityActorInfo* ActorInfo, const FGameplayAbilityActivationInfo ActivationInfo, FGameplayAbilityTargetDataHandle Target, const UGameplayEffect* GameplayEffect, int32 GameplayEffectLevel)
+FActiveGameplayEffectHandle UGameplayAbility::ApplyGameplayEffectToTarget(const FGameplayAbilitySpecHandle Handle, const FGameplayAbilityActorInfo* ActorInfo, const FGameplayAbilityActivationInfo ActivationInfo, FGameplayAbilityTargetDataHandle Target, TSubclassOf<UGameplayEffect> GameplayEffectClass, float GameplayEffectLevel)
 {
-	if (GameplayEffect == nullptr)
+	if (GameplayEffectClass == nullptr)
 	{
 		ABILITY_LOG(Error, TEXT("K2_ApplyGameplayEffectToTarget called on ability %s with no GameplayEffect."), *GetName());
 		return FActiveGameplayEffectHandle();
@@ -893,9 +892,11 @@ FActiveGameplayEffectHandle UGameplayAbility::ApplyGameplayEffectToTarget(const
 	FActiveGameplayEffectHandle EffectHandle;
 	if (ActivationInfo.ActivationMode == EGameplayAbilityActivationMode::Authority || ActivationInfo.ActivationMode == EGameplayAbilityActivationMode::Predicting)
 	{
+		FGameplayEffectSpecHandle SpecHandle = MakeOutgoingGameplayEffectSpec(Handle, ActorInfo, ActivationInfo, GameplayEffectClass, GameplayEffectLevel);
+
 		for (auto Data : Target.Data)
 		{
-			TArray<FActiveGameplayEffectHandle> EffectHandles = Data->ApplyGameplayEffect(GameplayEffect, GetEffectContext(ActorInfo), (float)GameplayEffectLevel, ActivationInfo.GetPredictionKeyForNewAction());
+			TArray<FActiveGameplayEffectHandle> EffectHandles = Data->ApplyGameplayEffectSpec(*SpecHandle.Data.Get(), ActivationInfo.GetPredictionKeyForNewAction());
 			if (EffectHandles.Num() > 0)
 			{
 				EffectHandle = EffectHandles[0];
diff --git a/Engine/Source/Runtime/GameplayAbilities/Private/AbilitySystemComponent.cpp b/Engine/Source/Runtime/GameplayAbilities/Private/AbilitySystemComponent.cpp
--- a/Engine/Source/Runtime/GameplayAbilities/Private/AbilitySystemComponent.cpp
+++ b/Engine/Source/Runtime/GameplayAbilities/Private/AbilitySystemComponent.cpp
@@ -459,6 +459,12 @@ FActiveGameplayEffectHandle UAbilitySystemComponent::ApplyGameplayEffectSpecToSe
 		}
 	}
 
+	// Are we currently immunte to this? (ApplicationImmunity)
+	if (ActiveGameplayEffects.CheckApplicationImmunity(Spec))
+	{
+		return FActiveGameplayEffectHandle();
+	}
+
 	// Check AttributeSet requirements: do we have everything this GameplayEffectSpec expects?
 	// We may want to cache this off in some way to make the runtime check quicker.
 	// We also need to handle things in the execution list
diff --git a/Engine/Source/Runtime/GameplayAbilities/Private/AbilitySystemStats.cpp b/Engine/Source/Runtime/GameplayAbilities/Private/AbilitySystemStats.cpp
--- a/Engine/Source/Runtime/GameplayAbilities/Private/AbilitySystemStats.cpp
+++ b/Engine/Source/Runtime/GameplayAbilities/Private/AbilitySystemStats.cpp
@@ -19,5 +19,6 @@ DEFINE_STAT(STAT_InitAttributeSetDefaults);
 DEFINE_STAT(STAT_TickAbilityTasks);
 DEFINE_STAT(STAT_FindAbilitySpecFromHandle);
 DEFINE_STAT(STAT_AggregatorEvaluate);
+DEFINE_STAT(STAT_CheckApplicationImmunity);
 
 
diff --git a/Engine/Source/Runtime/GameplayAbilities/Private/GameplayAbilityTargetTypes.cpp b/Engine/Source/Runtime/GameplayAbilities/Private/GameplayAbilityTargetTypes.cpp
--- a/Engine/Source/Runtime/GameplayAbilities/Private/GameplayAbilityTargetTypes.cpp
+++ b/Engine/Source/Runtime/GameplayAbilities/Private/GameplayAbilityTargetTypes.cpp
@@ -35,6 +35,37 @@ TArray<FActiveGameplayEffectHandle> FGameplayAbilityTargetData::ApplyGameplayEff
 	return AppliedHandles;
 }
 
+TArray<FActiveGameplayEffectHandle> FGameplayAbilityTargetData::ApplyGameplayEffectSpec(FGameplayEffectSpec& Spec, FPredictionKey PredictionKey)
+{
+	TArray<FActiveGameplayEffectHandle>	AppliedHandles;
+
+	if (!ensure(Spec.GetContext().IsValid()))
+	{
+		return AppliedHandles;
+	}
+
+	FGameplayEffectContextHandle EffectContext = Spec.GetContext();
+	AddTargetDataToContext(EffectContext);
+
+	TArray<TWeakObjectPtr<AActor> > Actors = GetActors();
+	
+	AppliedHandles.Reserve(Actors.Num());
+
+	for (TWeakObjectPtr<AActor> TargetActor : Actors)
+	{
+		if (TargetActor.IsValid())
+		{
+			UAbilitySystemComponent* TargetComponent = UAbilitySystemBlueprintLibrary::GetAbilitySystemComponent(TargetActor.Get());
+			if (TargetComponent && EffectContext.GetInstigatorAbilitySystemComponent())
+			{
+				AppliedHandles.Add(EffectContext.GetInstigatorAbilitySystemComponent()->ApplyGameplayEffectSpecToTarget(Spec, TargetComponent, PredictionKey));
+			}
+		}
+	}
+
+	return AppliedHandles;
+}
+
 void FGameplayAbilityTargetData::AddTargetDataToContext(FGameplayEffectContextHandle& Context)
 {
 	if (HasHitResult())
diff --git a/Engine/Source/Runtime/GameplayAbilities/Private/GameplayEffect.cpp b/Engine/Source/Runtime/GameplayAbilities/Private/GameplayEffect.cpp
--- a/Engine/Source/Runtime/GameplayAbilities/Private/GameplayEffect.cpp
+++ b/Engine/Source/Runtime/GameplayAbilities/Private/GameplayEffect.cpp
@@ -1479,6 +1479,9 @@ void FActiveGameplayEffectsContainer::AddActiveGameplayEffectGrantedTagsAndModif
 
 	Owner->UpdateTagMap(Effect.Spec.DynamicGrantedTags, 1);
 
+	ApplicationImmunityGameplayTagCountContainer.UpdateTagMap(Effect.Spec.Def->GrantedApplicationImmunityTags.RequireTags, 1);
+	ApplicationImmunityGameplayTagCountContainer.UpdateTagMap(Effect.Spec.Def->GrantedApplicationImmunityTags.IgnoreTags, 1);
+
 	for (const FGameplayEffectCue& Cue : Effect.Spec.Def->GameplayCues)
 	{
 		Owner->UpdateTagMap(Cue.GameplayCueTags, 1);
@@ -1600,6 +1603,9 @@ void FActiveGameplayEffectsContainer::RemoveActiveGameplayEffectGrantedTagsAndMo
 	IGameplayTagsModule& GameplayTagsModule = IGameplayTagsModule::Get();
 	Owner->UpdateTagMap(Effect.Spec.Def->InheritableOwnedTagsContainer.CombinedTags, -1);
 
+	ApplicationImmunityGameplayTagCountContainer.UpdateTagMap(Effect.Spec.Def->GrantedApplicationImmunityTags.RequireTags, -1);
+	ApplicationImmunityGameplayTagCountContainer.UpdateTagMap(Effect.Spec.Def->GrantedApplicationImmunityTags.IgnoreTags, -1);
+
 	Owner->UpdateTagMap(Effect.Spec.DynamicGrantedTags, -1);
 
 	for (const FGameplayEffectCue& Cue : Effect.Spec.Def->GameplayCues)
@@ -1646,6 +1652,32 @@ void FActiveGameplayEffectsContainer::OnOwnerTagChange(FGameplayTag TagChange, i
 	}
 }
 
+bool FActiveGameplayEffectsContainer::CheckApplicationImmunity(const FGameplayEffectSpec& SpecToApply) const
+{
+	SCOPE_CYCLE_COUNTER(STAT_CheckApplicationImmunity)
+
+	const FGameplayTagContainer* Tags = SpecToApply.CapturedSourceTags.GetAggregatedTags();
+	if (!ensure(Tags))
+		return false;
+	
+
+	// Quick map test
+	if (!ApplicationImmunityGameplayTagCountContainer.HasAnyMatchingGameplayTags(*Tags, EGameplayTagMatchType::Explicit, false))
+	{
+		return false;
+	}
+
+	for (const FActiveGameplayEffect& Effect : GameplayEffects)
+	{
+		if (Effect.Spec.Def->GrantedApplicationImmunityTags.RequirementsMet( *Tags ))
+		{
+			return true;
+		}
+	}
+
+	return false;
+}
+
 bool FActiveGameplayEffectsContainer::IsNetAuthority() const
 {
 	check(Owner);
diff --git a/Engine/Source/Runtime/GameplayAbilities/Private/GameplayEffectTypes.cpp b/Engine/Source/Runtime/GameplayAbilities/Private/GameplayEffectTypes.cpp
--- a/Engine/Source/Runtime/GameplayAbilities/Private/GameplayEffectTypes.cpp
+++ b/Engine/Source/Runtime/GameplayAbilities/Private/GameplayEffectTypes.cpp
@@ -107,6 +107,19 @@ void FGameplayEffectContext::AddOrigin(FVector InOrigin)
 	WorldOrigin = InOrigin;
 }
 
+void FGameplayEffectContext::GetOwnedGameplayTags(OUT FGameplayTagContainer &TagContainer) const
+{
+	IGameplayTagAssetInterface* TagInterface = Cast<IGameplayTagAssetInterface>(Instigator);
+	if (TagInterface)
+	{
+		TagInterface->GetOwnedGameplayTags(TagContainer);
+	}
+	else if(InstigatorAbilitySystemComponent)
+	{
+		InstigatorAbilitySystemComponent->GetOwnedGameplayTags(TagContainer);
+	}
+}
+
 bool FGameplayEffectContextHandle::NetSerialize(FArchive& Ar, class UPackageMap* Map, bool& bOutSuccess)
 {
 	UScriptStruct* ScriptStruct = Data.IsValid() ? Data->GetScriptStruct() : NULL;
diff --git a/Engine/Source/Runtime/GameplayAbilities/Public/Abilities/GameplayAbility.h b/Engine/Source/Runtime/GameplayAbilities/Public/Abilities/GameplayAbility.h
--- a/Engine/Source/Runtime/GameplayAbilities/Public/Abilities/GameplayAbility.h
+++ b/Engine/Source/Runtime/GameplayAbilities/Public/Abilities/GameplayAbility.h
@@ -361,7 +361,7 @@ protected:
 	FActiveGameplayEffectHandle K2_ApplyGameplayEffectToTarget(FGameplayAbilityTargetDataHandle TargetData, const UGameplayEffect* GameplayEffect, int32 GameplayEffectLevel = 1);
 
 	/** Non blueprintcallable, safe to call on CDO/NonInstance abilities */
-	FActiveGameplayEffectHandle ApplyGameplayEffectToTarget(const FGameplayAbilitySpecHandle Handle, const FGameplayAbilityActorInfo* ActorInfo, const FGameplayAbilityActivationInfo ActivationInfo, FGameplayAbilityTargetDataHandle Target, const UGameplayEffect* GameplayEffect, int32 GameplayEffectLevel);
+	FActiveGameplayEffectHandle ApplyGameplayEffectToTarget(const FGameplayAbilitySpecHandle Handle, const FGameplayAbilityActorInfo* ActorInfo, const FGameplayAbilityActivationInfo ActivationInfo, FGameplayAbilityTargetDataHandle Target, TSubclassOf<UGameplayEffect> GameplayEffectClass, float GameplayEffectLevel);
 
 	// -------------------------------------
 	//	GameplayCue
diff --git a/Engine/Source/Runtime/GameplayAbilities/Public/Abilities/GameplayAbilityTargetTypes.h b/Engine/Source/Runtime/GameplayAbilities/Public/Abilities/GameplayAbilityTargetTypes.h
--- a/Engine/Source/Runtime/GameplayAbilities/Public/Abilities/GameplayAbilityTargetTypes.h
+++ b/Engine/Source/Runtime/GameplayAbilities/Public/Abilities/GameplayAbilityTargetTypes.h
@@ -65,6 +65,8 @@ struct GAMEPLAYABILITIES_API FGameplayAbilityTargetData
 
 	TArray<FActiveGameplayEffectHandle> ApplyGameplayEffect(const UGameplayEffect* GameplayEffect, const FGameplayEffectContextHandle& InEffectContext, float Level, FPredictionKey PredictionKey = FPredictionKey());
 
+	TArray<FActiveGameplayEffectHandle> ApplyGameplayEffectSpec(FGameplayEffectSpec& Spec, FPredictionKey PredictionKey = FPredictionKey());
+
 	virtual void AddTargetDataToContext(FGameplayEffectContextHandle& Context);
 
 	virtual TArray<TWeakObjectPtr<AActor> >	GetActors() const
diff --git a/Engine/Source/Runtime/GameplayAbilities/Public/AbilitySystemComponent.h b/Engine/Source/Runtime/GameplayAbilities/Public/AbilitySystemComponent.h
--- a/Engine/Source/Runtime/GameplayAbilities/Public/AbilitySystemComponent.h
+++ b/Engine/Source/Runtime/GameplayAbilities/Public/AbilitySystemComponent.h
@@ -754,9 +754,6 @@ private:
 	void UpdateTagMap(const FGameplayTag& BaseTag, int32 CountDelta);
 	
 	void UpdateTagMap(const FGameplayTagContainer& Container, int32 CountDelta);
-
-	// Version for delegate callback.
-	FGameplayTagContainer GetOwnedGameplayTagsDel() const;
 	
 	// ---------------------------------------------
 
diff --git a/Engine/Source/Runtime/GameplayAbilities/Public/AbilitySystemStats.h b/Engine/Source/Runtime/GameplayAbilities/Public/AbilitySystemStats.h
--- a/Engine/Source/Runtime/GameplayAbilities/Public/AbilitySystemStats.h
+++ b/Engine/Source/Runtime/GameplayAbilities/Public/AbilitySystemStats.h
@@ -22,4 +22,5 @@ DECLARE_CYCLE_STAT_EXTERN(TEXT("GetOutgoingSpec"), STAT_GetOutgoingSpec, STATGRO
 DECLARE_CYCLE_STAT_EXTERN(TEXT("InitAttributeSetDefaults"), STAT_InitAttributeSetDefaults, STATGROUP_AbilitySystem, );
 DECLARE_CYCLE_STAT_EXTERN(TEXT("TickAbilityTasks"), STAT_TickAbilityTasks, STATGROUP_AbilitySystem, );
 DECLARE_CYCLE_STAT_EXTERN(TEXT("FindAbilitySpecFromHandle"), STAT_FindAbilitySpecFromHandle, STATGROUP_AbilitySystem, );
-DECLARE_CYCLE_STAT_EXTERN(TEXT("Aggregator Evaluate"), STAT_AggregatorEvaluate, STATGROUP_AbilitySystem, );
\ No newline at end of file
+DECLARE_CYCLE_STAT_EXTERN(TEXT("Aggregator Evaluate"), STAT_AggregatorEvaluate, STATGROUP_AbilitySystem, );
+DECLARE_CYCLE_STAT_EXTERN(TEXT("Check Application Immunity"), STAT_CheckApplicationImmunity, STATGROUP_AbilitySystem, );
\ No newline at end of file
diff --git a/Engine/Source/Runtime/GameplayAbilities/Public/GameplayEffect.h b/Engine/Source/Runtime/GameplayAbilities/Public/GameplayEffect.h
--- a/Engine/Source/Runtime/GameplayAbilities/Public/GameplayEffect.h
+++ b/Engine/Source/Runtime/GameplayAbilities/Public/GameplayEffect.h
@@ -581,6 +581,9 @@ public:
 	UPROPERTY(EditDefaultsOnly, BlueprintReadOnly, Category=Deprecated)
 	FGameplayTagContainer ClearTagsContainer;
 
+	/** Grants the owner immunity from these source tags. */
+	UPROPERTY(EditDefaultsOnly, BlueprintReadOnly, Category = Tags, meta = (DisplayName = "GrantedApplicationImmunityTags"))
+	FGameplayTagRequirements GrantedApplicationImmunityTags;
 };
 
 /** Holds evaluated magnitude from a GameplayEffect modifier */
@@ -1167,6 +1170,8 @@ struct FActiveGameplayEffectsContainer : public FFastArraySerializer
 
 	void OnOwnerTagChange(FGameplayTag TagChange, int32 NewCount);
 
+	bool CheckApplicationImmunity(const FGameplayEffectSpec& SpecToApply) const;
+
 private:
 
 	FTimerHandle StackHandle;
@@ -1206,6 +1211,8 @@ private:
 
 	TMap<FGameplayTag, TSet<FActiveGameplayEffectHandle> >	ActiveEffectTagDependencies;
 
+	FGameplayTagCountContainer ApplicationImmunityGameplayTagCountContainer;
+
 	FAggregatorRef& FindOrCreateAttributeAggregator(FGameplayAttribute Attribute);
 
 	void OnAttributeAggregatorDirty(FAggregator* Aggregator, FGameplayAttribute Attribute);
diff --git a/Engine/Source/Runtime/GameplayAbilities/Public/GameplayEffectTypes.h b/Engine/Source/Runtime/GameplayAbilities/Public/GameplayEffectTypes.h
--- a/Engine/Source/Runtime/GameplayAbilities/Public/GameplayEffectTypes.h
+++ b/Engine/Source/Runtime/GameplayAbilities/Public/GameplayEffectTypes.h
@@ -270,14 +270,7 @@ struct GAMEPLAYABILITIES_API FGameplayEffectContext
 	}
 
 	/** Returns the list of gameplay tags applicable to this effect, defaults to the owner's tags */
-	virtual void GetOwnedGameplayTags(OUT FGameplayTagContainer &TagContainer) const
-	{
-		IGameplayTagAssetInterface* TagInterface = Cast<IGameplayTagAssetInterface>(Instigator);
-		if (TagInterface)
-		{
-			TagInterface->GetOwnedGameplayTags(TagContainer);
-		}
-	}
+	virtual void GetOwnedGameplayTags(OUT FGameplayTagContainer &TagContainer) const;
 
 	/** Sets the instigator and effect causer. Instigator is who owns the ability that spawned this, EffectCauser is the actor that is the physical source of the effect, such as a weapon. They can be the same. */
 	virtual void AddInstigator(class AActor *InInstigator, class AActor *InEffectCauser);
