WaitAttributeChange now supports a mode that permits comparisons.

diff --git a/Engine/Source/Runtime/GameplayAbilities/Private/Abilities/Tasks/AbilityTask_WaitAttributeChange.cpp b/Engine/Source/Runtime/GameplayAbilities/Private/Abilities/Tasks/AbilityTask_WaitAttributeChange.cpp
--- a/Engine/Source/Runtime/GameplayAbilities/Private/Abilities/Tasks/AbilityTask_WaitAttributeChange.cpp
+++ b/Engine/Source/Runtime/GameplayAbilities/Private/Abilities/Tasks/AbilityTask_WaitAttributeChange.cpp
@@ -6,7 +6,7 @@
 
 #include "Abilities/GameplayAbility.h"
 #include "GameplayEffectExtension.h"
-
+#pragma optimize("",off)
 UAbilityTask_WaitAttributeChange::UAbilityTask_WaitAttributeChange(const FObjectInitializer& ObjectInitializer)
 	: Super(ObjectInitializer)
 {
@@ -18,6 +18,19 @@ UAbilityTask_WaitAttributeChange* UAbilityTask_WaitAttributeChange::WaitForAttri
 	MyObj->WithTag = InWithTag;
 	MyObj->WithoutTag = InWithoutTag;
 	MyObj->Attribute = InAttribute;
+	MyObj->ComparisonType = EWaitAttributeChangeComparison::None;
+
+	return MyObj;
+}
+
+UAbilityTask_WaitAttributeChange* UAbilityTask_WaitAttributeChange::WaitForAttributeChangeWithComparison(UObject* WorldContextObject, FGameplayAttribute InAttribute, FGameplayTag InWithTag, FGameplayTag InWithoutTag, TEnumAsByte<EWaitAttributeChangeComparison::Type> InComparisonType, float InComparisonValue)
+{
+	auto MyObj = NewTask<UAbilityTask_WaitAttributeChange>(WorldContextObject);
+	MyObj->WithTag = InWithTag;
+	MyObj->WithoutTag = InWithoutTag;
+	MyObj->Attribute = InAttribute;
+	MyObj->ComparisonType = InComparisonType;
+	MyObj->ComparisonValue = InComparisonValue;
 
 	return MyObj;
 }
@@ -51,9 +64,35 @@ void UAbilityTask_WaitAttributeChange::OnAttributeChange(float NewValue, const F
 		}
 	}	
 
-	OnChange.Broadcast();
-
-	EndTask();
+	bool PassedComparison = true;
+	switch (ComparisonType)
+	{
+	case EWaitAttributeChangeComparison::ExactlyEqualTo:
+		PassedComparison = (NewValue == ComparisonValue);
+		break;		
+	case EWaitAttributeChangeComparison::GreaterThan:
+		PassedComparison = (NewValue > ComparisonValue);
+		break;
+	case EWaitAttributeChangeComparison::GreaterThanOrEqualTo:
+		PassedComparison = (NewValue >= ComparisonValue);
+		break;
+	case EWaitAttributeChangeComparison::LessThan:
+		PassedComparison = (NewValue < ComparisonValue);
+		break;
+	case EWaitAttributeChangeComparison::LessThanOrEqualTo:
+		PassedComparison = (NewValue <= ComparisonValue);
+		break;
+	case EWaitAttributeChangeComparison::NotEqualTo:
+		PassedComparison = (NewValue != ComparisonValue);
+		break;
+	default:
+		break;
+	}
+	if (PassedComparison)
+	{
+		OnChange.Broadcast();
+		EndTask();
+	}
 }
 
 void UAbilityTask_WaitAttributeChange::OnDestroy(bool AbilityEnded)
diff --git a/Engine/Source/Runtime/GameplayAbilities/Public/Abilities/Tasks/AbilityTask_WaitAttributeChange.h b/Engine/Source/Runtime/GameplayAbilities/Public/Abilities/Tasks/AbilityTask_WaitAttributeChange.h
--- a/Engine/Source/Runtime/GameplayAbilities/Public/Abilities/Tasks/AbilityTask_WaitAttributeChange.h
+++ b/Engine/Source/Runtime/GameplayAbilities/Public/Abilities/Tasks/AbilityTask_WaitAttributeChange.h
@@ -10,6 +10,22 @@
 
 struct FGameplayEffectModCallbackData;
 
+UENUM()
+namespace EWaitAttributeChangeComparison
+{
+	enum Type
+	{
+		None,
+		GreaterThan,
+		LessThan,
+		GreaterThanOrEqualTo,
+		LessThanOrEqualTo,
+		NotEqualTo,
+		ExactlyEqualTo,
+		MAX UMETA(Hidden)
+	};
+}
+
 /**
  *	Waits for the actor to activate another ability
  */
@@ -25,17 +41,21 @@ class UAbilityTask_WaitAttributeChange : public UAbilityTask
 
 	virtual void Activate() override;
 
-	
-		
 	void OnAttributeChange(float NewValue, const FGameplayEffectModCallbackData*);
 
 	/** Wait until an attribute changes. */
 	UFUNCTION(BlueprintCallable, Category="Ability|Tasks", meta = (HidePin = "WorldContextObject", DefaultToSelf = "WorldContextObject", BlueprintInternalUseOnly = "TRUE"))
 	static UAbilityTask_WaitAttributeChange* WaitForAttributeChange(UObject* WorldContextObject, FGameplayAttribute Attribute, FGameplayTag WithSrcTag, FGameplayTag WithoutSrcTag);
 
+	/** Wait until an attribute changes to pass a given test. */
+	UFUNCTION(BlueprintCallable, Category = "Ability|Tasks", meta = (HidePin = "WorldContextObject", DefaultToSelf = "WorldContextObject", BlueprintInternalUseOnly = "TRUE"))
+	static UAbilityTask_WaitAttributeChange* WaitForAttributeChangeWithComparison(UObject* WorldContextObject, FGameplayAttribute InAttribute, FGameplayTag InWithTag, FGameplayTag InWithoutTag, TEnumAsByte<EWaitAttributeChangeComparison::Type> InComparisonType, float InComparisonValue);
+
 	FGameplayTag WithTag;
 	FGameplayTag WithoutTag;
 	FGameplayAttribute	Attribute;
+	TEnumAsByte<EWaitAttributeChangeComparison::Type> ComparisonType;
+	float ComparisonValue;
 
 protected:
 
