fix XP awarding for new GE changes

diff --git a/Engine/Source/Runtime/GameplayAbilities/Private/GameplayEffect.cpp b/Engine/Source/Runtime/GameplayAbilities/Private/GameplayEffect.cpp
--- a/Engine/Source/Runtime/GameplayAbilities/Private/GameplayEffect.cpp
+++ b/Engine/Source/Runtime/GameplayAbilities/Private/GameplayEffect.cpp
@@ -572,6 +572,22 @@ void FGameplayEffectSpec::SetModifierMagnitude(int32 ModIdx, float EvaluatedMagn
 	Modifiers[ModIdx].EvaluatedMagnitude = EvaluatedMagnitude;
 }
 
+void FGameplayEffectSpec::SetModifierMagnitude(FGameplayAttribute Attribute, float EvaluatedMagnitude)
+{
+	for (int32 ModIdx=0; ModIdx < Def->Modifiers.Num(); ++ModIdx)
+	{
+		const FGameplayModifierInfo& ModDef = Def->Modifiers[ModIdx];
+		FModifierSpec& ModSpec = Modifiers[ModIdx];
+
+		if (ModDef.Attribute == Attribute && ModDef.ModifierMagnitude.GetMagnitudeCalculationType() == EGameplayEffectMagnitudeCalculation::SetByCaller )
+		{
+			ModSpec.EvaluatedMagnitude = EvaluatedMagnitude;
+			return;
+		}
+	}
+
+	ABILITY_LOG(Error, TEXT("FGameplayEffectSpec::SetModifierMagnitude - called on Def %s for Attribute %s but no valid modifier was found."), *Def->GetName(), *Attribute.GetName());
+}
 
 // --------------------------------------------------------------------------------------------------------------------------------------------------------
 //
diff --git a/Engine/Source/Runtime/GameplayAbilities/Public/GameplayEffect.h b/Engine/Source/Runtime/GameplayAbilities/Public/GameplayEffect.h
--- a/Engine/Source/Runtime/GameplayAbilities/Public/GameplayEffect.h
+++ b/Engine/Source/Runtime/GameplayAbilities/Public/GameplayEffect.h
@@ -865,6 +865,9 @@ struct GAMEPLAYABILITIES_API FGameplayEffectSpec
 	/** Sets the magnitude of one of our modifiers. This must be done on modifiers which are EGameplayEffectMagnitudeCalculation::SetByCaller  */
 	void SetModifierMagnitude(int32 ModIdx, float EvaluatedMagnitude);
 
+	/** Set magnitude of (the first) SetByCaller modifier with the given GameplayAttribute  */
+	void SetModifierMagnitude(FGameplayAttribute Attribute, float EvaluatedMagnitude);
+
 	// The duration in seconds of this effect
 	// instantaneous effects should have a duration of UGameplayEffect::INSTANT_APPLICATION
 	// effects that last forever should have a duration of UGameplayEffect::INFINITE_DURATION
