CIS Fix: removing #pragma optimize (does not compile on Mac/iOS)

diff --git a/Engine/Source/Runtime/GameplayAbilities/Private/Abilities/Tasks/AbilityTask_WaitAttributeChange.cpp b/Engine/Source/Runtime/GameplayAbilities/Private/Abilities/Tasks/AbilityTask_WaitAttributeChange.cpp
--- a/Engine/Source/Runtime/GameplayAbilities/Private/Abilities/Tasks/AbilityTask_WaitAttributeChange.cpp
+++ b/Engine/Source/Runtime/GameplayAbilities/Private/Abilities/Tasks/AbilityTask_WaitAttributeChange.cpp
@@ -6,7 +6,7 @@
 
 #include "Abilities/GameplayAbility.h"
 #include "GameplayEffectExtension.h"
-#pragma optimize("",off)
+
 UAbilityTask_WaitAttributeChange::UAbilityTask_WaitAttributeChange(const FObjectInitializer& ObjectInitializer)
 	: Super(ObjectInitializer)
 {
