Added task-instance feature to MoveToLocation task.

diff --git a/Engine/Source/Runtime/GameplayAbilities/Private/Abilities/Tasks/AbilityTask_MoveToLocation.cpp b/Engine/Source/Runtime/GameplayAbilities/Private/Abilities/Tasks/AbilityTask_MoveToLocation.cpp
--- a/Engine/Source/Runtime/GameplayAbilities/Private/Abilities/Tasks/AbilityTask_MoveToLocation.cpp
+++ b/Engine/Source/Runtime/GameplayAbilities/Private/Abilities/Tasks/AbilityTask_MoveToLocation.cpp
@@ -13,9 +13,9 @@ UAbilityTask_MoveToLocation::UAbilityTask_MoveToLocation(const FObjectInitialize
 	bIsFinished = false;
 }
 
-UAbilityTask_MoveToLocation* UAbilityTask_MoveToLocation::MoveToLocation(class UObject* WorldContextObject, FVector Location, float Duration, UCurveFloat* OptionalInterpolationCurve)
+UAbilityTask_MoveToLocation* UAbilityTask_MoveToLocation::MoveToLocation(class UObject* WorldContextObject, FName TaskInstanceName, FVector Location, float Duration, UCurveFloat* OptionalInterpolationCurve)
 {
-	auto MyObj = NewTask<UAbilityTask_MoveToLocation>(WorldContextObject);
+	auto MyObj = NewTask<UAbilityTask_MoveToLocation>(WorldContextObject, TaskInstanceName);
 
 	MyObj->StartLocation = MyObj->GetAvatarActor()->GetActorLocation();
 	MyObj->TargetLocation = Location;
diff --git a/Engine/Source/Runtime/GameplayAbilities/Public/Abilities/Tasks/AbilityTask_MoveToLocation.h b/Engine/Source/Runtime/GameplayAbilities/Public/Abilities/Tasks/AbilityTask_MoveToLocation.h
--- a/Engine/Source/Runtime/GameplayAbilities/Public/Abilities/Tasks/AbilityTask_MoveToLocation.h
+++ b/Engine/Source/Runtime/GameplayAbilities/Public/Abilities/Tasks/AbilityTask_MoveToLocation.h
@@ -32,7 +32,7 @@ class UAbilityTask_MoveToLocation : public UAbilityTask
 
 	/** Move to the specified location, using the curve (range 0 - 1) or linearly if no curve is specified */
 	UFUNCTION(BlueprintCallable, Category = "Ability|Tasks", meta = (HidePin = "WorldContextObject", DefaultToSelf = "WorldContextObject", BlueprintInternalUseOnly = "TRUE"))
-	static UAbilityTask_MoveToLocation* MoveToLocation(UObject* WorldContextObject, FVector Location, float Duration, UCurveFloat* OptionalInterpolationCurve);
+	static UAbilityTask_MoveToLocation* MoveToLocation(UObject* WorldContextObject, FName TaskInstanceName, FVector Location, float Duration, UCurveFloat* OptionalInterpolationCurve);
 
 	virtual void Activate() override;
 
