fix issue where ASC were always ticking until anim montage was played

diff --git a/Engine/Source/Runtime/GameplayAbilities/Public/Abilities/GameplayAbilityTypes.h b/Engine/Source/Runtime/GameplayAbilities/Public/Abilities/GameplayAbilityTypes.h
--- a/Engine/Source/Runtime/GameplayAbilities/Public/Abilities/GameplayAbilityTypes.h
+++ b/Engine/Source/Runtime/GameplayAbilities/Public/Abilities/GameplayAbilityTypes.h
@@ -377,7 +377,7 @@ struct GAMEPLAYABILITIES_API FGameplayAbilityRepAnimMontage
 		PlayRate(0.f),
 		Position(0.f),
 		NextSectionID(0),
-		IsStopped(0),
+		IsStopped(true),
 		ForcePlayBit(0)
 	{
 	}
