allow GC_ prefix for gameplaycue notify assets

diff --git a/Engine/Source/Runtime/GameplayAbilities/Private/GameplayCueNotify_Actor.cpp b/Engine/Source/Runtime/GameplayAbilities/Private/GameplayCueNotify_Actor.cpp
--- a/Engine/Source/Runtime/GameplayAbilities/Private/GameplayCueNotify_Actor.cpp
+++ b/Engine/Source/Runtime/GameplayAbilities/Private/GameplayCueNotify_Actor.cpp
@@ -30,6 +30,7 @@ void AGameplayCueNotify_Actor::DeriveGameplayCueTagFromAssetName()
 			FString MyName = GetName();
 
 			MyName.RemoveFromStart(TEXT("Default__"));
+			MyName.RemoveFromStart(TEXT("GC_"));		// allow GC_ prefix in asset name
 			MyName.RemoveFromEnd(TEXT("_c"));
 
 			MyName.ReplaceInline(TEXT("_"), TEXT("."));
diff --git a/Engine/Source/Runtime/GameplayAbilities/Private/GameplayCueNotify_Static.cpp b/Engine/Source/Runtime/GameplayAbilities/Private/GameplayCueNotify_Static.cpp
--- a/Engine/Source/Runtime/GameplayAbilities/Private/GameplayCueNotify_Static.cpp
+++ b/Engine/Source/Runtime/GameplayAbilities/Private/GameplayCueNotify_Static.cpp
@@ -30,6 +30,7 @@ void UGameplayCueNotify_Static::DeriveGameplayCueTagFromAssetName()
 			FString MyName = GetName();
 
 			MyName.RemoveFromStart(TEXT("Default__"));
+			MyName.RemoveFromStart(TEXT("GC_"));		// allow GC_ prefix in asset name
 			MyName.RemoveFromEnd(TEXT("_c"));
 
 			MyName.ReplaceInline(TEXT("_"), TEXT("."));
