Minor changes
-Fix for immunity wrt empty containers
-export scoped modifier structs to dll
-mark pending kill all InstancedPerActor abilities on destroy

diff --git a/Engine/Source/Runtime/GameplayAbilities/Private/AbilitySystemComponent_Abilities.cpp b/Engine/Source/Runtime/GameplayAbilities/Private/AbilitySystemComponent_Abilities.cpp
--- a/Engine/Source/Runtime/GameplayAbilities/Private/AbilitySystemComponent_Abilities.cpp
+++ b/Engine/Source/Runtime/GameplayAbilities/Private/AbilitySystemComponent_Abilities.cpp
@@ -62,7 +62,17 @@ void UAbilitySystemComponent::OnComponentDestroyed()
 	{
 		// Cancel all abilities before we are destroyed.
 		CancelAbilities();
-	}	
+
+		// Mark pending kill any remainging instanced abilities
+		// (CancelAbilities() will only MarkPending kill InstancePerExecution abilities).
+		for (FGameplayAbilitySpec& Spec : ActivatableAbilities)
+		{
+			if (Spec.Ability && Spec.Ability->HasAnyFlags(RF_ClassDefaultObject) == false && Spec.Ability->IsPendingKill() == false)
+			{
+				Spec.Ability->MarkPendingKill();
+			}
+		}
+	}
 
 	// Call the super at the end, after we've done what we needed to do
 	Super::OnComponentDestroyed();
diff --git a/Engine/Source/Runtime/GameplayAbilities/Private/GameplayEffect.cpp b/Engine/Source/Runtime/GameplayAbilities/Private/GameplayEffect.cpp
--- a/Engine/Source/Runtime/GameplayAbilities/Private/GameplayEffect.cpp
+++ b/Engine/Source/Runtime/GameplayAbilities/Private/GameplayEffect.cpp
@@ -1685,7 +1685,7 @@ bool FActiveGameplayEffectsContainer::CheckApplicationImmunity(const FGameplayEf
 
 	for (const FActiveGameplayEffect& Effect : GameplayEffects)
 	{
-		if (Effect.Spec.Def->GrantedApplicationImmunityTags.RequirementsMet( *Tags ))
+		if (Effect.Spec.Def->GrantedApplicationImmunityTags.IsEmpty() == false && Effect.Spec.Def->GrantedApplicationImmunityTags.RequirementsMet( *Tags ))
 		{
 			return true;
 		}
diff --git a/Engine/Source/Runtime/GameplayAbilities/Public/AbilitySystemComponent.h b/Engine/Source/Runtime/GameplayAbilities/Public/AbilitySystemComponent.h
--- a/Engine/Source/Runtime/GameplayAbilities/Public/AbilitySystemComponent.h
+++ b/Engine/Source/Runtime/GameplayAbilities/Public/AbilitySystemComponent.h
@@ -264,7 +264,7 @@ class GAMEPLAYABILITIES_API UAbilitySystemComponent : public UActorComponent, pu
 	 *
 	 *	Tags added this way are not replicated! 
 	 *	
-	 *	It is up to the calling GameCode to make sure these tags are added on a;; clients/server where necessary
+	 *	It is up to the calling GameCode to make sure these tags are added on clients/server where necessary
 	 */
 
 	void AddLooseGameplayTag(const FGameplayTag& GameplayTag, int32 Count=1);
diff --git a/Engine/Source/Runtime/GameplayAbilities/Public/GameplayEffect.h b/Engine/Source/Runtime/GameplayAbilities/Public/GameplayEffect.h
--- a/Engine/Source/Runtime/GameplayAbilities/Public/GameplayEffect.h
+++ b/Engine/Source/Runtime/GameplayAbilities/Public/GameplayEffect.h
@@ -610,7 +610,7 @@ private:
 
 /** Saves list of modified attributes, to use for gameplay cues or later processing */
 USTRUCT()
-struct FGameplayEffectModifiedAttribute
+struct GAMEPLAYABILITIES_API FGameplayEffectModifiedAttribute
 {
 	GENERATED_USTRUCT_BODY()
 
@@ -627,7 +627,7 @@ struct FGameplayEffectModifiedAttribute
 
 /** Struct used to hold the result of a gameplay attribute capture; Initially seeded by definition data, but then populated by ability system component when appropriate */
 USTRUCT()
-struct FGameplayEffectAttributeCaptureSpec
+struct GAMEPLAYABILITIES_API FGameplayEffectAttributeCaptureSpec
 {
 	// Allow these as friends so they can seed the aggregator, which we don't otherwise want exposed
 	friend struct FActiveGameplayEffectsContainer;
