Fix prediction issues with WaitInputRelease/Press nodes

diff --git a/Engine/Source/Runtime/GameplayAbilities/Private/Abilities/Tasks/AbilityTask_WaitInputPress.cpp b/Engine/Source/Runtime/GameplayAbilities/Private/Abilities/Tasks/AbilityTask_WaitInputPress.cpp
--- a/Engine/Source/Runtime/GameplayAbilities/Private/Abilities/Tasks/AbilityTask_WaitInputPress.cpp
+++ b/Engine/Source/Runtime/GameplayAbilities/Private/Abilities/Tasks/AbilityTask_WaitInputPress.cpp
@@ -26,14 +26,6 @@ void UAbilityTask_WaitInputPress::OnPressCallback(int32 InputID)
 		return;
 	}
 
-	FScopedPredictionWindow	ScopedPrediction(MyAbility);
-
-	if (MyAbility->GetCurrentActivationInfo().ActivationMode != EGameplayAbilityActivationMode::Authority)
-	{
-		// Tell the server we pressed.
-		AbilitySystemComponent->ServerInputPress(MyAbility->GetCurrentAbilitySpecHandle(), ScopedPrediction.ScopedPredictionKey);
-	}
-
 	// We are done. Kill us so we don't keep getting broadcast messages
 	OnPress.Broadcast(ElapsedTime);
 	EndTask();
diff --git a/Engine/Source/Runtime/GameplayAbilities/Private/Abilities/Tasks/AbilityTask_WaitInputRelease.cpp b/Engine/Source/Runtime/GameplayAbilities/Private/Abilities/Tasks/AbilityTask_WaitInputRelease.cpp
--- a/Engine/Source/Runtime/GameplayAbilities/Private/Abilities/Tasks/AbilityTask_WaitInputRelease.cpp
+++ b/Engine/Source/Runtime/GameplayAbilities/Private/Abilities/Tasks/AbilityTask_WaitInputRelease.cpp
@@ -26,14 +26,6 @@ void UAbilityTask_WaitInputRelease::OnReleaseCallback(int32 InputID)
 		return;
 	}
 
-	FScopedPredictionWindow	ScopedPrediction(MyAbility);
-
-	if (MyAbility->GetCurrentActivationInfo().ActivationMode != EGameplayAbilityActivationMode::Authority)
-	{
-		// Tell the server we released.
-		AbilitySystemComponent->ServerInputRelease(MyAbility->GetCurrentAbilitySpecHandle(), ScopedPrediction.ScopedPredictionKey);
-	}
-
 	// We are done. Kill us so we don't keep getting broadcast messages
 	OnRelease.Broadcast(ElapsedTime);
 	EndTask();
diff --git a/Engine/Source/Runtime/GameplayAbilities/Private/AbilitySystemComponent.cpp b/Engine/Source/Runtime/GameplayAbilities/Private/AbilitySystemComponent.cpp
--- a/Engine/Source/Runtime/GameplayAbilities/Private/AbilitySystemComponent.cpp
+++ b/Engine/Source/Runtime/GameplayAbilities/Private/AbilitySystemComponent.cpp
@@ -1065,9 +1065,11 @@ void UAbilitySystemComponent::DisplayDebug(class UCanvas* Canvas, const class FD
 				const FModifierSpec& ModSpec = ActiveGE.Spec.Modifiers[ModIdx];
 				const FGameplayModifierInfo& ModInfo = ActiveGE.Spec.Def->Modifiers[ModIdx];
 
+				// Do a quick Qualifies() check to see if this mod is active.
 				FAggregatorMod TempMod;
 				TempMod.SourceTagReqs = &ModInfo.SourceTags;
 				TempMod.TargetTagReqs = &ModInfo.TargetTags;
+				TempMod.IsPredicted = false;
 
 				FAggregatorEvaluateParameters EmptyParams;
 				bool IsActivelyModifyingAttribute = TempMod.Qualifies(EmptyParams);
