Fixed a comment.

diff --git a/Engine/Source/Runtime/GameplayAbilities/Public/Abilities/Tasks/AbilityTask_WaitGameplayEffectRemoved.h b/Engine/Source/Runtime/GameplayAbilities/Public/Abilities/Tasks/AbilityTask_WaitGameplayEffectRemoved.h
--- a/Engine/Source/Runtime/GameplayAbilities/Public/Abilities/Tasks/AbilityTask_WaitGameplayEffectRemoved.h
+++ b/Engine/Source/Runtime/GameplayAbilities/Public/Abilities/Tasks/AbilityTask_WaitGameplayEffectRemoved.h
@@ -30,7 +30,7 @@ class UAbilityTask_WaitGameplayEffectRemoved : public UAbilityTask
 	UFUNCTION()
 	void OnGameplayEffectRemoved();
 
-	/** Wait until an overlap occurs. This will need to be better fleshed out so we can specify game specific collision requirements */
+	/** Wait until the specified gameplay effect is removed. */
 	UFUNCTION(BlueprintCallable, Category="Ability|Tasks", meta = (HidePin = "WorldContextObject", DefaultToSelf = "WorldContextObject", BlueprintInternalUseOnly = "TRUE"))
 	static UAbilityTask_WaitGameplayEffectRemoved* WaitForGameplayEffectRemoved(UObject* WorldContextObject, FActiveGameplayEffectHandle Handle);
 
