GCNAs can tick, by designer request.

diff --git a/Engine/Source/Runtime/GameplayAbilities/Private/GameplayCueNotify_Actor.cpp b/Engine/Source/Runtime/GameplayAbilities/Private/GameplayCueNotify_Actor.cpp
--- a/Engine/Source/Runtime/GameplayAbilities/Private/GameplayCueNotify_Actor.cpp
+++ b/Engine/Source/Runtime/GameplayAbilities/Private/GameplayCueNotify_Actor.cpp
@@ -10,6 +10,7 @@ AGameplayCueNotify_Actor::AGameplayCueNotify_Actor(const FObjectInitializer& Obj
 : Super(ObjectInitializer)
 {
 	IsOverride = true;
+	PrimaryActorTick.bCanEverTick = true;
 }
 
 #if WITH_EDITOR
