Added an actor array to gameplay cue parameters so that target-selection data is transmitted.

diff --git a/Engine/Source/Runtime/GameplayAbilities/Private/AbilitySystemBlueprintLibrary.cpp b/Engine/Source/Runtime/GameplayAbilities/Private/AbilitySystemBlueprintLibrary.cpp
--- a/Engine/Source/Runtime/GameplayAbilities/Private/AbilitySystemBlueprintLibrary.cpp
+++ b/Engine/Source/Runtime/GameplayAbilities/Private/AbilitySystemBlueprintLibrary.cpp
@@ -272,6 +272,11 @@ bool UAbilitySystemBlueprintLibrary::IsInstigatorLocallyControlled(FGameplayCueP
 }
 
 
+TArray<TWeakObjectPtr<AActor>> UAbilitySystemBlueprintLibrary::GetActors(FGameplayCueParameters Parameters)
+{
+	return Parameters.EffectContext.GetActors();
+}
+
 FHitResult UAbilitySystemBlueprintLibrary::GetHitResult(FGameplayCueParameters Parameters)
 {
 	if (Parameters.EffectContext.GetHitResult())
diff --git a/Engine/Source/Runtime/GameplayAbilities/Private/GameplayAbilityTargetTypes.cpp b/Engine/Source/Runtime/GameplayAbilities/Private/GameplayAbilityTargetTypes.cpp
--- a/Engine/Source/Runtime/GameplayAbilities/Private/GameplayAbilityTargetTypes.cpp
+++ b/Engine/Source/Runtime/GameplayAbilities/Private/GameplayAbilityTargetTypes.cpp
@@ -68,6 +68,11 @@ TArray<FActiveGameplayEffectHandle> FGameplayAbilityTargetData::ApplyGameplayEff
 
 void FGameplayAbilityTargetData::AddTargetDataToContext(FGameplayEffectContextHandle& Context)
 {
+	if (GetActors().Num() > 0)
+	{
+		Context.AddActors(GetActors());
+	}
+
 	if (HasHitResult())
 	{
 		Context.AddHitResult(*GetHitResult());
diff --git a/Engine/Source/Runtime/GameplayAbilities/Private/GameplayEffectTypes.cpp b/Engine/Source/Runtime/GameplayAbilities/Private/GameplayEffectTypes.cpp
--- a/Engine/Source/Runtime/GameplayAbilities/Private/GameplayEffectTypes.cpp
+++ b/Engine/Source/Runtime/GameplayAbilities/Private/GameplayEffectTypes.cpp
@@ -41,6 +41,17 @@ void FGameplayEffectContext::AddInstigator(class AActor *InInstigator, class AAc
 	}
 }
 
+void FGameplayEffectContext::AddActors(TArray<TWeakObjectPtr<AActor>> InActors, bool bReset)
+{
+	if (bReset && Actors.Num())
+	{
+		Actors.Reset();
+	}
+
+	check(!Actors.Num());
+	Actors = InActors;
+}
+
 void FGameplayEffectContext::AddHitResult(const FHitResult InHitResult, bool bReset)
 {
 	if (bReset && HitResult.IsValid())
@@ -62,6 +73,13 @@ bool FGameplayEffectContext::NetSerialize(FArchive& Ar, class UPackageMap* Map,
 	Ar << Instigator;
 	Ar << EffectCauser;
 
+	int32 ActorArraySize = Actors.Num();		//This can probably be smaller than int32
+	Ar << ActorArraySize;
+	if (ActorArraySize > 0)
+	{
+		Ar << Actors;
+	}
+
 	bool HasHitResults = HitResult.IsValid();
 	Ar << HasHitResults;
 	if (Ar.IsLoading())
@@ -115,7 +133,7 @@ void FGameplayEffectContext::GetOwnedGameplayTags(OUT FGameplayTagContainer &Tag
 	{
 		TagInterface->GetOwnedGameplayTags(TagContainer);
 	}
-	else if(InstigatorAbilitySystemComponent)
+	else if (InstigatorAbilitySystemComponent)
 	{
 		InstigatorAbilitySystemComponent->GetOwnedGameplayTags(TagContainer);
 	}
@@ -407,4 +425,4 @@ const FGameplayTagContainer& FTagContainerAggregator::GetSpecTags() const
 {
 	CacheIsValid = false;
 	return CapturedSpecTags;
-}
\ No newline at end of file
+}
diff --git a/Engine/Source/Runtime/GameplayAbilities/Public/AbilitySystemBlueprintLibrary.h b/Engine/Source/Runtime/GameplayAbilities/Public/AbilitySystemBlueprintLibrary.h
--- a/Engine/Source/Runtime/GameplayAbilities/Public/AbilitySystemBlueprintLibrary.h
+++ b/Engine/Source/Runtime/GameplayAbilities/Public/AbilitySystemBlueprintLibrary.h
@@ -85,6 +85,9 @@ class GAMEPLAYABILITIES_API UAbilitySystemBlueprintLibrary : public UBlueprintFu
 	UFUNCTION(BlueprintPure, Category="Ability|GameplayCue")
 	static bool IsInstigatorLocallyControlled(FGameplayCueParameters Parameters);
 
+	UFUNCTION(BlueprintPure, Category = "Ability|GameplayCue")
+	static TArray<TWeakObjectPtr<AActor>> GetActors(FGameplayCueParameters Parameters);
+
 	UFUNCTION(BlueprintPure, Category = "Ability|GameplayCue")
 	static FHitResult GetHitResult(FGameplayCueParameters Parameters);
 
diff --git a/Engine/Source/Runtime/GameplayAbilities/Public/GameplayEffectTypes.h b/Engine/Source/Runtime/GameplayAbilities/Public/GameplayEffectTypes.h
--- a/Engine/Source/Runtime/GameplayAbilities/Public/GameplayEffectTypes.h
+++ b/Engine/Source/Runtime/GameplayAbilities/Public/GameplayEffectTypes.h
@@ -305,8 +305,15 @@ struct GAMEPLAYABILITIES_API FGameplayEffectContext
 		return InstigatorAbilitySystemComponent;
 	}
 
+	virtual void AddActors(TArray<TWeakObjectPtr<AActor>> InActor, bool bReset = false);
+
 	virtual void AddHitResult(const FHitResult InHitResult, bool bReset = false);
 
+	virtual const TArray<TWeakObjectPtr<AActor>> GetActors() const
+	{
+		return Actors;
+	}
+
 	virtual const FHitResult* GetHitResult() const
 	{
 		if (HitResult.IsValid())
@@ -343,6 +350,7 @@ struct GAMEPLAYABILITIES_API FGameplayEffectContext
 	{
 		FGameplayEffectContext* NewContext = new FGameplayEffectContext();
 		*NewContext = *this;
+		NewContext->AddActors(Actors);
 		if (GetHitResult())
 		{
 			// Does a deep copy of the hit result
@@ -369,6 +377,9 @@ protected:
 	UPROPERTY(NotReplicated)
 	UAbilitySystemComponent* InstigatorAbilitySystemComponent;
 
+	UPROPERTY()
+	TArray<TWeakObjectPtr<AActor>> Actors;
+
 	/** Trace information - may be NULL in many cases */
 	TSharedPtr<FHitResult>	HitResult;
 
@@ -511,6 +522,14 @@ struct FGameplayEffectContextHandle
 		return false;
 	}
 
+	void AddActors(TArray<TWeakObjectPtr<AActor>> InActors, bool bReset = false)
+	{
+		if (IsValid())
+		{
+			Data->AddActors(InActors, bReset);
+		}
+	}
+
 	void AddHitResult(const FHitResult InHitResult, bool bReset = false)
 	{
 		if (IsValid())
@@ -519,6 +538,11 @@ struct FGameplayEffectContextHandle
 		}
 	}
 
+	const TArray<TWeakObjectPtr<AActor>> GetActors() const
+	{
+		return Data->GetActors();
+	}
+
 	const FHitResult* GetHitResult() const
 	{
 		if (IsValid())
