Scoped lock protects against removal during periodic effect updates. No lock on add yet - haven't seen adds cause a crash.

diff --git a/Engine/Source/Runtime/GameplayAbilities/Private/AbilitySystemComponent.cpp b/Engine/Source/Runtime/GameplayAbilities/Private/AbilitySystemComponent.cpp
--- a/Engine/Source/Runtime/GameplayAbilities/Private/AbilitySystemComponent.cpp
+++ b/Engine/Source/Runtime/GameplayAbilities/Private/AbilitySystemComponent.cpp
@@ -426,18 +426,13 @@ void UAbilitySystemComponent::UpdateTagMap(const FGameplayTagContainer& Containe
 
 FActiveGameplayEffectHandle UAbilitySystemComponent::ApplyGameplayEffectSpecToTarget(OUT FGameplayEffectSpec &Spec, UAbilitySystemComponent *Target, FPredictionKey PredictionKey)
 {
-	if (HasNetworkAuthorityToApplyGameplayEffect(PredictionKey))
-	{
-		// Apply outgoing Effects to the Spec.
-		// Outgoing immunity may stop the outgoing effect from being applied to the target
-		return Target->ApplyGameplayEffectSpecToSelf(Spec, PredictionKey);
-	}
-
-	return FActiveGameplayEffectHandle();
+	return Target->ApplyGameplayEffectSpecToSelf(Spec, PredictionKey);
 }
 
 FActiveGameplayEffectHandle UAbilitySystemComponent::ApplyGameplayEffectSpecToSelf(OUT FGameplayEffectSpec &Spec, FPredictionKey PredictionKey)
 {
+	//FScopedActiveGameplayEffectLock AGELock;
+
 	// Check Network Authority
 	if (!HasNetworkAuthorityToApplyGameplayEffect(PredictionKey))
 	{
@@ -449,7 +444,7 @@ FActiveGameplayEffectHandle UAbilitySystemComponent::ApplyGameplayEffectSpecToSe
 	{
 		if(IsOwnerActorAuthoritative())
 		{
-			// Server continue with invliad prediction key
+			// Server continue with invalid prediction key
 			PredictionKey = FPredictionKey();
 		}
 		else
@@ -520,7 +515,7 @@ FActiveGameplayEffectHandle UAbilitySystemComponent::ApplyGameplayEffectSpecToSe
 			MyHandle = NewActiveEffect.Handle;
 			OurCopyOfSpec = &NewActiveEffect.Spec;
 		}
-		
+
 		if (!OurCopyOfSpec)
 		{
 			StackSpec = TSharedPtr<FGameplayEffectSpec>(new FGameplayEffectSpec(Spec));
@@ -531,7 +526,7 @@ FActiveGameplayEffectHandle UAbilitySystemComponent::ApplyGameplayEffectSpecToSe
 
 		// if necessary add a modifier to OurCopyOfSpec to force it to have an infinite duration
 		if (bTreatAsInfiniteDuration)
-		{			
+		{
 			// This should just be a straight set of the duration float now
 			OurCopyOfSpec->SetDuration(UGameplayEffect::INFINITE_DURATION);
 		}
diff --git a/Engine/Source/Runtime/GameplayAbilities/Private/GameplayEffect.cpp b/Engine/Source/Runtime/GameplayAbilities/Private/GameplayEffect.cpp
--- a/Engine/Source/Runtime/GameplayAbilities/Private/GameplayEffect.cpp
+++ b/Engine/Source/Runtime/GameplayAbilities/Private/GameplayEffect.cpp
@@ -997,6 +997,7 @@ void FActiveGameplayEffectsContainer::ExecuteActiveEffectsFrom(FGameplayEffectSp
 
 void FActiveGameplayEffectsContainer::ExecutePeriodicGameplayEffect(FActiveGameplayEffectHandle Handle)
 {
+	FScopedActiveGameplayEffectLock AGELock;
 	FActiveGameplayEffect* ActiveEffect = GetActiveGameplayEffect(Handle);
 	if (ActiveEffect)
 	{
@@ -1018,6 +1019,26 @@ FActiveGameplayEffect* FActiveGameplayEffectsContainer::GetActiveGameplayEffect(
 	return nullptr;
 }
 
+int32 FScopedActiveGameplayEffectLock::AGELockCount = 0;
+FActiveGameplayEffectActionHandle FScopedActiveGameplayEffectLock::DeferredAGEActions;
+FScopedActiveGameplayEffectLock::FScopedActiveGameplayEffectLock()
+{
+	++AGELockCount;
+}
+
+FScopedActiveGameplayEffectLock::~FScopedActiveGameplayEffectLock()
+{
+	--AGELockCount;
+	if (!IsLockInEffect())
+	{
+		for (int32 i = 0; i < DeferredAGEActions.Num(); ++i)
+		{
+			DeferredAGEActions.Get(i)->PerformAction();
+		}
+		DeferredAGEActions.Clear();
+	}
+}
+
 FAggregatorRef& FActiveGameplayEffectsContainer::FindOrCreateAttributeAggregator(FGameplayAttribute Attribute)
 {
 	FAggregatorRef *RefPtr = AttributeAggregatorMap.Find(Attribute);
@@ -1189,7 +1210,7 @@ float FActiveGameplayEffectsContainer::GetGameplayEffectMagnitude(FActiveGamepla
 	return -1.f;
 }
 
-bool FActiveGameplayEffectsContainer::IsGameplayEffectActive(FActiveGameplayEffectHandle Handle) const
+bool FActiveGameplayEffectsContainer::IsGameplayEffectActive(FActiveGameplayEffectHandle Handle, bool IncludeEffectsBlockedByStackingRules) const
 {
 	// Could make this a map for quicker lookup
 	for (const FActiveGameplayEffect& Effect : GameplayEffects)
@@ -1200,7 +1221,7 @@ bool FActiveGameplayEffectsContainer::IsGameplayEffectActive(FActiveGameplayEffe
 			if (Effect.Spec.GetStackingType() != EGameplayEffectStackingPolicy::Unlimited &&
 				Effect.Spec.bTopOfStack == false)
 			{
-				return false;
+				return IncludeEffectsBlockedByStackingRules;
 			}
 			return true;
 		}
@@ -1336,6 +1357,8 @@ FActiveGameplayEffect& FActiveGameplayEffectsContainer::CreateNewActiveGameplayE
 {
 	SCOPE_CYCLE_COUNTER(STAT_CreateNewActiveGameplayEffect);
 
+	//TODO Add FScopedActiveGameplayEffectLock here once we can get a crash to happen
+
 	if (Owner && Owner->OwnerActor)
 	{
 		if ( IsNetAuthority() )
@@ -1512,6 +1535,12 @@ bool FActiveGameplayEffectsContainer::RemoveActiveGameplayEffect(FActiveGameplay
 	{
 		if(GameplayEffects[Idx].Handle == Handle)
 		{
+			//Block removal if we're scope-locked
+			if (FScopedActiveGameplayEffectLock::IsLockInEffect())
+			{
+				FScopedActiveGameplayEffectLock::AddAction(new FActiveGameplayEffectAction_Remove(*this, Handle));
+				return true;		//We are assuming the "ensure(Idx < GameplayEffects.Num())" would be passed.
+			}
 			return InternalRemoveActiveGameplayEffect(Idx);
 		}
 	}
@@ -1529,6 +1558,8 @@ bool FActiveGameplayEffectsContainer::InternalRemoveActiveGameplayEffect(int32 I
 	{
 		FActiveGameplayEffect& Effect = GameplayEffects[Idx];
 
+		check(!FScopedActiveGameplayEffectLock::IsLockInEffect());
+
 		bool ShouldInvokeGameplayCueEvent = true;
 		const bool bIsNetAuthority = IsNetAuthority();
 		if (!bIsNetAuthority && Effect.PredictionKey.IsValidKey() && Effect.PredictionKey.WasReceived() == false)
@@ -1566,7 +1597,7 @@ bool FActiveGameplayEffectsContainer::InternalRemoveActiveGameplayEffect(int32 I
 		MarkArrayDirty();
 
 		// Hack: force netupdate on owner. This isn't really necessary in real gameplay but is nice
-		// during debugging where breakpoints or pausing can messup network update times. Open issue
+		// during debugging where breakpoints or pausing can mess up network update times. Open issue
 		// with network team.
 		Owner->GetOwner()->ForceNetUpdate();
 
diff --git a/Engine/Source/Runtime/GameplayAbilities/Public/GameplayEffect.h b/Engine/Source/Runtime/GameplayAbilities/Public/GameplayEffect.h
--- a/Engine/Source/Runtime/GameplayAbilities/Public/GameplayEffect.h
+++ b/Engine/Source/Runtime/GameplayAbilities/Public/GameplayEffect.h
@@ -1056,6 +1056,7 @@ struct FActiveGameplayEffectQuery
 	FGameplayAttribute ModifyingAttribute;
 };
 
+
 /**
  * Active GameplayEffects Container
  *	-Bucket of ActiveGameplayEffects
@@ -1097,7 +1098,7 @@ struct FActiveGameplayEffectsContainer : public FFastArraySerializer
 
 	// returns true if the handle points to an effect in this container that is not a stacking effect or an effect in this container that does stack and is applied by the current stacking rules
 	// returns false if the handle points to an effect that is not in this container or is not applied because of the current stacking rules
-	bool IsGameplayEffectActive(FActiveGameplayEffectHandle Handle) const;
+	bool IsGameplayEffectActive(FActiveGameplayEffectHandle Handle, bool IncludeEffectsBlockedByStackingRules = false) const;
 
 	/**
 	 * Get the source tags from the gameplay spec represented by the specified handle, if possible
@@ -1233,6 +1234,106 @@ struct TStructOpsTypeTraits< FActiveGameplayEffectsContainer > : public TStructO
 };
 
 
+USTRUCT()
+struct FActiveGameplayEffectAction
+{
+	GENERATED_USTRUCT_BODY()
+	FActiveGameplayEffectAction()
+	{
+	}
+
+	virtual void PerformAction()
+	{
+	}
+};
+
+USTRUCT()
+struct FActiveGameplayEffectAction_Remove : public FActiveGameplayEffectAction
+{
+	GENERATED_USTRUCT_BODY()
+	FActiveGameplayEffectAction_Remove()
+	{
+	}
+
+	FActiveGameplayEffectAction_Remove(FActiveGameplayEffectsContainer& InContainer, FActiveGameplayEffectHandle& InHandle)
+		: Container(InContainer)
+		, Handle(InHandle)
+	{
+	}
+
+	virtual void PerformAction()
+	{
+		Container.RemoveActiveGameplayEffect(Handle);
+	}
+
+	UPROPERTY()
+	FActiveGameplayEffectsContainer Container;
+
+	UPROPERTY()
+	FActiveGameplayEffectHandle Handle;
+};
+
+USTRUCT()
+struct FActiveGameplayEffectActionHandle
+{
+	GENERATED_USTRUCT_BODY()
+
+	FActiveGameplayEffectActionHandle()
+	{
+	}
+
+	FActiveGameplayEffectActionHandle(struct FActiveGameplayEffectAction* DataPtr)
+	{
+		Add(DataPtr);
+	}
+
+	void Add(struct FActiveGameplayEffectAction* DataPtr)
+	{
+		Data.Add(TSharedPtr<FActiveGameplayEffectAction>(DataPtr));
+	}
+
+	void Clear()
+	{
+		Data.Reset();
+	}
+
+	int32 Num()
+	{
+		return Data.Num();
+	}
+
+	FActiveGameplayEffectAction* Get(int32 Index)
+	{
+		check((Index >= 0) && (Index < Data.Num()));
+		return Data[Index].Get();
+	}
+
+	TArray<TSharedPtr<FActiveGameplayEffectAction>>	Data;
+};
+
+struct FScopedActiveGameplayEffectLock
+{
+	FScopedActiveGameplayEffectLock();
+	~FScopedActiveGameplayEffectLock();
+
+private:
+	static int32 AGELockCount;
+	static FActiveGameplayEffectActionHandle DeferredAGEActions;
+
+public:
+	static void AddAction(FActiveGameplayEffectAction* NewAction)
+	{
+		check(IsLockInEffect());
+		DeferredAGEActions.Add(NewAction);
+	}
+
+	static bool IsLockInEffect()
+	{
+		return (AGELockCount ? true : false);
+	}
+};
+
+
 /** Allows blueprints to generate a GameplayEffectSpec once and then reference it by handle, to apply it multiple times/multiple targets. */
 USTRUCT(BlueprintType)
 struct GAMEPLAYABILITIES_API FGameplayEffectSpecHandle
