-Added coeffeciet/pre/post additions to custom class magnitudes. Renamed custom calculation function to make it clear it is calculating a base value.
-SetNumericAttribute replaced with SetNumericAttributeBase: made this path go through attribute aggregators when called from outside. Added SetNumericAttribute_Internal for setting the actual attribute set property internally.

diff --git a/Engine/Source/Editor/GameplayAbilitiesEditor/Private/GameplayEffectModifierMagnitudeDetails.cpp b/Engine/Source/Editor/GameplayAbilitiesEditor/Private/GameplayEffectModifierMagnitudeDetails.cpp
--- a/Engine/Source/Editor/GameplayAbilitiesEditor/Private/GameplayEffectModifierMagnitudeDetails.cpp
+++ b/Engine/Source/Editor/GameplayAbilitiesEditor/Private/GameplayEffectModifierMagnitudeDetails.cpp
@@ -30,7 +30,7 @@ void FGameplayEffectModifierMagnitudeDetails::CustomizeChildren(TSharedRef<IProp
 	PropertyToCalcEnumMap.Empty();
 	PropertyToCalcEnumMap.Add(FindFieldChecked<UProperty>(FGameplayEffectModifierMagnitude::StaticStruct(), GET_MEMBER_NAME_CHECKED(FGameplayEffectModifierMagnitude, ScalableFloatMagnitude)), EGameplayEffectMagnitudeCalculation::ScalableFloat);
 	PropertyToCalcEnumMap.Add(FindFieldChecked<UProperty>(FGameplayEffectModifierMagnitude::StaticStruct(), GET_MEMBER_NAME_CHECKED(FGameplayEffectModifierMagnitude, AttributeBasedMagnitude)), EGameplayEffectMagnitudeCalculation::AttributeBased);
-	PropertyToCalcEnumMap.Add(FindFieldChecked<UProperty>(FGameplayEffectModifierMagnitude::StaticStruct(), GET_MEMBER_NAME_CHECKED(FGameplayEffectModifierMagnitude, CalculationClassMagnitude)), EGameplayEffectMagnitudeCalculation::CustomCalculationClass);
+	PropertyToCalcEnumMap.Add(FindFieldChecked<UProperty>(FGameplayEffectModifierMagnitude::StaticStruct(), GET_MEMBER_NAME_CHECKED(FGameplayEffectModifierMagnitude, CustomMagnitude)), EGameplayEffectMagnitudeCalculation::CustomCalculationClass);
 
 	// Hook into calculation type changes
 	MagnitudeCalculationTypePropertyHandle = StructPropertyHandle->GetChildHandle(GET_MEMBER_NAME_CHECKED(FGameplayEffectModifierMagnitude, MagnitudeCalculationType));
@@ -56,11 +56,11 @@ void FGameplayEffectModifierMagnitudeDetails::CustomizeChildren(TSharedRef<IProp
 		PropRow.Visibility(TAttribute<EVisibility>::Create(TAttribute<EVisibility>::FGetter::CreateSP(this, &FGameplayEffectModifierMagnitudeDetails::GetMagnitudeCalculationPropertyVisibility, AttributeBasedMagnitudePropertyHandle->GetProperty())));
 	}
 
-	TSharedPtr<IPropertyHandle> CalculationClassMagnitudePropertyHandle  = StructPropertyHandle->GetChildHandle(GET_MEMBER_NAME_CHECKED(FGameplayEffectModifierMagnitude, CalculationClassMagnitude));
-	if (CalculationClassMagnitudePropertyHandle.IsValid())
+	TSharedPtr<IPropertyHandle> CustomMagnitudePropertyHandle  = StructPropertyHandle->GetChildHandle(GET_MEMBER_NAME_CHECKED(FGameplayEffectModifierMagnitude, CustomMagnitude));
+	if (CustomMagnitudePropertyHandle.IsValid())
 	{
-		IDetailPropertyRow& PropRow = StructBuilder.AddChildProperty(CalculationClassMagnitudePropertyHandle.ToSharedRef());
-		PropRow.Visibility(TAttribute<EVisibility>::Create(TAttribute<EVisibility>::FGetter::CreateSP(this, &FGameplayEffectModifierMagnitudeDetails::GetMagnitudeCalculationPropertyVisibility, CalculationClassMagnitudePropertyHandle->GetProperty())));
+		IDetailPropertyRow& PropRow = StructBuilder.AddChildProperty(CustomMagnitudePropertyHandle.ToSharedRef());
+		PropRow.Visibility(TAttribute<EVisibility>::Create(TAttribute<EVisibility>::FGetter::CreateSP(this, &FGameplayEffectModifierMagnitudeDetails::GetMagnitudeCalculationPropertyVisibility, CustomMagnitudePropertyHandle->GetProperty())));
 	}
 }
 
diff --git a/Engine/Source/Runtime/GameplayAbilities/Private/AbilitySystemComponent.cpp b/Engine/Source/Runtime/GameplayAbilities/Private/AbilitySystemComponent.cpp
--- a/Engine/Source/Runtime/GameplayAbilities/Private/AbilitySystemComponent.cpp
+++ b/Engine/Source/Runtime/GameplayAbilities/Private/AbilitySystemComponent.cpp
@@ -144,8 +144,15 @@ bool UAbilitySystemComponent::HasNetworkAuthorityToApplyGameplayEffect(FPredicti
 	return (IsOwnerActorAuthoritative() || PredictionKey.IsValidForMorePrediction());
 }
 
-void UAbilitySystemComponent::SetNumericAttribute(const FGameplayAttribute &Attribute, float NewFloatValue)
+void UAbilitySystemComponent::SetNumericAttributeBase(const FGameplayAttribute &Attribute, float NewFloatValue)
 {
+	// Go through our active gameplay effects container so that aggregation/mods are handled properly.
+	ActiveGameplayEffects.SetAttributeBaseValue(Attribute, NewFloatValue);
+}
+
+void UAbilitySystemComponent::SetNumericAttribute_Internal(const FGameplayAttribute &Attribute, float NewFloatValue)
+{
+	// Set the attribute directly: update the UProperty on the attribute set.
 	const UAttributeSet* AttributeSet = GetAttributeSubobjectChecked(Attribute.GetAttributeSetClass());
 	Attribute.SetNumericValueChecked(NewFloatValue, const_cast<UAttributeSet*>(AttributeSet));
 }
@@ -161,6 +168,11 @@ float UAbilitySystemComponent::GetNumericAttribute(const FGameplayAttribute &Att
 	return Attribute.GetNumericValueChecked(AttributeSet);
 }
 
+void UAbilitySystemComponent::ApplyModToAttribute(const FGameplayAttribute &Attribute, TEnumAsByte<EGameplayModOp::Type> ModifierOp, float ModifierMagnitude)
+{
+	ActiveGameplayEffects.ApplyModToAttribute(Attribute, ModifierOp, ModifierMagnitude);
+}
+
 FGameplayEffectSpecHandle UAbilitySystemComponent::MakeOutgoingSpec(TSubclassOf<UGameplayEffect> GameplayEffectClass, float Level, FGameplayEffectContextHandle Context) const
 {
 	SCOPE_CYCLE_COUNTER(STAT_GetOutgoingSpec);
diff --git a/Engine/Source/Runtime/GameplayAbilities/Private/AttributeSet.cpp b/Engine/Source/Runtime/GameplayAbilities/Private/AttributeSet.cpp
--- a/Engine/Source/Runtime/GameplayAbilities/Private/AttributeSet.cpp
+++ b/Engine/Source/Runtime/GameplayAbilities/Private/AttributeSet.cpp
@@ -324,7 +324,7 @@ void FAttributeSetInitter::InitAttributeSetDefaults(UAbilitySystemComponent* Abi
 				if (Set->ShouldInitProperty(bInitialInit, DataPair.Property))
 				{
 					FGameplayAttribute AttributeToModify(DataPair.Property);
-					AbilitySystemComponent->SetNumericAttribute(AttributeToModify, DataPair.Value);
+					AbilitySystemComponent->SetNumericAttributeBase(AttributeToModify, DataPair.Value);
 				}
 			}
 		}		
diff --git a/Engine/Source/Runtime/GameplayAbilities/Private/GameplayEffect.cpp b/Engine/Source/Runtime/GameplayAbilities/Private/GameplayEffect.cpp
--- a/Engine/Source/Runtime/GameplayAbilities/Private/GameplayEffect.cpp
+++ b/Engine/Source/Runtime/GameplayAbilities/Private/GameplayEffect.cpp
@@ -183,6 +183,23 @@ float FAttributeBasedFloat::CalculateMagnitude(const FGameplayEffectSpec& InRele
 	return ((Coefficient.GetValueAtLevel(SpecLvl) * (AttribValue + PreMultiplyAdditiveValue.GetValueAtLevel(SpecLvl))) + PostMultiplyAdditiveValue.GetValueAtLevel(SpecLvl));
 }
 
+// --------------------------------------------------------------------------------------------------------------------------------------------------------
+//
+//	FCustomCalculationBasedFloat
+//
+// --------------------------------------------------------------------------------------------------------------------------------------------------------
+
+float FCustomCalculationBasedFloat::CalculateMagnitude(const FGameplayEffectSpec& InRelevantSpec) const
+{
+	const UGameplayModMagnitudeCalculation* CalcCDO = CalculationClassMagnitude->GetDefaultObject<UGameplayModMagnitudeCalculation>();
+	check(CalcCDO);
+
+	float CustomBaseValue = CalcCDO->CalculateBaseMagnitude(InRelevantSpec);
+
+	const float SpecLvl = InRelevantSpec.GetLevel();
+	return ((Coefficient.GetValueAtLevel(SpecLvl) * (CustomBaseValue + PreMultiplyAdditiveValue.GetValueAtLevel(SpecLvl))) + PostMultiplyAdditiveValue.GetValueAtLevel(SpecLvl));
+}
+
 // --------------------------------------------------------------------------------------------------------------------------------------------------------
 //
 //	FGameplayEffectMagnitude
@@ -219,10 +236,7 @@ bool FGameplayEffectModifierMagnitude::AttemptCalculateMagnitude(const FGameplay
 
 			case EGameplayEffectMagnitudeCalculation::CustomCalculationClass:
 			{
-				const UGameplayModMagnitudeCalculation* CalcCDO = CalculationClassMagnitude->GetDefaultObject<UGameplayModMagnitudeCalculation>();
-				check(CalcCDO);
-
-				OutCalculatedMagnitude = CalcCDO->CalculateMagnitude(InRelevantSpec);
+				OutCalculatedMagnitude = CustomMagnitude.CalculateMagnitude(InRelevantSpec);
 			}
 			break;
 
@@ -278,9 +292,9 @@ void FGameplayEffectModifierMagnitude::GetAttributeCaptureDefinitions(OUT TArray
 
 		case EGameplayEffectMagnitudeCalculation::CustomCalculationClass:
 		{
-			if (CalculationClassMagnitude)
+			if (CustomMagnitude.CalculationClassMagnitude)
 			{
-				const UGameplayModMagnitudeCalculation* CalcCDO = CalculationClassMagnitude->GetDefaultObject<UGameplayModMagnitudeCalculation>();
+				const UGameplayModMagnitudeCalculation* CalcCDO = CustomMagnitude.CalculationClassMagnitude->GetDefaultObject<UGameplayModMagnitudeCalculation>();
 				check(CalcCDO);
 
 				OutCaptureDefs.Append(CalcCDO->GetAttributeCaptureDefinitions());
@@ -1275,7 +1289,7 @@ void FActiveGameplayEffectsContainer::CaptureAttributeForGameplayEffect(OUT FGam
 void FActiveGameplayEffectsContainer::InternalUpdateNumericalAttribute(FGameplayAttribute Attribute, float NewValue, const FGameplayEffectModCallbackData* ModData)
 {
 	ABILITY_LOG(Log, TEXT("Property %s new value is: %.2f"), *Attribute.GetName(), NewValue);
-	Owner->SetNumericAttribute(Attribute, NewValue);
+	Owner->SetNumericAttribute_Internal(Attribute, NewValue);
 
 	FOnGameplayAttributeChange* Delegate = AttributeChangeDelegates.Find(Attribute);
 	if (Delegate)
@@ -1284,6 +1298,22 @@ void FActiveGameplayEffectsContainer::InternalUpdateNumericalAttribute(FGameplay
 	}
 }
 
+void FActiveGameplayEffectsContainer::SetAttributeBaseValue(FGameplayAttribute Attribute, float NewBaseValue)
+{
+	FAggregatorRef* RefPtr = AttributeAggregatorMap.Find(Attribute);
+	if (RefPtr)
+	{
+		// There is an aggregator for this attribute, so set the base value. The dirty callback chain
+		// will update the actual AttributeSet property value for us.
+		RefPtr->Get()->SetBaseValue(NewBaseValue);
+	}
+	else
+	{
+		// There is no aggregator yet, so we can just set the numeric value directly
+		InternalUpdateNumericalAttribute(Attribute, NewBaseValue, nullptr);
+	}
+}
+
 bool FActiveGameplayEffectsContainer::InternalExecuteMod(FGameplayEffectSpec& Spec, FGameplayModifierEvaluatedData& ModEvalData)
 {
 	bool bExecuted = false;
@@ -1301,21 +1331,7 @@ bool FActiveGameplayEffectsContainer::InternalExecuteMod(FGameplayEffectSpec& Sp
 		 */
 		if (AttributeSet->PreGameplayEffectExecute(ExecuteData))
 		{
-			// Do we have active GE's that are already modifying this?
-			FAggregatorRef* RefPtr = AttributeAggregatorMap.Find(ModEvalData.Attribute);
-			if (RefPtr)
-			{
-				ABILITY_LOG(Log, TEXT("Property %s has active mods. Adding to Aggregator."), *ModEvalData.Attribute.GetName());
-				RefPtr->Get()->ExecModOnBaseValue(ModEvalData.ModifierOp, ModEvalData.Magnitude);
-			}
-			else
-			{
-				// Modify the property in place, without putting it in the AttributeAggregatorMap map
-				float		CurrentValueOfProperty = Owner->GetNumericAttribute(ModEvalData.Attribute);
-				const float NewPropertyValue = FAggregator::StaticExecModOnBaseValue(CurrentValueOfProperty, ModEvalData.ModifierOp, ModEvalData.Magnitude);
-
-				InternalUpdateNumericalAttribute(ModEvalData.Attribute, NewPropertyValue, &ExecuteData);
-			}
+			ApplyModToAttribute(ModEvalData.Attribute, ModEvalData.ModifierOp, ModEvalData.Magnitude, &ExecuteData);
 
 			FGameplayEffectModifiedAttribute* ModifiedAttribute = Spec.GetModifiedAttribute(ModEvalData.Attribute);
 			if (!ModifiedAttribute)
@@ -1340,6 +1356,25 @@ bool FActiveGameplayEffectsContainer::InternalExecuteMod(FGameplayEffectSpec& Sp
 	return bExecuted;
 }
 
+void FActiveGameplayEffectsContainer::ApplyModToAttribute(const FGameplayAttribute &Attribute, TEnumAsByte<EGameplayModOp::Type> ModifierOp, float ModifierMagnitude, const FGameplayEffectModCallbackData* ModData)
+{
+	// Do we have active GE's that are already modifying this?
+	FAggregatorRef* RefPtr = AttributeAggregatorMap.Find(Attribute);
+	if (RefPtr)
+	{
+		ABILITY_LOG(Log, TEXT("Property %s has active mods. Adding to Aggregator."), *Attribute.GetName());
+		RefPtr->Get()->ExecModOnBaseValue(ModifierOp, ModifierMagnitude);
+	}
+	else
+	{
+		// Modify the property in place, without putting it in the AttributeAggregatorMap map
+		float		CurrentValueOfProperty = Owner->GetNumericAttribute(Attribute);
+		const float NewPropertyValue = FAggregator::StaticExecModOnBaseValue(CurrentValueOfProperty, ModifierOp, ModifierMagnitude);
+
+		InternalUpdateNumericalAttribute(Attribute, NewPropertyValue, ModData);
+	}
+}
+
 void FActiveGameplayEffectsContainer::StacksNeedToRecalculate()
 {
 	if (!bNeedToRecalculateStacks)
diff --git a/Engine/Source/Runtime/GameplayAbilities/Private/GameplayModMagnitudeCalculation.cpp b/Engine/Source/Runtime/GameplayAbilities/Private/GameplayModMagnitudeCalculation.cpp
--- a/Engine/Source/Runtime/GameplayAbilities/Private/GameplayModMagnitudeCalculation.cpp
+++ b/Engine/Source/Runtime/GameplayAbilities/Private/GameplayModMagnitudeCalculation.cpp
@@ -9,7 +9,25 @@ UGameplayModMagnitudeCalculation::UGameplayModMagnitudeCalculation(const FObject
 
 }
 
-float UGameplayModMagnitudeCalculation::CalculateMagnitude_Implementation(const FGameplayEffectSpec& Spec) const
+float UGameplayModMagnitudeCalculation::CalculateBaseMagnitude_Implementation(const FGameplayEffectSpec& Spec) const
 {
 	return 0.f;
 }
+
+
+bool UGameplayModMagnitudeCalculation::GetCapturedAttributeMagnitude(const FGameplayEffectAttributeCaptureDefinition& Def, const FGameplayEffectSpec& Spec, const FAggregatorEvaluateParameters& EvaluationParameters, OUT float& Magnitude) const
+{
+	const FGameplayEffectAttributeCaptureSpec* CaptureSpec = Spec.CapturedRelevantAttributes.FindCaptureSpecByDefinition(Def, true);
+	if (CaptureSpec == nullptr)
+	{
+		ABILITY_LOG(Error, TEXT("GetCapturedAttributeMagnitude unable to get capture spec."));
+		return false;
+	}
+	if (CaptureSpec->AttemptCalculateAttributeMagnitude(EvaluationParameters, Magnitude) == false)
+	{
+		ABILITY_LOG(Error, TEXT("GetCapturedAttributeMagnitude unable to calculate Health attribute magnitude."));
+		return false;
+	}
+
+	return true;
+}
\ No newline at end of file
diff --git a/Engine/Source/Runtime/GameplayAbilities/Public/AbilitySystemComponent.h b/Engine/Source/Runtime/GameplayAbilities/Public/AbilitySystemComponent.h
--- a/Engine/Source/Runtime/GameplayAbilities/Public/AbilitySystemComponent.h
+++ b/Engine/Source/Runtime/GameplayAbilities/Public/AbilitySystemComponent.h
@@ -108,7 +108,21 @@ class GAMEPLAYABILITIES_API UAbilitySystemComponent : public UActorComponent, pu
 	UPROPERTY(Replicated)
 	TArray<UAttributeSet*>	SpawnedAttributes;
 
-	void SetNumericAttribute(const FGameplayAttribute &Attribute, float NewFloatValue);
+	/** Sets the base value of an attribute. Existing active modifiers are NOT cleared and will act upon the new base value. */
+	void SetNumericAttributeBase(const FGameplayAttribute &Attribute, float NewBaseValue);
+
+	/**
+	 *	Applies an inplace mod to the given attribute. This correctly update the attribute's aggregator, updates the attribute set property,
+	 *	and invokes the OnDirty callbacks.
+	 *	
+	 *	This does not invoke Pre/PostGameplayEffectExecute calls on the attribute set. This does no tag checking, application requirements, immunity, etc.
+	 *	No GameplayEffectSpec is created or is applied!
+	 *
+	 *	This should only be used in cases where applying a real GameplayEffectSpec is too slow or not possible.
+	 */
+	void ApplyModToAttribute(const FGameplayAttribute &Attribute, TEnumAsByte<EGameplayModOp::Type> ModifierOp, float ModifierMagnitude);
+
+	/** Returns current (final) value of an attribute */
 	float GetNumericAttribute(const FGameplayAttribute &Attribute);
 
 	virtual void DisplayDebug(class UCanvas* Canvas, const class FDebugDisplayInfo& DebugDisplay, float& YL, float& YPos);
@@ -716,6 +730,9 @@ public:
 
 private:
 
+	/** Actually pushes the final attribute value to the attribute set's property. Should not be called by outside code since this does not go through the attribute aggregator system. */
+	void SetNumericAttribute_Internal(const FGameplayAttribute &Attribute, float NewFloatValue);
+
 	bool HasNetworkAuthorityToApplyGameplayEffect(FPredictionKey PredictionKey) const;
 
 	void ExecutePeriodicEffect(FActiveGameplayEffectHandle	Handle);
diff --git a/Engine/Source/Runtime/GameplayAbilities/Public/GameplayEffect.h b/Engine/Source/Runtime/GameplayAbilities/Public/GameplayEffect.h
--- a/Engine/Source/Runtime/GameplayAbilities/Public/GameplayEffect.h
+++ b/Engine/Source/Runtime/GameplayAbilities/Public/GameplayEffect.h
@@ -169,6 +169,48 @@ public:
 	FGameplayTagContainer TargetTagFilter;
 };
 
+USTRUCT()
+struct FCustomCalculationBasedFloat
+{
+	GENERATED_USTRUCT_BODY()
+
+	FCustomCalculationBasedFloat()
+		: CalculationClassMagnitude(nullptr)
+		, Coefficient(1.f)
+		, PreMultiplyAdditiveValue(0.f)
+		, PostMultiplyAdditiveValue(0.f)
+	{}
+
+public:
+
+	/**
+	 * Calculate and return the magnitude of the float given the specified gameplay effect spec.
+	 * 
+	 * @note:	This function assumes (and asserts on) the existence of the required captured attribute within the spec.
+	 *			It is the responsibility of the caller to verify that the spec is properly setup before calling this function.
+	 *			
+	 *	@param InRelevantSpec	Gameplay effect spec providing the backing attribute capture
+	 *	
+	 *	@return Evaluated magnitude based upon the spec & calculation policy
+	 */
+	float CalculateMagnitude(const FGameplayEffectSpec& InRelevantSpec) const;
+
+	UPROPERTY(EditDefaultsOnly, Category=CustomCalculation)
+	TSubclassOf<UGameplayModMagnitudeCalculation> CalculationClassMagnitude;
+
+	/** Coefficient to the custom calculation */
+	UPROPERTY(EditDefaultsOnly, Category=CustomCalculation)
+	FScalableFloat Coefficient;
+
+	/** Additive value to the attribute calculation, added in before the coefficient applies */
+	UPROPERTY(EditDefaultsOnly, Category=AttributeFloat)
+	FScalableFloat PreMultiplyAdditiveValue;
+
+	/** Additive value to the attribute calculation, added in after the coefficient applies */
+	UPROPERTY(EditDefaultsOnly, Category=AttributeFloat)
+	FScalableFloat PostMultiplyAdditiveValue;
+};
+
 /** Struct representing the magnitude of a gameplay effect modifier, potentially calculated in numerous different ways */
 USTRUCT()
 struct FGameplayEffectModifierMagnitude
@@ -182,7 +224,7 @@ public:
 		: MagnitudeCalculationType(EGameplayEffectMagnitudeCalculation::ScalableFloat)
 		, ScalableFloatMagnitude()
 		, AttributeBasedMagnitude()
-		, CalculationClassMagnitude(nullptr)
+		, CustomMagnitude()
 	{
 	}
 
@@ -234,7 +276,7 @@ protected:
 
 	/** Magnitude value represented by a custom calculation class */
 	UPROPERTY(EditDefaultsOnly, Category=Magnitude)
-	TSubclassOf<UGameplayModMagnitudeCalculation> CalculationClassMagnitude;
+	FCustomCalculationBasedFloat CustomMagnitude;
 
 	// @hack: @todo: This is temporary to aid in post-load fix-up w/o exposing members publicly
 	friend class UGameplayEffect;
@@ -1100,6 +1142,11 @@ struct FActiveGameplayEffectsContainer : public FFastArraySerializer
 	// returns false if the handle points to an effect that is not in this container or is not applied because of the current stacking rules
 	bool IsGameplayEffectActive(FActiveGameplayEffectHandle Handle, bool IncludeEffectsBlockedByStackingRules = false) const;
 
+	void SetAttributeBaseValue(FGameplayAttribute Attribute, float NewBaseValue);
+
+	/** Actually applies given mod to the attribute */
+	void ApplyModToAttribute(const FGameplayAttribute &Attribute, TEnumAsByte<EGameplayModOp::Type> ModifierOp, float ModifierMagnitude, const FGameplayEffectModCallbackData* ModData=nullptr);
+
 	/**
 	 * Get the source tags from the gameplay spec represented by the specified handle, if possible
 	 * 
diff --git a/Engine/Source/Runtime/GameplayAbilities/Public/GameplayModMagnitudeCalculation.h b/Engine/Source/Runtime/GameplayAbilities/Public/GameplayModMagnitudeCalculation.h
--- a/Engine/Source/Runtime/GameplayAbilities/Public/GameplayModMagnitudeCalculation.h
+++ b/Engine/Source/Runtime/GameplayAbilities/Public/GameplayModMagnitudeCalculation.h
@@ -19,12 +19,19 @@ public:
 	GENERATED_UCLASS_BODY()
 
 	/**
-	 * Calculate the magnitude of the gameplay effect modifier, given the specified spec
+	 * Calculate the base magnitude of the gameplay effect modifier, given the specified spec. Note that the owngin spec def can still modify this base value
+	 * with a coeffecient and pre/post multiply adds. see FCustomCalculationBasedFloat::CalculateMagnitude for details.
 	 * 
 	 * @param Spec	Gameplay effect spec to use to calculate the magnitude with
 	 * 
 	 * @return Computed magnitude of the modifier
 	 */
 	UFUNCTION(BlueprintNativeEvent, Category="Calculation")
-	float CalculateMagnitude(const FGameplayEffectSpec& Spec) const;
+	float CalculateBaseMagnitude(const FGameplayEffectSpec& Spec) const;
+
+protected:
+	
+
+	/** Convenience method to get attribute magnitude during a CalculateMagnitude call */
+	bool GetCapturedAttributeMagnitude(const FGameplayEffectAttributeCaptureDefinition& Def, const FGameplayEffectSpec& Spec, const FAggregatorEvaluateParameters& EvaluationParameters, OUT float& Magnitude) const;
 };
\ No newline at end of file
