fix clang compile

diff --git a/Engine/Source/Runtime/GameplayAbilities/Public/GameplayEffect.h b/Engine/Source/Runtime/GameplayAbilities/Public/GameplayEffect.h
--- a/Engine/Source/Runtime/GameplayAbilities/Public/GameplayEffect.h
+++ b/Engine/Source/Runtime/GameplayAbilities/Public/GameplayEffect.h
@@ -1289,6 +1289,11 @@ struct FActiveGameplayEffectAction
 	{
 	}
 
+	virtual ~FActiveGameplayEffectAction()
+	{
+
+	}
+
 	virtual void PerformAction()
 	{
 	}
