Reimplemented GE removal scope-lock to fix a bug.

diff --git a/Engine/Source/Runtime/GameplayAbilities/Private/GameplayEffect.cpp b/Engine/Source/Runtime/GameplayAbilities/Private/GameplayEffect.cpp
--- a/Engine/Source/Runtime/GameplayAbilities/Private/GameplayEffect.cpp
+++ b/Engine/Source/Runtime/GameplayAbilities/Private/GameplayEffect.cpp
@@ -2185,3 +2185,9 @@ void FInheritedTagContainer::RemoveTag(FGameplayTag TagToRemove)
 {
 	CombinedTags.RemoveTag(TagToRemove);
 }
+
+void FActiveGameplayEffectAction_Remove::PerformAction()
+{
+	check(OwningASC);		//This should not have been destroyed during the lock. If it has been, something unusual may be going on.
+	OwningASC->RemoveActiveGameplayEffect(Handle);
+}
diff --git a/Engine/Source/Runtime/GameplayAbilities/Public/GameplayEffect.h b/Engine/Source/Runtime/GameplayAbilities/Public/GameplayEffect.h
--- a/Engine/Source/Runtime/GameplayAbilities/Public/GameplayEffect.h
+++ b/Engine/Source/Runtime/GameplayAbilities/Public/GameplayEffect.h
@@ -1291,7 +1291,6 @@ struct FActiveGameplayEffectAction
 
 	virtual ~FActiveGameplayEffectAction()
 	{
-
 	}
 
 	virtual void PerformAction()
@@ -1308,18 +1307,21 @@ struct FActiveGameplayEffectAction_Remove : public FActiveGameplayEffectAction
 	}
 
 	FActiveGameplayEffectAction_Remove(FActiveGameplayEffectsContainer& InContainer, FActiveGameplayEffectHandle& InHandle)
-		: Container(InContainer)
-		, Handle(InHandle)
+		: Handle(InHandle)
 	{
+		OwningASC = InContainer.Owner;
 	}
 
-	virtual void PerformAction()
+	FActiveGameplayEffectAction_Remove(UAbilitySystemComponent* InOwningASC, FActiveGameplayEffectHandle& InHandle)
+		: OwningASC(InOwningASC)
+		, Handle(InHandle)
 	{
-		Container.RemoveActiveGameplayEffect(Handle);
 	}
 
+	virtual void PerformAction() override;
+
 	UPROPERTY()
-	FActiveGameplayEffectsContainer Container;
+	UAbilitySystemComponent* OwningASC;
 
 	UPROPERTY()
 	FActiveGameplayEffectHandle Handle;
