Fix issue where first ticking task wouldnt turn on ticking. (only recently exposed when initial ticking of ASC was properly turned off)

diff --git a/Engine/Source/Runtime/GameplayAbilities/Private/AbilitySystemComponent.cpp b/Engine/Source/Runtime/GameplayAbilities/Private/AbilitySystemComponent.cpp
--- a/Engine/Source/Runtime/GameplayAbilities/Private/AbilitySystemComponent.cpp
+++ b/Engine/Source/Runtime/GameplayAbilities/Private/AbilitySystemComponent.cpp
@@ -821,13 +821,15 @@ void UAbilitySystemComponent::TaskStarted(UAbilityTask* NewTask)
 {
 	if (NewTask->bTickingTask)
 	{
+		check(TickingTasks.Contains(NewTask) == false);
+		TickingTasks.Add(NewTask);
+
 		// If this is our first ticking task, set this component as active so it begins ticking
-		if (TickingTasks.Num() == 0)
+		if (TickingTasks.Num() == 1)
 		{
 			UpdateShouldTick();
 		}
-		check(TickingTasks.Contains(NewTask) == false);
-		TickingTasks.Add(NewTask);
+		
 	}
 	if (NewTask->bSimulatedTask)
 	{
