guard MoveToLocation task against the avatar being destroyed

diff --git a/Engine/Source/Runtime/GameplayAbilities/Private/Abilities/Tasks/AbilityTask_MoveToLocation.cpp b/Engine/Source/Runtime/GameplayAbilities/Private/Abilities/Tasks/AbilityTask_MoveToLocation.cpp
--- a/Engine/Source/Runtime/GameplayAbilities/Private/Abilities/Tasks/AbilityTask_MoveToLocation.cpp
+++ b/Engine/Source/Runtime/GameplayAbilities/Private/Abilities/Tasks/AbilityTask_MoveToLocation.cpp
@@ -17,7 +17,11 @@ UAbilityTask_MoveToLocation* UAbilityTask_MoveToLocation::MoveToLocation(class U
 {
 	auto MyObj = NewTask<UAbilityTask_MoveToLocation>(WorldContextObject, TaskInstanceName);
 
-	MyObj->StartLocation = MyObj->GetAvatarActor()->GetActorLocation();
+	if (MyObj->GetAvatarActor() != nullptr)
+	{
+		MyObj->StartLocation = MyObj->GetAvatarActor()->GetActorLocation();
+	}
+
 	MyObj->TargetLocation = Location;
 	MyObj->DurationOfMovement = FMath::Max(Duration, 0.001f);		// Avoid negative or divide-by-zero cases
 	MyObj->TimeMoveStarted = MyObj->GetWorld()->GetTimeSeconds();
@@ -89,6 +93,11 @@ void UAbilityTask_MoveToLocation::TickTask(float DeltaTime)
 			MyActor->SetActorLocation(FMath::Lerp<FVector, float>(StartLocation, TargetLocation, MoveFraction));
 		}
 	}
+	else
+	{
+		bIsFinished = true;
+		EndTask();
+	}
 }
 
 void UAbilityTask_MoveToLocation::GetLifetimeReplicatedProps(TArray< FLifetimeProperty > & OutLifetimeProps) const
