properly end WaitDelay tasks when they finish

diff --git a/Engine/Source/Runtime/GameplayAbilities/Private/Abilities/Tasks/AbilityTask_WaitDelay.cpp b/Engine/Source/Runtime/GameplayAbilities/Private/Abilities/Tasks/AbilityTask_WaitDelay.cpp
--- a/Engine/Source/Runtime/GameplayAbilities/Private/Abilities/Tasks/AbilityTask_WaitDelay.cpp
+++ b/Engine/Source/Runtime/GameplayAbilities/Private/Abilities/Tasks/AbilityTask_WaitDelay.cpp
@@ -31,6 +31,7 @@ void UAbilityTask_WaitDelay::OnDestroy(bool AbilityEnded)
 void UAbilityTask_WaitDelay::OnTimeFinish()
 {
 	OnFinish.Broadcast();
+	EndTask();
 }
 
 FString UAbilityTask_WaitDelay::GetDebugString() const
