Scope-lock protections on GameplayEffect add/remove implemented.

diff --git a/Engine/Source/Runtime/GameplayAbilities/Private/GameplayEffect.cpp b/Engine/Source/Runtime/GameplayAbilities/Private/GameplayEffect.cpp
--- a/Engine/Source/Runtime/GameplayAbilities/Private/GameplayEffect.cpp
+++ b/Engine/Source/Runtime/GameplayAbilities/Private/GameplayEffect.cpp
@@ -1392,8 +1392,6 @@ FActiveGameplayEffect& FActiveGameplayEffectsContainer::CreateNewActiveGameplayE
 {
 	SCOPE_CYCLE_COUNTER(STAT_CreateNewActiveGameplayEffect);
 
-	//TODO Add FScopedActiveGameplayEffectLock here once we can get a crash to happen
-
 	if (Owner && Owner->OwnerActor)
 	{
 		if ( IsNetAuthority() )
@@ -1403,7 +1401,17 @@ FActiveGameplayEffect& FActiveGameplayEffectsContainer::CreateNewActiveGameplayE
 	}
 
 	FActiveGameplayEffectHandle NewHandle = FActiveGameplayEffectHandle::GenerateNewHandle(Owner);
-	FActiveGameplayEffect& NewEffect = *new (GameplayEffects)FActiveGameplayEffect(NewHandle, Spec, GetWorldTime(), GetGameStateTime(), InPredictionKey);
+	FActiveGameplayEffect* NewEffectPtr = new FActiveGameplayEffect(NewHandle, Spec, GetWorldTime(), GetGameStateTime(), InPredictionKey);
+	FActiveGameplayEffect& NewEffect = *NewEffectPtr;
+	if (FScopedActiveGameplayEffectLock::IsLockInEffect())
+	{
+		GameplayEffectsPendingAdd.Add(NewEffect);
+		FScopedActiveGameplayEffectLock::AddAction(new FActiveGameplayEffectAction_Add(TWeakObjectPtr<UAbilitySystemComponent>(Owner)));
+	}
+	else
+	{
+		GameplayEffects.Add(NewEffect);
+	}
 
 	UAbilitySystemGlobals::Get().GlobalPreGameplayEffectSpecApply(NewEffect.Spec, Owner);
 
@@ -1573,7 +1581,7 @@ bool FActiveGameplayEffectsContainer::RemoveActiveGameplayEffect(FActiveGameplay
 			//Block removal if we're scope-locked
 			if (FScopedActiveGameplayEffectLock::IsLockInEffect())
 			{
-				FScopedActiveGameplayEffectLock::AddAction(new FActiveGameplayEffectAction_Remove(*this, Handle));
+				FScopedActiveGameplayEffectLock::AddAction(new FActiveGameplayEffectAction_Remove(TWeakObjectPtr<UAbilitySystemComponent>(Owner), Handle));
 				return true;		//We are assuming the "ensure(Idx < GameplayEffects.Num())" would be passed.
 			}
 			return InternalRemoveActiveGameplayEffect(Idx);
@@ -2188,6 +2196,19 @@ void FInheritedTagContainer::RemoveTag(FGameplayTag TagToRemove)
 
 void FActiveGameplayEffectAction_Remove::PerformAction()
 {
-	check(OwningASC);		//This should not have been destroyed during the lock. If it has been, something unusual may be going on.
-	OwningASC->RemoveActiveGameplayEffect(Handle);
+	if (OwningASC.IsValid())
+	{
+		OwningASC.Get()->RemoveActiveGameplayEffect(Handle);
+	}
+}
+
+void FActiveGameplayEffectAction_Add::PerformAction()
+{
+	if (OwningASC.IsValid())
+	{
+		UAbilitySystemComponent* ASC = OwningASC.Get();
+		check(ASC->ActiveGameplayEffects.GameplayEffectsPendingAdd.Num());
+		ASC->ActiveGameplayEffects.GameplayEffects.Add(ASC->ActiveGameplayEffects.GameplayEffectsPendingAdd[0]);
+		ASC->ActiveGameplayEffects.GameplayEffectsPendingAdd.RemoveAt(0, 1, false);
+	}
 }
diff --git a/Engine/Source/Runtime/GameplayAbilities/Public/AbilitySystemComponent.h b/Engine/Source/Runtime/GameplayAbilities/Public/AbilitySystemComponent.h
--- a/Engine/Source/Runtime/GameplayAbilities/Public/AbilitySystemComponent.h
+++ b/Engine/Source/Runtime/GameplayAbilities/Public/AbilitySystemComponent.h
@@ -53,6 +53,7 @@ class GAMEPLAYABILITIES_API UAbilitySystemComponent : public UActorComponent, pu
 	/** Used to register callbacks to confirm/cancel input */
 	DECLARE_DYNAMIC_MULTICAST_DELEGATE(FAbilityConfirmOrCancel);
 
+	friend struct FActiveGameplayEffectAction_Add;
 	friend FGameplayEffectSpec;
 	friend class AAbilitySystemDebugHUD;
 
diff --git a/Engine/Source/Runtime/GameplayAbilities/Public/GameplayEffect.h b/Engine/Source/Runtime/GameplayAbilities/Public/GameplayEffect.h
--- a/Engine/Source/Runtime/GameplayAbilities/Public/GameplayEffect.h
+++ b/Engine/Source/Runtime/GameplayAbilities/Public/GameplayEffect.h
@@ -1120,7 +1120,10 @@ struct FActiveGameplayEffectsContainer : public FFastArraySerializer
 	UAbilitySystemComponent* Owner;
 
 	UPROPERTY()
-	TArray< FActiveGameplayEffect >	GameplayEffects;
+	TArray<FActiveGameplayEffect>	GameplayEffects;
+
+	UPROPERTY()
+	TArray<FActiveGameplayEffect>	GameplayEffectsPendingAdd;
 
 	void RegisterWithOwner(UAbilitySystemComponent* Owner);	
 	
@@ -1282,7 +1285,7 @@ struct TStructOpsTypeTraits< FActiveGameplayEffectsContainer > : public TStructO
 
 
 USTRUCT()
-struct FActiveGameplayEffectAction
+struct GAMEPLAYABILITIES_API FActiveGameplayEffectAction
 {
 	GENERATED_USTRUCT_BODY()
 	FActiveGameplayEffectAction()
@@ -1298,37 +1301,51 @@ struct FActiveGameplayEffectAction
 	}
 };
 
+
 USTRUCT()
-struct FActiveGameplayEffectAction_Remove : public FActiveGameplayEffectAction
+struct GAMEPLAYABILITIES_API FActiveGameplayEffectAction_Add : public FActiveGameplayEffectAction
 {
 	GENERATED_USTRUCT_BODY()
-	FActiveGameplayEffectAction_Remove()
+	FActiveGameplayEffectAction_Add()
 	{
 	}
 
-	FActiveGameplayEffectAction_Remove(FActiveGameplayEffectsContainer& InContainer, FActiveGameplayEffectHandle& InHandle)
-		: Handle(InHandle)
+	FActiveGameplayEffectAction_Add(TWeakObjectPtr<UAbilitySystemComponent> InOwningASC)
+	: OwningASC(InOwningASC)
+	{
+	}
+
+	virtual void PerformAction() override;
+
+	UPROPERTY()
+	TWeakObjectPtr<UAbilitySystemComponent> OwningASC;
+};
+
+USTRUCT()
+struct GAMEPLAYABILITIES_API FActiveGameplayEffectAction_Remove : public FActiveGameplayEffectAction
+{
+	GENERATED_USTRUCT_BODY()
+	FActiveGameplayEffectAction_Remove()
 	{
-		OwningASC = InContainer.Owner;
 	}
 
-	FActiveGameplayEffectAction_Remove(UAbilitySystemComponent* InOwningASC, FActiveGameplayEffectHandle& InHandle)
-		: OwningASC(InOwningASC)
-		, Handle(InHandle)
+	FActiveGameplayEffectAction_Remove(TWeakObjectPtr<UAbilitySystemComponent> InOwningASC, FActiveGameplayEffectHandle& InHandle)
+	: OwningASC(InOwningASC)
+	, Handle(InHandle)
 	{
 	}
 
 	virtual void PerformAction() override;
 
 	UPROPERTY()
-	UAbilitySystemComponent* OwningASC;
+	TWeakObjectPtr<UAbilitySystemComponent> OwningASC;
 
 	UPROPERTY()
 	FActiveGameplayEffectHandle Handle;
 };
 
 USTRUCT()
-struct FActiveGameplayEffectActionHandle
+struct GAMEPLAYABILITIES_API FActiveGameplayEffectActionHandle
 {
 	GENERATED_USTRUCT_BODY()
 
@@ -1365,7 +1382,7 @@ struct FActiveGameplayEffectActionHandle
 	TArray<TSharedPtr<FActiveGameplayEffectAction>>	Data;
 };
 
-struct FScopedActiveGameplayEffectLock
+struct GAMEPLAYABILITIES_API FScopedActiveGameplayEffectLock
 {
 	FScopedActiveGameplayEffectLock();
 	~FScopedActiveGameplayEffectLock();
