Removed some wasteful data from Gameplay Effect Context serialization.

diff --git a/Engine/Source/Runtime/GameplayAbilities/Private/GameplayEffectTypes.cpp b/Engine/Source/Runtime/GameplayAbilities/Private/GameplayEffectTypes.cpp
--- a/Engine/Source/Runtime/GameplayAbilities/Private/GameplayEffectTypes.cpp
+++ b/Engine/Source/Runtime/GameplayAbilities/Private/GameplayEffectTypes.cpp
@@ -72,13 +72,7 @@ bool FGameplayEffectContext::NetSerialize(FArchive& Ar, class UPackageMap* Map,
 {
 	Ar << Instigator;
 	Ar << EffectCauser;
-
-	int32 ActorArraySize = Actors.Num();		//This can probably be smaller than int32
-	Ar << ActorArraySize;
-	if (ActorArraySize > 0)
-	{
-		Ar << Actors;
-	}
+	Ar << Actors;
 
 	bool HasHitResults = HitResult.IsValid();
 	Ar << HasHitResults;
