Changed from array to queue implementation in GE add/remove lock code.

diff --git a/Engine/Source/Runtime/GameplayAbilities/Private/AbilitySystemComponent.cpp b/Engine/Source/Runtime/GameplayAbilities/Private/AbilitySystemComponent.cpp
--- a/Engine/Source/Runtime/GameplayAbilities/Private/AbilitySystemComponent.cpp
+++ b/Engine/Source/Runtime/GameplayAbilities/Private/AbilitySystemComponent.cpp
@@ -443,8 +443,6 @@ FActiveGameplayEffectHandle UAbilitySystemComponent::ApplyGameplayEffectSpecToTa
 
 FActiveGameplayEffectHandle UAbilitySystemComponent::ApplyGameplayEffectSpecToSelf(OUT FGameplayEffectSpec &Spec, FPredictionKey PredictionKey)
 {
-	//FScopedActiveGameplayEffectLock AGELock;
-
 	// Check Network Authority
 	if (!HasNetworkAuthorityToApplyGameplayEffect(PredictionKey))
 	{
diff --git a/Engine/Source/Runtime/GameplayAbilities/Private/GameplayEffect.cpp b/Engine/Source/Runtime/GameplayAbilities/Private/GameplayEffect.cpp
--- a/Engine/Source/Runtime/GameplayAbilities/Private/GameplayEffect.cpp
+++ b/Engine/Source/Runtime/GameplayAbilities/Private/GameplayEffect.cpp
@@ -1405,7 +1405,7 @@ FActiveGameplayEffect& FActiveGameplayEffectsContainer::CreateNewActiveGameplayE
 	FActiveGameplayEffect& NewEffect = *NewEffectPtr;
 	if (FScopedActiveGameplayEffectLock::IsLockInEffect())
 	{
-		GameplayEffectsPendingAdd.Add(NewEffect);
+		GameplayEffectsPendingAdd.Enqueue(NewEffect);
 		FScopedActiveGameplayEffectLock::AddAction(new FActiveGameplayEffectAction_Add(TWeakObjectPtr<UAbilitySystemComponent>(Owner)));
 	}
 	else
@@ -2207,8 +2207,9 @@ void FActiveGameplayEffectAction_Add::PerformAction()
 	if (OwningASC.IsValid())
 	{
 		UAbilitySystemComponent* ASC = OwningASC.Get();
-		check(ASC->ActiveGameplayEffects.GameplayEffectsPendingAdd.Num());
-		ASC->ActiveGameplayEffects.GameplayEffects.Add(ASC->ActiveGameplayEffects.GameplayEffectsPendingAdd[0]);
-		ASC->ActiveGameplayEffects.GameplayEffectsPendingAdd.RemoveAt(0, 1, false);
+		check(!ASC->ActiveGameplayEffects.GameplayEffectsPendingAdd.IsEmpty());
+		FActiveGameplayEffect TempEffect;
+		ASC->ActiveGameplayEffects.GameplayEffectsPendingAdd.Dequeue(TempEffect);
+		ASC->ActiveGameplayEffects.GameplayEffects.Add(TempEffect);
 	}
 }
diff --git a/Engine/Source/Runtime/GameplayAbilities/Public/GameplayEffect.h b/Engine/Source/Runtime/GameplayAbilities/Public/GameplayEffect.h
--- a/Engine/Source/Runtime/GameplayAbilities/Public/GameplayEffect.h
+++ b/Engine/Source/Runtime/GameplayAbilities/Public/GameplayEffect.h
@@ -1122,8 +1122,7 @@ struct FActiveGameplayEffectsContainer : public FFastArraySerializer
 	UPROPERTY()
 	TArray<FActiveGameplayEffect>	GameplayEffects;
 
-	UPROPERTY()
-	TArray<FActiveGameplayEffect>	GameplayEffectsPendingAdd;
+	TQueue<FActiveGameplayEffect>	GameplayEffectsPendingAdd;
 
 	void RegisterWithOwner(UAbilitySystemComponent* Owner);	
 	
