guard against calling ApplyGameplayEffectToTarget on null

diff --git a/Engine/Source/Runtime/GameplayAbilities/Private/AbilitySystemComponent.cpp b/Engine/Source/Runtime/GameplayAbilities/Private/AbilitySystemComponent.cpp
--- a/Engine/Source/Runtime/GameplayAbilities/Private/AbilitySystemComponent.cpp
+++ b/Engine/Source/Runtime/GameplayAbilities/Private/AbilitySystemComponent.cpp
@@ -215,15 +215,22 @@ FGameplayEffectContextHandle UAbilitySystemComponent::GetEffectContext() const
 	return Context;
 }
 
-FActiveGameplayEffectHandle UAbilitySystemComponent::BP_ApplyGameplayEffectToTarget(TSubclassOf<UGameplayEffect> GameplayEffectClass, UAbilitySystemComponent *Target, float Level, FGameplayEffectContextHandle Context)
+FActiveGameplayEffectHandle UAbilitySystemComponent::BP_ApplyGameplayEffectToTarget(TSubclassOf<UGameplayEffect> GameplayEffectClass, UAbilitySystemComponent* Target, float Level, FGameplayEffectContextHandle Context)
 {
-	if (GameplayEffectClass)
+	if (Target == nullptr)
 	{
-		UGameplayEffect* GameplayEffect = GameplayEffectClass->GetDefaultObject<UGameplayEffect>();
-		return ApplyGameplayEffectToTarget(GameplayEffect, Target, Level, Context);
+		ABILITY_LOG(Error, TEXT("UAbilitySystemComponent::BP_ApplyGameplayEffectToTarget called with null Target. Context: %s"), *Context.ToString());
+		return FActiveGameplayEffectHandle();
 	}
 
-	return FActiveGameplayEffectHandle();
+	if (GameplayEffectClass == nullptr)
+	{
+		ABILITY_LOG(Error, TEXT("UAbilitySystemComponent::BP_ApplyGameplayEffectToTarget called with null GameplayEffectClass. Context: %s"), *Context.ToString());
+		return FActiveGameplayEffectHandle();
+	}
+
+	UGameplayEffect* GameplayEffect = GameplayEffectClass->GetDefaultObject<UGameplayEffect>();
+	return ApplyGameplayEffectToTarget(GameplayEffect, Target, Level, Context);	
 }
 
 FActiveGameplayEffectHandle UAbilitySystemComponent::K2_ApplyGameplayEffectToTarget(UGameplayEffect *GameplayEffect, UAbilitySystemComponent *Target, float Level, FGameplayEffectContextHandle Context)
