HitScanTube fixes

* HitScanTube now tracks hit actors to make sure it doesn't add the same actor twice when overlapping multiple components.

* Fixed HitScanTube using the wrong overlap hit result.

* Fix to make sure FGameplayEffectContextHandle only gets one hit result when adding multiple target data.

diff --git a/Engine/Source/Runtime/GameplayAbilities/Private/GameplayAbilityTargetTypes.cpp b/Engine/Source/Runtime/GameplayAbilities/Private/GameplayAbilityTargetTypes.cpp
--- a/Engine/Source/Runtime/GameplayAbilities/Private/GameplayAbilityTargetTypes.cpp
+++ b/Engine/Source/Runtime/GameplayAbilities/Private/GameplayAbilityTargetTypes.cpp
@@ -73,7 +73,7 @@ void FGameplayAbilityTargetData::AddTargetDataToContext(FGameplayEffectContextHa
 		Context.AddActors(GetActors());
 	}
 
-	if (HasHitResult())
+	if (HasHitResult() && !Context.GetHitResult())
 	{
 		Context.AddHitResult(*GetHitResult());
 	}
