Fix for health/energy not being clamped

diff --git a/Engine/Source/Runtime/GameplayAbilities/Private/AttributeSet.cpp b/Engine/Source/Runtime/GameplayAbilities/Private/AttributeSet.cpp
--- a/Engine/Source/Runtime/GameplayAbilities/Private/AttributeSet.cpp
+++ b/Engine/Source/Runtime/GameplayAbilities/Private/AttributeSet.cpp
@@ -14,7 +14,7 @@ FGameplayAttribute::FGameplayAttribute(UProperty *NewProperty)
 	Attribute = Cast<UNumericProperty>(NewProperty);
 }
 
-void FGameplayAttribute::SetNumericValueChecked(const float NewValue, class UAttributeSet* Dest) const
+void FGameplayAttribute::SetNumericValueChecked(float NewValue, class UAttributeSet* Dest) const
 {
 	UNumericProperty *NumericProperty = CastChecked<UNumericProperty>(Attribute);
 	void* ValuePtr = NumericProperty->ContainerPtrToValuePtr<void>(Dest);
diff --git a/Engine/Source/Runtime/GameplayAbilities/Public/AttributeSet.h b/Engine/Source/Runtime/GameplayAbilities/Public/AttributeSet.h
--- a/Engine/Source/Runtime/GameplayAbilities/Public/AttributeSet.h
+++ b/Engine/Source/Runtime/GameplayAbilities/Public/AttributeSet.h
@@ -42,7 +42,7 @@ struct GAMEPLAYABILITIES_API FGameplayAttribute
 
 	bool IsSystemAttribute() const;
 
-	void SetNumericValueChecked(const float NewValue, class UAttributeSet* Dest) const;
+	void SetNumericValueChecked(float NewValue, class UAttributeSet* Dest) const;
 
 	float GetNumericValueChecked(const UAttributeSet* Src) const;
 	
@@ -106,8 +106,10 @@ public:
 	 *	Called just before any modification happens to an attribute. This is lower level than PreAttributeModify/PostAttribute modify.
 	 *	There is no additional context provided here since anything can trigger this. Executed effects, duration based effects, effects being removed, immunity being applied, stacking rules changing, etc.
 	 *	This function is meant to enforce things like "Health = Clamp(Health, 0, MaxHealth)" and NOT things like "trigger this extra thing if damage is applied, etc".
+	 *	
+	 *	NewValue is a mutable reference so you are able to clamp the newly applied value as well.
 	 */
-	virtual void PreAttributeChange(const FGameplayAttribute& Attribute, float NewValue) { }
+	virtual void PreAttributeChange(const FGameplayAttribute& Attribute, float& NewValue) { }
 
 	/** 
 	 * Called to determine the set of gameplay attributes that the specified attribute requires as pre-requisites in order to accurately compute
