readd/implement linked TargetGameplayEffects

diff --git a/Engine/Source/Runtime/GameplayAbilities/Private/GameplayEffect.cpp b/Engine/Source/Runtime/GameplayAbilities/Private/GameplayEffect.cpp
--- a/Engine/Source/Runtime/GameplayAbilities/Private/GameplayEffect.cpp
+++ b/Engine/Source/Runtime/GameplayAbilities/Private/GameplayEffect.cpp
@@ -48,7 +48,7 @@ void UGameplayEffect::GetOwnedGameplayTags(FGameplayTagContainer& TagContainer)
 // 	TagContainer.AppendTags(InheritableClearTagsContainer.CombinedTags);
 // }
 
-void UGameplayEffect::GetTargetEffects(TArray<UGameplayEffect*>& OutEffects) const
+void UGameplayEffect::GetTargetEffects(TArray<const UGameplayEffect*>& OutEffects) const
 {
 	OutEffects.Append(TargetEffects);
 
@@ -357,7 +357,10 @@ FGameplayEffectSpec::FGameplayEffectSpec(const UGameplayEffect* InDef, const FGa
 	CapturedSourceTags.GetSpecTags().AppendTags(InDef->InheritableGameplayEffectTags.CombinedTags);
 
 	// Make TargetEffectSpecs too
-	for (UGameplayEffect* TargetDef : InDef->TargetEffects)
+	TArray<const UGameplayEffect*> TargetEffectDefs;
+	InDef->GetTargetEffects(TargetEffectDefs);
+
+	for (const UGameplayEffect* TargetDef : TargetEffectDefs)
 	{
 		TargetEffectSpecs.Add(TSharedRef<FGameplayEffectSpec>(new FGameplayEffectSpec(TargetDef, EffectContext, Level)));
 	}
diff --git a/Engine/Source/Runtime/GameplayAbilities/Public/GameplayEffect.h b/Engine/Source/Runtime/GameplayAbilities/Public/GameplayEffect.h
--- a/Engine/Source/Runtime/GameplayAbilities/Public/GameplayEffect.h
+++ b/Engine/Source/Runtime/GameplayAbilities/Public/GameplayEffect.h
@@ -487,7 +487,7 @@ public:
 	UPROPERTY(VisibleDefaultsOnly, Category = Deprecated)
 	TArray<UGameplayEffect*> TargetEffects;
 
-	void GetTargetEffects(TArray<UGameplayEffect*>& OutEffects) const;
+	void GetTargetEffects(TArray<const UGameplayEffect*>& OutEffects) const;
 
 	// ------------------------------------------------
 	// Gameplay tag interface
