reload attribute data when backing data asset changes in editor

diff --git a/Engine/Source/Runtime/GameplayAbilities/Private/AbilitySystemGlobals.cpp b/Engine/Source/Runtime/GameplayAbilities/Private/AbilitySystemGlobals.cpp
--- a/Engine/Source/Runtime/GameplayAbilities/Private/AbilitySystemGlobals.cpp
+++ b/Engine/Source/Runtime/GameplayAbilities/Private/AbilitySystemGlobals.cpp
@@ -68,17 +68,10 @@ T* UAbilitySystemGlobals::InternalGetLoadTable(T*& Table, FString TableName)
 
 void UAbilitySystemGlobals::OnTableReimported(UObject* InObject)
 {
-	if (GIsEditor && !IsRunningCommandlet() && InObject && InObject == GlobalCurveTable)
-	{
-	}
-	
 	if (GIsEditor && !IsRunningCommandlet() && InObject && InObject == GlobalAttributeDefaultsTable)
 	{
-	}
-
-	if (GIsEditor && !IsRunningCommandlet() && InObject && InObject == GlobalAttributeMetaDataTable)
-	{
-	}
+		ReloadAttributeDefaults();
+	}	
 }
 
 #endif
@@ -188,11 +181,16 @@ void UAbilitySystemGlobals::InitAttributeDefaults()
 {
 	if (InternalGetLoadTable<UCurveTable>(GlobalAttributeDefaultsTable, GlobalAttributeSetDefaultsTableName))
 	{	
-		AllocAttributeSetInitter();
-		GlobalAttributeSetInitter->PreloadAttributeSetData(GlobalAttributeDefaultsTable);
+		ReloadAttributeDefaults();
 	}
 }
 
+void UAbilitySystemGlobals::ReloadAttributeDefaults()
+{
+	AllocAttributeSetInitter();
+	GlobalAttributeSetInitter->PreloadAttributeSetData(GlobalAttributeDefaultsTable);
+}
+
 // --------------------------------------------------------------------
 
 UGameplayCueManager* UAbilitySystemGlobals::GetGameplayCueManager()
diff --git a/Engine/Source/Runtime/GameplayAbilities/Public/AbilitySystemGlobals.h b/Engine/Source/Runtime/GameplayAbilities/Public/AbilitySystemGlobals.h
--- a/Engine/Source/Runtime/GameplayAbilities/Public/AbilitySystemGlobals.h
+++ b/Engine/Source/Runtime/GameplayAbilities/Public/AbilitySystemGlobals.h
@@ -92,6 +92,7 @@ class GAMEPLAYABILITIES_API UAbilitySystemGlobals : public UObject
 protected:
 
 	virtual void InitAttributeDefaults();
+	virtual void ReloadAttributeDefaults();
 	virtual void AllocAttributeSetInitter();
 
 private:
