remove active gameplayeffect handles from the global handle map when they expire

diff --git a/Engine/Source/Runtime/GameplayAbilities/Private/GameplayEffect.cpp b/Engine/Source/Runtime/GameplayAbilities/Private/GameplayEffect.cpp
--- a/Engine/Source/Runtime/GameplayAbilities/Private/GameplayEffect.cpp
+++ b/Engine/Source/Runtime/GameplayAbilities/Private/GameplayEffect.cpp
@@ -1638,6 +1638,9 @@ bool FActiveGameplayEffectsContainer::InternalRemoveActiveGameplayEffect(int32 I
 			Owner->OwnerActor->FlushNetDormancy();
 		}
 
+		// Remove this handle from the global map
+		Effect.Handle.RemoveFromGlobalMap();
+
 		// Finally remove the ActiveGameplayEffect
 		GameplayEffects.RemoveAtSwap(Idx);
 		MarkArrayDirty();
@@ -2113,6 +2116,11 @@ const UAbilitySystemComponent* FActiveGameplayEffectHandle::GetOwningAbilitySyst
 	return nullptr;
 }
 
+void FActiveGameplayEffectHandle::RemoveFromGlobalMap()
+{
+	GlobalActiveGameplayEffectHandles::Map.Remove(*this);
+}
+
 // -----------------------------------------------------------------
 
 bool FActiveGameplayEffectQuery::Matches(const FActiveGameplayEffect& Effect) const
diff --git a/Engine/Source/Runtime/GameplayAbilities/Public/GameplayEffectTypes.h b/Engine/Source/Runtime/GameplayAbilities/Public/GameplayEffectTypes.h
--- a/Engine/Source/Runtime/GameplayAbilities/Public/GameplayEffectTypes.h
+++ b/Engine/Source/Runtime/GameplayAbilities/Public/GameplayEffectTypes.h
@@ -95,7 +95,7 @@ enum class EGameplayEffectAttributeCaptureSource : uint8
  *	through a handle. a pointer or index into the active list is not sufficient.
  */
 USTRUCT(BlueprintType)
-struct FActiveGameplayEffectHandle
+struct GAMEPLAYABILITIES_API FActiveGameplayEffectHandle
 {
 	GENERATED_USTRUCT_BODY()
 
@@ -121,6 +121,8 @@ struct FActiveGameplayEffectHandle
 	UAbilitySystemComponent* GetOwningAbilitySystemComponent();
 	const UAbilitySystemComponent* GetOwningAbilitySystemComponent() const;
 
+	void RemoveFromGlobalMap();
+
 	bool operator==(const FActiveGameplayEffectHandle& Other) const
 	{
 		return Handle == Other.Handle;
