Updating scope-locked Active GE code to remove the effect of scope-locking on additions of GEs. Also removing actor arrays from GE transmission, but leaving in the case of GCs.

diff --git a/Engine/Source/Runtime/GameplayAbilities/Private/Abilities/GameplayAbility.cpp b/Engine/Source/Runtime/GameplayAbilities/Private/Abilities/GameplayAbility.cpp
--- a/Engine/Source/Runtime/GameplayAbilities/Private/Abilities/GameplayAbility.cpp
+++ b/Engine/Source/Runtime/GameplayAbilities/Private/Abilities/GameplayAbility.cpp
@@ -750,7 +750,7 @@ FGameplayEffectContextHandle UGameplayAbility::GetContextFromOwner(FGameplayAbil
 	{
 		if (Data.IsValid())
 		{
-			Data->AddTargetDataToContext(Context);
+			Data->AddTargetDataToContext(Context, true);
 		}
 	}
 
diff --git a/Engine/Source/Runtime/GameplayAbilities/Private/AbilitySystemBlueprintLibrary.cpp b/Engine/Source/Runtime/GameplayAbilities/Private/AbilitySystemBlueprintLibrary.cpp
--- a/Engine/Source/Runtime/GameplayAbilities/Private/AbilitySystemBlueprintLibrary.cpp
+++ b/Engine/Source/Runtime/GameplayAbilities/Private/AbilitySystemBlueprintLibrary.cpp
@@ -271,10 +271,18 @@ bool UAbilitySystemBlueprintLibrary::IsInstigatorLocallyControlled(FGameplayCueP
 	return Parameters.EffectContext.IsLocallyControlled();
 }
 
+int32 UAbilitySystemBlueprintLibrary::GetActorCount(FGameplayCueParameters Parameters)
+{
+	return Parameters.EffectContext.GetActors().Num();
+}
 
-TArray<TWeakObjectPtr<AActor>> UAbilitySystemBlueprintLibrary::GetActors(FGameplayCueParameters Parameters)
+AActor* UAbilitySystemBlueprintLibrary::GetActorByIndex(FGameplayCueParameters Parameters, int32 Index)
 {
-	return Parameters.EffectContext.GetActors();
+	if (Parameters.EffectContext.GetActors().IsValidIndex(Index))
+	{
+		return Parameters.EffectContext.GetActors()[Index].Get();
+	}
+	return NULL;
 }
 
 FHitResult UAbilitySystemBlueprintLibrary::GetHitResult(FGameplayCueParameters Parameters)
diff --git a/Engine/Source/Runtime/GameplayAbilities/Private/GameplayAbilityTargetTypes.cpp b/Engine/Source/Runtime/GameplayAbilities/Private/GameplayAbilityTargetTypes.cpp
--- a/Engine/Source/Runtime/GameplayAbilities/Private/GameplayAbilityTargetTypes.cpp
+++ b/Engine/Source/Runtime/GameplayAbilities/Private/GameplayAbilityTargetTypes.cpp
@@ -14,7 +14,7 @@ TArray<FActiveGameplayEffectHandle> FGameplayAbilityTargetData::ApplyGameplayEff
 	
 	FGameplayEffectSpec	SpecToApply(GameplayEffect, EffectContext, Level);
 	
-	AddTargetDataToContext(EffectContext);
+	AddTargetDataToContext(EffectContext, false);
 
 	TArray<TWeakObjectPtr<AActor> > Actors = GetActors();
 	TArray<FActiveGameplayEffectHandle>	AppliedHandles;
@@ -45,7 +45,7 @@ TArray<FActiveGameplayEffectHandle> FGameplayAbilityTargetData::ApplyGameplayEff
 	}
 
 	FGameplayEffectContextHandle EffectContext = Spec.GetContext();
-	AddTargetDataToContext(EffectContext);
+	AddTargetDataToContext(EffectContext, false);
 
 	TArray<TWeakObjectPtr<AActor> > Actors = GetActors();
 	
@@ -66,9 +66,9 @@ TArray<FActiveGameplayEffectHandle> FGameplayAbilityTargetData::ApplyGameplayEff
 	return AppliedHandles;
 }
 
-void FGameplayAbilityTargetData::AddTargetDataToContext(FGameplayEffectContextHandle& Context)
+void FGameplayAbilityTargetData::AddTargetDataToContext(FGameplayEffectContextHandle& Context, bool bIncludeActorArray)
 {
-	if (GetActors().Num() > 0)
+	if (bIncludeActorArray && (GetActors().Num() > 0))
 	{
 		Context.AddActors(GetActors());
 	}
diff --git a/Engine/Source/Runtime/GameplayAbilities/Private/GameplayEffect.cpp b/Engine/Source/Runtime/GameplayAbilities/Private/GameplayEffect.cpp
--- a/Engine/Source/Runtime/GameplayAbilities/Private/GameplayEffect.cpp
+++ b/Engine/Source/Runtime/GameplayAbilities/Private/GameplayEffect.cpp
@@ -1404,17 +1404,7 @@ FActiveGameplayEffect& FActiveGameplayEffectsContainer::CreateNewActiveGameplayE
 	}
 
 	FActiveGameplayEffectHandle NewHandle = FActiveGameplayEffectHandle::GenerateNewHandle(Owner);
-	FActiveGameplayEffect* NewEffectPtr = new FActiveGameplayEffect(NewHandle, Spec, GetWorldTime(), GetGameStateTime(), InPredictionKey);
-	FActiveGameplayEffect& NewEffect = *NewEffectPtr;
-	if (FScopedActiveGameplayEffectLock::IsLockInEffect())
-	{
-		GameplayEffectsPendingAdd.Enqueue(NewEffect);
-		FScopedActiveGameplayEffectLock::AddAction(new FActiveGameplayEffectAction_Add(TWeakObjectPtr<UAbilitySystemComponent>(Owner)));
-	}
-	else
-	{
-		GameplayEffects.Add(NewEffect);
-	}
+	FActiveGameplayEffect& NewEffect = *new (GameplayEffects)FActiveGameplayEffect(NewHandle, Spec, GetWorldTime(), GetGameStateTime(), InPredictionKey);
 
 	UAbilitySystemGlobals::Get().GlobalPreGameplayEffectSpecApply(NewEffect.Spec, Owner);
 
diff --git a/Engine/Source/Runtime/GameplayAbilities/Public/Abilities/GameplayAbilityTargetTypes.h b/Engine/Source/Runtime/GameplayAbilities/Public/Abilities/GameplayAbilityTargetTypes.h
--- a/Engine/Source/Runtime/GameplayAbilities/Public/Abilities/GameplayAbilityTargetTypes.h
+++ b/Engine/Source/Runtime/GameplayAbilities/Public/Abilities/GameplayAbilityTargetTypes.h
@@ -67,7 +67,7 @@ struct GAMEPLAYABILITIES_API FGameplayAbilityTargetData
 
 	TArray<FActiveGameplayEffectHandle> ApplyGameplayEffectSpec(FGameplayEffectSpec& Spec, FPredictionKey PredictionKey = FPredictionKey());
 
-	virtual void AddTargetDataToContext(FGameplayEffectContextHandle& Context);
+	virtual void AddTargetDataToContext(FGameplayEffectContextHandle& Context, bool bIncludeActorArray);
 
 	virtual TArray<TWeakObjectPtr<AActor> >	GetActors() const
 	{
diff --git a/Engine/Source/Runtime/GameplayAbilities/Public/AbilitySystemBlueprintLibrary.h b/Engine/Source/Runtime/GameplayAbilities/Public/AbilitySystemBlueprintLibrary.h
--- a/Engine/Source/Runtime/GameplayAbilities/Public/AbilitySystemBlueprintLibrary.h
+++ b/Engine/Source/Runtime/GameplayAbilities/Public/AbilitySystemBlueprintLibrary.h
@@ -86,7 +86,10 @@ class GAMEPLAYABILITIES_API UAbilitySystemBlueprintLibrary : public UBlueprintFu
 	static bool IsInstigatorLocallyControlled(FGameplayCueParameters Parameters);
 
 	UFUNCTION(BlueprintPure, Category = "Ability|GameplayCue")
-	static TArray<TWeakObjectPtr<AActor>> GetActors(FGameplayCueParameters Parameters);
+	static int32 GetActorCount(FGameplayCueParameters Parameters);
+
+	UFUNCTION(BlueprintPure, Category = "Ability|GameplayCue")
+	static AActor* GetActorByIndex(FGameplayCueParameters Parameters, int32 Index);
 
 	UFUNCTION(BlueprintPure, Category = "Ability|GameplayCue")
 	static FHitResult GetHitResult(FGameplayCueParameters Parameters);
diff --git a/Engine/Source/Runtime/GameplayAbilities/Public/GameplayCueNotify_Actor.h b/Engine/Source/Runtime/GameplayAbilities/Public/GameplayCueNotify_Actor.h
--- a/Engine/Source/Runtime/GameplayAbilities/Public/GameplayCueNotify_Actor.h
+++ b/Engine/Source/Runtime/GameplayAbilities/Public/GameplayCueNotify_Actor.h
@@ -44,7 +44,7 @@ class GAMEPLAYABILITIES_API AGameplayCueNotify_Actor : public AActor
 	void K2_HandleGameplayCue(TWeakObjectPtr<AActor> MyTarget, EGameplayCueEvent::Type EventType, FGameplayCueParameters Parameters);
 
 	UFUNCTION(BlueprintImplementableEvent)
-	bool OnExecute(TWeakObjectPtr<AActor> MyTarget, FGameplayCueParameters Parameters) const;
+	bool OnExecute(TWeakObjectPtr<AActor> MyTarget, FGameplayCueParameters Parameters);
 
 	UFUNCTION(BlueprintImplementableEvent)
 	bool OnActive(TWeakObjectPtr<AActor> MyTarget, FGameplayCueParameters Parameters);
@@ -52,7 +52,7 @@ class GAMEPLAYABILITIES_API AGameplayCueNotify_Actor : public AActor
 	UFUNCTION(BlueprintImplementableEvent)
 	bool OnRemove(TWeakObjectPtr<AActor> MyTarget, FGameplayCueParameters Parameters);
 
-	UPROPERTY(EditDefaultsOnly, Category = GameplayCue, meta = (Categories = "GameplayCue"))
+	UPROPERTY(EditDefaultsOnly, Category = GameplayCue)
 	FGameplayTag	GameplayCueTag;
 
 	/** Mirrors GameplayCueTag in order to be asset registry searchable */
diff --git a/Engine/Source/Runtime/GameplayAbilities/Public/GameplayCueNotify_Static.h b/Engine/Source/Runtime/GameplayAbilities/Public/GameplayCueNotify_Static.h
--- a/Engine/Source/Runtime/GameplayAbilities/Public/GameplayCueNotify_Static.h
+++ b/Engine/Source/Runtime/GameplayAbilities/Public/GameplayCueNotify_Static.h
@@ -52,7 +52,7 @@ class GAMEPLAYABILITIES_API UGameplayCueNotify_Static : public UObject
 	UFUNCTION(BlueprintImplementableEvent, BlueprintPure, Category = "GameplayCueNotify")
 	/*static*/ bool OnRemove(TWeakObjectPtr<AActor> MyTarget, FGameplayCueParameters Parameters);
 
-	UPROPERTY(EditDefaultsOnly, Category = GameplayCue, meta = (Categories = "GameplayCue"))
+	UPROPERTY(EditDefaultsOnly, Category = GameplayCue)
 	FGameplayTag	GameplayCueTag;
 
 	/** Mirrors GameplayCueTag in order to be asset registry searchable */
diff --git a/Engine/Source/Runtime/GameplayAbilities/Public/GameplayEffectTypes.h b/Engine/Source/Runtime/GameplayAbilities/Public/GameplayEffectTypes.h
--- a/Engine/Source/Runtime/GameplayAbilities/Public/GameplayEffectTypes.h
+++ b/Engine/Source/Runtime/GameplayAbilities/Public/GameplayEffectTypes.h
@@ -539,7 +539,7 @@ struct FGameplayEffectContextHandle
 		}
 	}
 
-	const TArray<TWeakObjectPtr<AActor>> GetActors() const
+	TArray<TWeakObjectPtr<AActor>> GetActors()
 	{
 		return Data->GetActors();
 	}
