Face placement actor toward owner by default. Added feature for 2D facing.

diff --git a/Engine/Source/Runtime/GameplayAbilities/Private/Abilities/GameplayAbilityTargetActor_ActorPlacement.cpp b/Engine/Source/Runtime/GameplayAbilities/Private/Abilities/GameplayAbilityTargetActor_ActorPlacement.cpp
--- a/Engine/Source/Runtime/GameplayAbilities/Private/Abilities/GameplayAbilityTargetActor_ActorPlacement.cpp
+++ b/Engine/Source/Runtime/GameplayAbilities/Private/Abilities/GameplayAbilityTargetActor_ActorPlacement.cpp
@@ -32,7 +32,7 @@ void AGameplayAbilityTargetActor_ActorPlacement::StartTargeting(UGameplayAbility
 	if (AActor *VisualizationActor = GetWorld()->SpawnActor(PlacedActorClass))
 	{
 		ActorVisualizationReticle = GetWorld()->SpawnActor<AGameplayAbilityWorldReticle_ActorVisualization>();
-		ActorVisualizationReticle->InitializeReticleVisualizationInformation(VisualizationActor, PlacedActorMaterial);
+		ActorVisualizationReticle->InitializeReticleVisualizationInformation(this, VisualizationActor, PlacedActorMaterial);
 		GetWorld()->DestroyActor(VisualizationActor);
 	}
 	if (AGameplayAbilityWorldReticle* CachedReticleActor = ReticleActor.Get())
diff --git a/Engine/Source/Runtime/GameplayAbilities/Private/Abilities/GameplayAbilityWorldReticle.cpp b/Engine/Source/Runtime/GameplayAbilities/Private/Abilities/GameplayAbilityWorldReticle.cpp
--- a/Engine/Source/Runtime/GameplayAbilities/Private/Abilities/GameplayAbilityWorldReticle.cpp
+++ b/Engine/Source/Runtime/GameplayAbilities/Private/Abilities/GameplayAbilityWorldReticle.cpp
@@ -63,8 +63,30 @@ void AGameplayAbilityWorldReticle::SetIsTargetAnActor(bool bNewValue)
 	}
 }
 
-void AGameplayAbilityWorldReticle::FaceTowardSource()
+void AGameplayAbilityWorldReticle::FaceTowardSource(bool bFaceIn2D)
 {
-	check(TargetingActor);
-	SetActorRotation((TargetingActor->StartLocation.GetTargetingTransform().GetLocation() - GetActorLocation()).Rotation());
+	if (TargetingActor)
+	{
+		if (bFaceIn2D)
+		{
+			FVector FacingVector = (TargetingActor->StartLocation.GetTargetingTransform().GetLocation() - GetActorLocation()).SafeNormal2D();
+			if (FacingVector.IsZero())
+			{
+				FacingVector = -GetActorForwardVector().SafeNormal2D();
+			}
+			if (!FacingVector.IsZero())
+			{
+				SetActorRotation(FacingVector.Rotation());
+			}
+		}
+		else
+		{
+			FVector FacingVector = (TargetingActor->StartLocation.GetTargetingTransform().GetLocation() - GetActorLocation()).SafeNormal();
+			if (FacingVector.IsZero())
+			{
+				FacingVector = -GetActorForwardVector().SafeNormal();
+			}
+			SetActorRotation(FacingVector.Rotation());
+		}
+	}
 }
diff --git a/Engine/Source/Runtime/GameplayAbilities/Private/Abilities/GameplayAbilityWorldReticle_ActorVisualization.cpp b/Engine/Source/Runtime/GameplayAbilities/Private/Abilities/GameplayAbilityWorldReticle_ActorVisualization.cpp
--- a/Engine/Source/Runtime/GameplayAbilities/Private/Abilities/GameplayAbilityWorldReticle_ActorVisualization.cpp
+++ b/Engine/Source/Runtime/GameplayAbilities/Private/Abilities/GameplayAbilityWorldReticle_ActorVisualization.cpp
@@ -29,7 +29,15 @@ AGameplayAbilityWorldReticle_ActorVisualization::AGameplayAbilityWorldReticle_Ac
 	RootComponent = CollisionComponent;
 }
 
-void AGameplayAbilityWorldReticle_ActorVisualization::InitializeReticleVisualizationInformation(AActor* VisualizationActor, UMaterialInterface *VisualizationMaterial)
+
+void AGameplayAbilityWorldReticle_ActorVisualization::Tick(float DeltaTime)
+{
+	Super::Tick(DeltaTime);
+
+	FaceTowardSource(true);
+}
+
+void AGameplayAbilityWorldReticle_ActorVisualization::InitializeReticleVisualizationInformation(AGameplayAbilityTargetActor* InTargetingActor, AActor* VisualizationActor, UMaterialInterface *VisualizationMaterial)
 {
 	if (VisualizationActor)
 	{
@@ -39,6 +47,8 @@ void AGameplayAbilityWorldReticle_ActorVisualization::InitializeReticleVisualiza
 		VisualizationActor->GetComponents(MeshComps);
 		check(MyRoot);
 
+		TargetingActor = InTargetingActor;
+
 		for (UMeshComponent* MeshComp : MeshComps)
 		{
 			//Special case: If we don't clear the root component explicitly, the component will be destroyed along with the original visualization actor.
diff --git a/Engine/Source/Runtime/GameplayAbilities/Public/Abilities/GameplayAbilityWorldReticle.h b/Engine/Source/Runtime/GameplayAbilities/Public/Abilities/GameplayAbilityWorldReticle.h
--- a/Engine/Source/Runtime/GameplayAbilities/Public/Abilities/GameplayAbilityWorldReticle.h
+++ b/Engine/Source/Runtime/GameplayAbilities/Public/Abilities/GameplayAbilityWorldReticle.h
@@ -58,7 +58,7 @@ public:
 	void SetReticleMaterialParamVector(FName ParamName, FVector value);
 
 	UFUNCTION(BlueprintCallable, Category = Reticle)
-	void FaceTowardSource();
+	void FaceTowardSource(bool bFaceIn2D);
 
 	UPROPERTY(EditAnywhere, BlueprintReadWrite, meta = (ExposeOnSpawn = "true"), Category = "Reticle")
 	FWorldReticleParameters Parameters;
diff --git a/Engine/Source/Runtime/GameplayAbilities/Public/Abilities/GameplayAbilityWorldReticle_ActorVisualization.h b/Engine/Source/Runtime/GameplayAbilities/Public/Abilities/GameplayAbilityWorldReticle_ActorVisualization.h
--- a/Engine/Source/Runtime/GameplayAbilities/Public/Abilities/GameplayAbilityWorldReticle_ActorVisualization.h
+++ b/Engine/Source/Runtime/GameplayAbilities/Public/Abilities/GameplayAbilityWorldReticle_ActorVisualization.h
@@ -14,8 +14,9 @@ class AGameplayAbilityWorldReticle_ActorVisualization : public AGameplayAbilityW
 	GENERATED_UCLASS_BODY()
 
 public:
+	virtual void Tick(float DeltaSeconds) override;
 
-	void InitializeReticleVisualizationInformation(AActor* VisualizationActor, UMaterialInterface *VisualizationMaterial);
+	void InitializeReticleVisualizationInformation(AGameplayAbilityTargetActor* InTargetingActor, AActor* VisualizationActor, UMaterialInterface *VisualizationMaterial);
 
 private:
 	/** Hardcoded collision component, so other objects don't think they can collide with the visualization actor */
