TargetActor_Radius no longer adds duplicates to its targeting data.

diff --git a/Engine/Source/Runtime/GameplayAbilities/Private/Abilities/GameplayAbilityTargetActor_Radius.cpp b/Engine/Source/Runtime/GameplayAbilities/Private/Abilities/GameplayAbilityTargetActor_Radius.cpp
--- a/Engine/Source/Runtime/GameplayAbilities/Private/Abilities/GameplayAbilityTargetActor_Radius.cpp
+++ b/Engine/Source/Runtime/GameplayAbilities/Private/Abilities/GameplayAbilityTargetActor_Radius.cpp
@@ -39,7 +39,7 @@ void AGameplayAbilityTargetActor_Radius::ConfirmTargetingAndContinue()
 	}
 }
 
-FGameplayAbilityTargetDataHandle AGameplayAbilityTargetActor_Radius::MakeTargetData(const TArray<TWeakObjectPtr<AActor> > Actors, const FVector& Origin) const
+FGameplayAbilityTargetDataHandle AGameplayAbilityTargetActor_Radius::MakeTargetData(const TArray<TWeakObjectPtr<AActor>> Actors, const FVector& Origin) const
 {
 	if (OwningAbility)
 	{
@@ -65,11 +65,11 @@ TArray<TWeakObjectPtr<AActor> >	AGameplayAbilityTargetActor_Radius::PerformOverl
 
 	TArray<TWeakObjectPtr<AActor>>	HitActors;
 
-	for (int32 i=0; i < Overlaps.Num(); ++i)
+	for (int32 i = 0; i < Overlaps.Num(); ++i)
 	{
 		//Should this check to see if these pawns are in the AimTarget list?
 		APawn* PawnActor = Cast<APawn>(Overlaps[i].GetActor());
-		if (PawnActor && Filter.FilterPassesForActor(PawnActor))
+		if (PawnActor && !HitActors.Contains(PawnActor) && Filter.FilterPassesForActor(PawnActor))
 		{
 			HitActors.Add(PawnActor);
 		}
