SetByCaller magnitude + expose spec modifying to blueprints

diff --git a/Engine/Source/Editor/GameplayAbilitiesEditor/Private/GameplayEffectModifierMagnitudeDetails.cpp b/Engine/Source/Editor/GameplayAbilitiesEditor/Private/GameplayEffectModifierMagnitudeDetails.cpp
--- a/Engine/Source/Editor/GameplayAbilitiesEditor/Private/GameplayEffectModifierMagnitudeDetails.cpp
+++ b/Engine/Source/Editor/GameplayAbilitiesEditor/Private/GameplayEffectModifierMagnitudeDetails.cpp
@@ -31,6 +31,7 @@ void FGameplayEffectModifierMagnitudeDetails::CustomizeChildren(TSharedRef<IProp
 	PropertyToCalcEnumMap.Add(FindFieldChecked<UProperty>(FGameplayEffectModifierMagnitude::StaticStruct(), GET_MEMBER_NAME_CHECKED(FGameplayEffectModifierMagnitude, ScalableFloatMagnitude)), EGameplayEffectMagnitudeCalculation::ScalableFloat);
 	PropertyToCalcEnumMap.Add(FindFieldChecked<UProperty>(FGameplayEffectModifierMagnitude::StaticStruct(), GET_MEMBER_NAME_CHECKED(FGameplayEffectModifierMagnitude, AttributeBasedMagnitude)), EGameplayEffectMagnitudeCalculation::AttributeBased);
 	PropertyToCalcEnumMap.Add(FindFieldChecked<UProperty>(FGameplayEffectModifierMagnitude::StaticStruct(), GET_MEMBER_NAME_CHECKED(FGameplayEffectModifierMagnitude, CustomMagnitude)), EGameplayEffectMagnitudeCalculation::CustomCalculationClass);
+	PropertyToCalcEnumMap.Add(FindFieldChecked<UProperty>(FGameplayEffectModifierMagnitude::StaticStruct(), GET_MEMBER_NAME_CHECKED(FGameplayEffectModifierMagnitude, SetByCallerMagnitude)), EGameplayEffectMagnitudeCalculation::SetByCaller);
 
 	// Hook into calculation type changes
 	MagnitudeCalculationTypePropertyHandle = StructPropertyHandle->GetChildHandle(GET_MEMBER_NAME_CHECKED(FGameplayEffectModifierMagnitude, MagnitudeCalculationType));
@@ -62,13 +63,19 @@ void FGameplayEffectModifierMagnitudeDetails::CustomizeChildren(TSharedRef<IProp
 		IDetailPropertyRow& PropRow = StructBuilder.AddChildProperty(CustomMagnitudePropertyHandle.ToSharedRef());
 		PropRow.Visibility(TAttribute<EVisibility>::Create(TAttribute<EVisibility>::FGetter::CreateSP(this, &FGameplayEffectModifierMagnitudeDetails::GetMagnitudeCalculationPropertyVisibility, CustomMagnitudePropertyHandle->GetProperty())));
 	}
+
+	TSharedPtr<IPropertyHandle> SetByCallerPropertyHandle  = StructPropertyHandle->GetChildHandle(GET_MEMBER_NAME_CHECKED(FGameplayEffectModifierMagnitude, SetByCallerMagnitude));
+	if (SetByCallerPropertyHandle.IsValid())
+	{
+		IDetailPropertyRow& PropRow = StructBuilder.AddChildProperty(SetByCallerPropertyHandle.ToSharedRef());
+		PropRow.Visibility(TAttribute<EVisibility>::Create(TAttribute<EVisibility>::FGetter::CreateSP(this, &FGameplayEffectModifierMagnitudeDetails::GetMagnitudeCalculationPropertyVisibility, SetByCallerPropertyHandle->GetProperty())));
+	}
 }
 
 void FGameplayEffectModifierMagnitudeDetails::OnCalculationTypeChanged()
 {
 	uint8 EnumVal;
 	MagnitudeCalculationTypePropertyHandle->GetValue(EnumVal);
-
 	VisibleCalculationType = static_cast<EGameplayEffectMagnitudeCalculation>(EnumVal);
 }
 
diff --git a/Engine/Source/Runtime/GameplayAbilities/Private/Abilities/GameplayAbility.cpp b/Engine/Source/Runtime/GameplayAbilities/Private/Abilities/GameplayAbility.cpp
--- a/Engine/Source/Runtime/GameplayAbilities/Private/Abilities/GameplayAbility.cpp
+++ b/Engine/Source/Runtime/GameplayAbilities/Private/Abilities/GameplayAbility.cpp
@@ -851,7 +851,7 @@ FActiveGameplayEffectHandle UGameplayAbility::ApplyGameplayEffectToOwner(const F
 		FGameplayEffectSpecHandle SpecHandle = MakeOutgoingGameplayEffectSpec(Handle, ActorInfo, ActivationInfo, GameplayEffect->GetClass(), GameplayEffectLevel);
 		if (SpecHandle.IsValid())
 		{
-			return ActorInfo->AbilitySystemComponent->ApplyGameplayEffectSpecToSelf(*SpecHandle.Data.Get(), ActivationInfo.GetPredictionKeyForNewAction());
+			return ApplyGameplayEffectSpecToOwner(Handle, ActorInfo, ActivationInfo, SpecHandle);
 		}
 	}
 
@@ -859,6 +859,23 @@ FActiveGameplayEffectHandle UGameplayAbility::ApplyGameplayEffectToOwner(const F
 	return FActiveGameplayEffectHandle();
 }
 
+FActiveGameplayEffectHandle UGameplayAbility::K2_ApplyGameplayEffectSpecToOwner(const FGameplayEffectSpecHandle EffectSpecHandle)
+{
+	return ApplyGameplayEffectSpecToOwner(CurrentSpecHandle, CurrentActorInfo, CurrentActivationInfo, EffectSpecHandle);
+}
+
+FActiveGameplayEffectHandle UGameplayAbility::ApplyGameplayEffectSpecToOwner(const FGameplayAbilitySpecHandle AbilityHandle, const FGameplayAbilityActorInfo* ActorInfo, const FGameplayAbilityActivationInfo ActivationInfo, const FGameplayEffectSpecHandle SpecHandle)
+{
+	if (SpecHandle.IsValid() && (ActivationInfo.ActivationMode == EGameplayAbilityActivationMode::Authority || ActivationInfo.ActivationMode == EGameplayAbilityActivationMode::Predicting))
+	{
+		ActorInfo->AbilitySystemComponent->ApplyGameplayEffectSpecToSelf(*SpecHandle.Data.Get(), ActivationInfo.GetPredictionKeyForNewAction());
+
+	}
+	return FActiveGameplayEffectHandle();
+}
+
+// -------------------------------
+
 FActiveGameplayEffectHandle UGameplayAbility::BP_ApplyGameplayEffectToTarget(FGameplayAbilityTargetDataHandle Target, TSubclassOf<UGameplayEffect> GameplayEffectClass, int32 GameplayEffectLevel)
 {
 	if (GameplayEffectClass == nullptr)
@@ -890,11 +907,27 @@ FActiveGameplayEffectHandle UGameplayAbility::ApplyGameplayEffectToTarget(const
 	}
 
 	FActiveGameplayEffectHandle EffectHandle;
+
 	if (ActivationInfo.ActivationMode == EGameplayAbilityActivationMode::Authority || ActivationInfo.ActivationMode == EGameplayAbilityActivationMode::Predicting)
 	{
 		FGameplayEffectSpecHandle SpecHandle = MakeOutgoingGameplayEffectSpec(Handle, ActorInfo, ActivationInfo, GameplayEffectClass, GameplayEffectLevel);
+		EffectHandle = ApplyGameplayEffectSpecToTarget(Handle, ActorInfo, ActivationInfo, SpecHandle, Target);
+	}
+
+	return EffectHandle;
+}
 
-		for (auto Data : Target.Data)
+FActiveGameplayEffectHandle UGameplayAbility::K2_ApplyGameplayEffectSpecToTarget(const FGameplayEffectSpecHandle SpecHandle, FGameplayAbilityTargetDataHandle TargetData)
+{
+	return ApplyGameplayEffectSpecToTarget(CurrentSpecHandle, CurrentActorInfo, CurrentActivationInfo, SpecHandle, TargetData);
+}
+
+FActiveGameplayEffectHandle UGameplayAbility::ApplyGameplayEffectSpecToTarget(const FGameplayAbilitySpecHandle AbilityHandle, const FGameplayAbilityActorInfo* ActorInfo, const FGameplayAbilityActivationInfo ActivationInfo, const FGameplayEffectSpecHandle SpecHandle, FGameplayAbilityTargetDataHandle TargetData)
+{
+	FActiveGameplayEffectHandle EffectHandle;
+	if (SpecHandle.IsValid() && ActivationInfo.ActivationMode == EGameplayAbilityActivationMode::Authority || ActivationInfo.ActivationMode == EGameplayAbilityActivationMode::Predicting)
+	{
+		for (auto Data : TargetData.Data)
 		{
 			TArray<FActiveGameplayEffectHandle> EffectHandles = Data->ApplyGameplayEffectSpec(*SpecHandle.Data.Get(), ActivationInfo.GetPredictionKeyForNewAction());
 			if (EffectHandles.Num() > 0)
@@ -902,7 +935,7 @@ FActiveGameplayEffectHandle UGameplayAbility::ApplyGameplayEffectToTarget(const
 				EffectHandle = EffectHandles[0];
 				if (EffectHandles.Num() > 1)
 				{
-					ABILITY_LOG(Warning, TEXT("ApplyGameplayEffectToTargetData_Single called on TargetData with multiple actor targets. %s"), *GetName());
+					ABILITY_LOG(Warning, TEXT("ApplyGameplayEffectSpecToTarget called on TargetData with multiple actor targets. Only returning 1 handle though we applied many. %s"), *GetName());
 				}
 			}
 		}
diff --git a/Engine/Source/Runtime/GameplayAbilities/Private/AbilitySystemBlueprintLibrary.cpp b/Engine/Source/Runtime/GameplayAbilities/Private/AbilitySystemBlueprintLibrary.cpp
--- a/Engine/Source/Runtime/GameplayAbilities/Private/AbilitySystemBlueprintLibrary.cpp
+++ b/Engine/Source/Runtime/GameplayAbilities/Private/AbilitySystemBlueprintLibrary.cpp
@@ -334,4 +334,81 @@ FVector UAbilitySystemBlueprintLibrary::GetOrigin(FGameplayCueParameters Paramet
 	}
 
 	return FVector::ZeroVector;
+}
+
+// ---------------------------------------------------------------------------------------
+
+FGameplayEffectSpecHandle UAbilitySystemBlueprintLibrary::AssignSetByCallerMagnitude(FGameplayEffectSpecHandle SpecHandle, FName DataName, float Magnitude)
+{
+	FGameplayEffectSpec* Spec = SpecHandle.Data.Get();
+	if (Spec)
+	{
+		Spec->SetMagnitude(DataName, Magnitude);
+	}
+	else
+	{
+		ABILITY_LOG(Warning, TEXT("UAbilitySystemBlueprintLibrary::AssignSetByCallerMagnitude called with invalid SpecHandle"));
+	}
+
+	return SpecHandle;
+}
+
+FGameplayEffectSpecHandle UAbilitySystemBlueprintLibrary::SetDuration(FGameplayEffectSpecHandle SpecHandle, float Duration)
+{
+	FGameplayEffectSpec* Spec = SpecHandle.Data.Get();
+	if (Spec)
+	{
+		Spec->SetDuration(Duration);
+	}
+	else
+	{
+		ABILITY_LOG(Warning, TEXT("UAbilitySystemBlueprintLibrary::SetDuration called with invalid SpecHandle"));
+	}
+
+	return SpecHandle;
+}
+
+FGameplayEffectSpecHandle UAbilitySystemBlueprintLibrary::AddGrantedTag(FGameplayEffectSpecHandle SpecHandle, FGameplayTag NewGameplayTag)
+{
+	FGameplayEffectSpec* Spec = SpecHandle.Data.Get();
+	if (Spec)
+	{
+		Spec->DynamicGrantedTags.AddTag(NewGameplayTag);
+	}
+	else
+	{
+		ABILITY_LOG(Warning, TEXT("UAbilitySystemBlueprintLibrary::AddGrantedTag called with invalid SpecHandle"));
+	}
+
+	return SpecHandle;
+}
+
+FGameplayEffectSpecHandle UAbilitySystemBlueprintLibrary::AddGrantedTags(FGameplayEffectSpecHandle SpecHandle, FGameplayTagContainer NewGameplayTags)
+{
+	FGameplayEffectSpec* Spec = SpecHandle.Data.Get();
+	if (Spec)
+	{
+		Spec->DynamicGrantedTags.AppendTags(NewGameplayTags);
+	}
+	else
+	{
+		ABILITY_LOG(Warning, TEXT("UAbilitySystemBlueprintLibrary::AddGrantedTags called with invalid SpecHandle"));
+	}
+
+	return SpecHandle;
+}
+	
+FGameplayEffectSpecHandle UAbilitySystemBlueprintLibrary::AddLinkedGameplayEffectSpec(FGameplayEffectSpecHandle SpecHandle, FGameplayEffectSpecHandle LinkedGameplayEffectSpec)
+{
+	FGameplayEffectSpec* Spec = SpecHandle.Data.Get();
+	if (Spec)
+	{
+		Spec->TargetEffectSpecs.Add(LinkedGameplayEffectSpec);
+	}
+	else
+	{
+		ABILITY_LOG(Warning, TEXT("UAbilitySystemBlueprintLibrary::AddLinkedGameplayEffectSpec called with invalid SpecHandle"));
+	}
+
+	return SpecHandle;
 }
\ No newline at end of file
diff --git a/Engine/Source/Runtime/GameplayAbilities/Private/AbilitySystemComponent.cpp b/Engine/Source/Runtime/GameplayAbilities/Private/AbilitySystemComponent.cpp
--- a/Engine/Source/Runtime/GameplayAbilities/Private/AbilitySystemComponent.cpp
+++ b/Engine/Source/Runtime/GameplayAbilities/Private/AbilitySystemComponent.cpp
@@ -581,9 +581,12 @@ FActiveGameplayEffectHandle UAbilitySystemComponent::ApplyGameplayEffectSpecToSe
 		ABILITY_LOG(Warning, TEXT("%s is periodic but also applies GameplayEffects to its target. GameplayEffects will only be applied once, not every period."), *Spec.Def->GetPathName());
 	}
 	// todo: this is ignoring the returned handles, should we put them into a TArray and return all of the handles?
-	for (const TSharedRef<FGameplayEffectSpec> TargetSpec : Spec.TargetEffectSpecs)
+	for (const FGameplayEffectSpecHandle TargetSpec: Spec.TargetEffectSpecs)
 	{
-		ApplyGameplayEffectSpecToSelf(TargetSpec.Get(), PredictionKey);
+		if (TargetSpec.IsValid())
+		{
+			ApplyGameplayEffectSpecToSelf(*TargetSpec.Data.Get(), PredictionKey);
+		}
 	}
 
 	return MyHandle;
diff --git a/Engine/Source/Runtime/GameplayAbilities/Private/GameplayEffect.cpp b/Engine/Source/Runtime/GameplayAbilities/Private/GameplayEffect.cpp
--- a/Engine/Source/Runtime/GameplayAbilities/Private/GameplayEffect.cpp
+++ b/Engine/Source/Runtime/GameplayAbilities/Private/GameplayEffect.cpp
@@ -236,7 +236,7 @@ bool FGameplayEffectModifierMagnitude::AttemptCalculateMagnitude(const FGameplay
 			break;
 
 			case EGameplayEffectMagnitudeCalculation::SetByCaller:
-				// Add check that we've been set?
+				OutCalculatedMagnitude = InRelevantSpec.GetMagnitude(SetByCallerMagnitude.DataName);
 				break;
 			default:
 				ABILITY_LOG(Error, TEXT("Unknown MagnitudeCalculationType %d in AttemptCalculateMagnitude"), (int32)MagnitudeCalculationType);
@@ -357,7 +357,7 @@ FGameplayEffectSpec::FGameplayEffectSpec(const UGameplayEffect* InDef, const FGa
 
 	for (const UGameplayEffect* TargetDef : TargetEffectDefs)
 	{
-		TargetEffectSpecs.Add(TSharedRef<FGameplayEffectSpec>(new FGameplayEffectSpec(TargetDef, EffectContext, Level)));
+		TargetEffectSpecs.Add(FGameplayEffectSpecHandle(new FGameplayEffectSpec(TargetDef, EffectContext, Level)));
 	}
 
 	// Everything is setup now, capture data from our source
@@ -567,38 +567,33 @@ void FGameplayEffectSpec::GetAllGrantedTags(OUT FGameplayTagContainer& Container
 	Container.AppendTags(Def->InheritableOwnedTagsContainer.CombinedTags);
 }
 
-void FGameplayEffectSpec::SetModifierMagnitude(int32 ModIdx, float EvaluatedMagnitude)
+void FGameplayEffectSpec::SetMagnitude(FName DataName, float Magnitude)
 {
-	if (Def->Modifiers.IsValidIndex(ModIdx) == false)
+#if !(UE_BUILD_SHIPPING || UE_BUILD_TEST)
+	const float* CurrentValue = SetByCallerMagnitudes.Find(DataName);
+	if (CurrentValue)
 	{
-		ABILITY_LOG(Error, TEXT("FGameplayEffectSpec::SetModifierMagnitude called on invalid index %d for Def: %s"), ModIdx, *Def->GetName());
-		return;
+		ABILITY_LOG(Error, TEXT("FGameplayEffectSpec::SetMagnitude called on Data %s for Def %s when this magnitude was already set. Current Value: %.2f"), *DataName.ToString(), *Def->GetName(), *CurrentValue);
 	}
+#endif
 
-	if (Def->Modifiers[ModIdx].ModifierMagnitude.GetMagnitudeCalculationType() != EGameplayEffectMagnitudeCalculation::SetByCaller)
-	{
-		ABILITY_LOG(Error, TEXT("FGameplayEffectSpec::SetModifierMagnitude called on index %d for Def: %s - this modifier not SetByCaller!"), ModIdx, *Def->GetName());
-		return;
-	}
-	
-	Modifiers[ModIdx].EvaluatedMagnitude = EvaluatedMagnitude;
+	SetByCallerMagnitudes.Add(DataName) = Magnitude;
 }
 
-void FGameplayEffectSpec::SetModifierMagnitude(FGameplayAttribute Attribute, float EvaluatedMagnitude)
+float FGameplayEffectSpec::GetMagnitude(FName DataName) const
 {
-	for (int32 ModIdx=0; ModIdx < Def->Modifiers.Num(); ++ModIdx)
+	float Magnitude = 0.f;
+	const float* Ptr = SetByCallerMagnitudes.Find(DataName);
+	if (Ptr)
 	{
-		const FGameplayModifierInfo& ModDef = Def->Modifiers[ModIdx];
-		FModifierSpec& ModSpec = Modifiers[ModIdx];
-
-		if (ModDef.Attribute == Attribute && ModDef.ModifierMagnitude.GetMagnitudeCalculationType() == EGameplayEffectMagnitudeCalculation::SetByCaller )
-		{
-			ModSpec.EvaluatedMagnitude = EvaluatedMagnitude;
-			return;
-		}
+		Magnitude = *Ptr;
+	}
+	else
+	{
+		ABILITY_LOG(Error, TEXT("FGameplayEffectSpec::GetMagnitude called for Data %s on Def %s when magnitude had not yet been set by caller."), *DataName.ToString(), *Def->GetName());
 	}
 
-	ABILITY_LOG(Error, TEXT("FGameplayEffectSpec::SetModifierMagnitude - called on Def %s for Attribute %s but no valid modifier was found."), *Def->GetName(), *Attribute.GetName());
+	return Magnitude;
 }
 
 // --------------------------------------------------------------------------------------------------------------------------------------------------------
diff --git a/Engine/Source/Runtime/GameplayAbilities/Public/Abilities/GameplayAbility.h b/Engine/Source/Runtime/GameplayAbilities/Public/Abilities/GameplayAbility.h
--- a/Engine/Source/Runtime/GameplayAbilities/Public/Abilities/GameplayAbility.h
+++ b/Engine/Source/Runtime/GameplayAbilities/Public/Abilities/GameplayAbility.h
@@ -354,6 +354,13 @@ protected:
 	/** Non blueprintcallable, safe to call on CDO/NonInstance abilities */
 	FActiveGameplayEffectHandle ApplyGameplayEffectToOwner(const FGameplayAbilitySpecHandle Handle, const FGameplayAbilityActorInfo* ActorInfo, const FGameplayAbilityActivationInfo ActivationInfo, const UGameplayEffect* GameplayEffect, float GameplayEffectLevel);
 
+	UFUNCTION(BlueprintCallable, Category = Ability, FriendlyName = "ApplyGameplayEffectSpecToOwner")
+	FActiveGameplayEffectHandle K2_ApplyGameplayEffectSpecToOwner(const FGameplayEffectSpecHandle EffectSpecHandle);
+
+	FActiveGameplayEffectHandle ApplyGameplayEffectSpecToOwner(const FGameplayAbilitySpecHandle AbilityHandle, const FGameplayAbilityActorInfo* ActorInfo, const FGameplayAbilityActivationInfo ActivationInfo, const FGameplayEffectSpecHandle SpecHandle);
+
+	// ---------
+
 	UFUNCTION(BlueprintCallable, Category = Ability, FriendlyName = "ApplyGameplayEffectToTarget")
 	FActiveGameplayEffectHandle BP_ApplyGameplayEffectToTarget(FGameplayAbilityTargetDataHandle TargetData, TSubclassOf<UGameplayEffect> GameplayEffectClass, int32 GameplayEffectLevel = 1);
 
@@ -363,6 +370,11 @@ protected:
 	/** Non blueprintcallable, safe to call on CDO/NonInstance abilities */
 	FActiveGameplayEffectHandle ApplyGameplayEffectToTarget(const FGameplayAbilitySpecHandle Handle, const FGameplayAbilityActorInfo* ActorInfo, const FGameplayAbilityActivationInfo ActivationInfo, FGameplayAbilityTargetDataHandle Target, TSubclassOf<UGameplayEffect> GameplayEffectClass, float GameplayEffectLevel);
 
+	UFUNCTION(BlueprintCallable, Category = Ability, FriendlyName = "ApplyGameplayEffectSpecToTarget")
+	FActiveGameplayEffectHandle K2_ApplyGameplayEffectSpecToTarget(const FGameplayEffectSpecHandle EffectSpecHandle, FGameplayAbilityTargetDataHandle TargetData);
+
+	FActiveGameplayEffectHandle ApplyGameplayEffectSpecToTarget(const FGameplayAbilitySpecHandle AbilityHandle, const FGameplayAbilityActorInfo* ActorInfo, const FGameplayAbilityActivationInfo ActivationInfo, const FGameplayEffectSpecHandle SpecHandle, FGameplayAbilityTargetDataHandle TargetData);
+
 	// -------------------------------------
 	//	GameplayCue
 	//	Abilities can invoke GameplayCues without having to create GameplayEffects
diff --git a/Engine/Source/Runtime/GameplayAbilities/Public/AbilitySystemBlueprintLibrary.h b/Engine/Source/Runtime/GameplayAbilities/Public/AbilitySystemBlueprintLibrary.h
--- a/Engine/Source/Runtime/GameplayAbilities/Public/AbilitySystemBlueprintLibrary.h
+++ b/Engine/Source/Runtime/GameplayAbilities/Public/AbilitySystemBlueprintLibrary.h
@@ -115,5 +115,22 @@ class GAMEPLAYABILITIES_API UAbilitySystemBlueprintLibrary : public UBlueprintFu
 
 
 	// -------------------------------------------------------------------------------
+	//		GameplayEffectSpec
+	// -------------------------------------------------------------------------------
+	
+	UFUNCTION(BlueprintCallable, Category = "Ability|GameplayEffect")
+	static FGameplayEffectSpecHandle AssignSetByCallerMagnitude(FGameplayEffectSpecHandle SpecHandle, FName DataName, float Magnitude);
+
+	UFUNCTION(BlueprintCallable, Category = "Ability|GameplayEffect")
+	static FGameplayEffectSpecHandle SetDuration(FGameplayEffectSpecHandle SpecHandle, float Duration);
+
+	UFUNCTION(BlueprintCallable, Category = "Ability|GameplayEffect")
+	static FGameplayEffectSpecHandle AddGrantedTag(FGameplayEffectSpecHandle SpecHandle, FGameplayTag NewGameplayTag);
+
+	UFUNCTION(BlueprintCallable, Category = "Ability|GameplayEffect")
+	static FGameplayEffectSpecHandle AddGrantedTags(FGameplayEffectSpecHandle SpecHandle, FGameplayTagContainer NewGameplayTags);
+
+	UFUNCTION(BlueprintCallable, Category = "Ability|GameplayEffect")
+	static FGameplayEffectSpecHandle AddLinkedGameplayEffectSpec(FGameplayEffectSpecHandle SpecHandle, FGameplayEffectSpecHandle LinkedGameplayEffectSpec);
 
 };
diff --git a/Engine/Source/Runtime/GameplayAbilities/Public/GameplayEffect.h b/Engine/Source/Runtime/GameplayAbilities/Public/GameplayEffect.h
--- a/Engine/Source/Runtime/GameplayAbilities/Public/GameplayEffect.h
+++ b/Engine/Source/Runtime/GameplayAbilities/Public/GameplayEffect.h
@@ -112,6 +112,7 @@ public:
 	FGameplayTagContainer TargetTagFilter;
 };
 
+/** Structure to encapsulte magnitude that are calculated via custom calculation */
 USTRUCT()
 struct FCustomCalculationBasedFloat
 {
@@ -138,7 +139,7 @@ public:
 	 */
 	float CalculateMagnitude(const FGameplayEffectSpec& InRelevantSpec) const;
 
-	UPROPERTY(EditDefaultsOnly, Category=CustomCalculation)
+	UPROPERTY(EditDefaultsOnly, Category=CustomCalculation, DisplayName="Calculation Class")
 	TSubclassOf<UGameplayModMagnitudeCalculation> CalculationClassMagnitude;
 
 	/** Coefficient to the custom calculation */
@@ -154,6 +155,21 @@ public:
 	FScalableFloat PostMultiplyAdditiveValue;
 };
 
+/** Struct for holding SetBytCaller data */
+USTRUCT()
+struct FSetByCallerFloat
+{
+	GENERATED_USTRUCT_BODY()
+
+	FSetByCallerFloat()
+	: DataName(NAME_None)
+	{}
+
+	/** The Name the caller (code or blueprint) will use to set this magnitude by. */
+	UPROPERTY(EditDefaultsOnly, Category=SetByCaller)
+	FName	DataName;
+};
+
 /** Struct representing the magnitude of a gameplay effect modifier, potentially calculated in numerous different ways */
 USTRUCT()
 struct FGameplayEffectModifierMagnitude
@@ -221,6 +237,10 @@ protected:
 	UPROPERTY(EditDefaultsOnly, Category=Magnitude)
 	FCustomCalculationBasedFloat CustomMagnitude;
 
+	/** Magnitude value represented by a SetByCaller magnitude */
+	UPROPERTY(EditDefaultsOnly, Category=Magnitude)
+	FSetByCallerFloat SetByCallerMagnitude;
+
 	// @hack: @todo: This is temporary to aid in post-load fix-up w/o exposing members publicly
 	friend class UGameplayEffect;
 	friend class FGameplayEffectModifierMagnitudeDetails;
@@ -480,7 +500,7 @@ public:
 	FScalableFloat	ChanceToExecuteOnGameplayEffect;
 
 	/** other gameplay effects that will be applied to the target of this effect if this effect applies */
-	UPROPERTY(EditDefaultsOnly, BlueprintReadOnly, Category = GameplayEffect)
+	UPROPERTY(EditDefaultsOnly, BlueprintReadOnly, Category = GameplayEffect,meta = (DisplayName = "LinkedGameplayEffects"))
 	TArray<TSubclassOf<UGameplayEffect>> TargetEffectClasses;
 
 	/** Deprecated. Use TargetEffectClasses instead */
@@ -842,7 +862,7 @@ struct GAMEPLAYABILITIES_API FGameplayEffectSpec
 	EGameplayEffectStackingPolicy::Type GetStackingType() const;
 
 	/** other effects that need to be applied to the target if this effect is successful */
-	TArray< TSharedRef< FGameplayEffectSpec > > TargetEffectSpecs;
+	TArray< FGameplayEffectSpecHandle > TargetEffectSpecs;
 
 	/** Set the context info: who and where this spec came from. */
 	void SetContext(FGameplayEffectContextHandle NewEffectContext);
@@ -854,11 +874,11 @@ struct GAMEPLAYABILITIES_API FGameplayEffectSpec
 
 	void GetAllGrantedTags(OUT FGameplayTagContainer& Container) const;
 
-	/** Sets the magnitude of one of our modifiers. This must be done on modifiers which are EGameplayEffectMagnitudeCalculation::SetByCaller  */
-	void SetModifierMagnitude(int32 ModIdx, float EvaluatedMagnitude);
+	/** Sets the magnitude of a SetByCaller modifier */
+	void SetMagnitude(FName DataName, float Magnitude);
 
-	/** Set magnitude of (the first) SetByCaller modifier with the given GameplayAttribute  */
-	void SetModifierMagnitude(FGameplayAttribute Attribute, float EvaluatedMagnitude);
+	/** Returns the magnitude of a SetByCaller modifier. Will return 0.f and Warn if the magnitude has not been set. */
+	float GetMagnitude(FName DataName) const;
 
 	// The duration in seconds of this effect
 	// instantaneous effects should have a duration of UGameplayEffect::INSTANT_APPLICATION
@@ -918,6 +938,9 @@ struct GAMEPLAYABILITIES_API FGameplayEffectSpec
 
 private:
 
+	/** Map of set by caller magnitudes */
+	TMap<FName, float>	SetByCallerMagnitudes;
+
 	void CaptureDataFromSource();
 
 	/** Helper function to initialize all of the capture definitions required by the spec */
@@ -1354,62 +1377,3 @@ public:
 };
 
 
-/** Allows blueprints to generate a GameplayEffectSpec once and then reference it by handle, to apply it multiple times/multiple targets. */
-USTRUCT(BlueprintType)
-struct GAMEPLAYABILITIES_API FGameplayEffectSpecHandle
-{
-	GENERATED_USTRUCT_BODY()
-
-	FGameplayEffectSpecHandle() { }
-	FGameplayEffectSpecHandle(FGameplayEffectSpec* DataPtr)
-		: Data(DataPtr)
-	{
-
-	}
-
-	TSharedPtr<FGameplayEffectSpec>	Data;
-
-	bool IsValidCache;
-
-	void Clear()
-	{
-		Data.Reset();
-	}
-
-	bool IsValid() const
-	{
-		return Data.IsValid();
-	}
-
-	bool NetSerialize(FArchive& Ar, class UPackageMap* Map, bool& bOutSuccess)
-	{
-		ABILITY_LOG(Fatal, TEXT("FGameplayEffectSpecHandle should not be NetSerialized"));
-		return false;
-	}
-
-	/** Comparison operator */
-	bool operator==(FGameplayEffectSpecHandle const& Other) const
-	{
-		// Both invalid structs or both valid and Pointer compare (???) // deep comparison equality
-		bool bBothValid = IsValid() && Other.IsValid();
-		bool bBothInvalid = !IsValid() && !Other.IsValid();
-		return (bBothInvalid || (bBothValid && (Data.Get() == Other.Data.Get())));
-	}
-
-	/** Comparison operator */
-	bool operator!=(FGameplayEffectSpecHandle const& Other) const
-	{
-		return !(FGameplayEffectSpecHandle::operator==(Other));
-	}
-};
-
-template<>
-struct TStructOpsTypeTraits<FGameplayEffectSpecHandle> : public TStructOpsTypeTraitsBase
-{
-	enum
-	{
-		WithCopy = true,		// Necessary so that TSharedPtr<FGameplayAbilityTargetData> Data is copied around
-		WithNetSerializer = true,
-		WithIdenticalViaEquality = true,
-	};
-};
diff --git a/Engine/Source/Runtime/GameplayAbilities/Public/GameplayEffectTypes.h b/Engine/Source/Runtime/GameplayAbilities/Public/GameplayEffectTypes.h
--- a/Engine/Source/Runtime/GameplayAbilities/Public/GameplayEffectTypes.h
+++ b/Engine/Source/Runtime/GameplayAbilities/Public/GameplayEffectTypes.h
@@ -5,6 +5,7 @@
 #include "GameplayTagAssetInterface.h"
 #include "AttributeSet.h"
 #include "GameplayPrediction.h"
+#include "AbilitySystemLog.h"
 #include "GameplayEffectTypes.generated.h"
 
 #define SKILL_SYSTEM_AGGREGATOR_DEBUG 1
@@ -777,3 +778,64 @@ private:
 	mutable FGameplayTagContainer CachedAggregator;
 	mutable bool CacheIsValid;
 };
+
+
+/** Allows blueprints to generate a GameplayEffectSpec once and then reference it by handle, to apply it multiple times/multiple targets. */
+USTRUCT(BlueprintType)
+struct GAMEPLAYABILITIES_API FGameplayEffectSpecHandle
+{
+	GENERATED_USTRUCT_BODY()
+
+	FGameplayEffectSpecHandle() { }
+	FGameplayEffectSpecHandle(FGameplayEffectSpec* DataPtr)
+		: Data(DataPtr)
+	{
+
+	}
+
+	TSharedPtr<FGameplayEffectSpec>	Data;
+
+	bool IsValidCache;
+
+	void Clear()
+	{
+		Data.Reset();
+	}
+
+	bool IsValid() const
+	{
+		return Data.IsValid();
+	}
+
+	bool NetSerialize(FArchive& Ar, class UPackageMap* Map, bool& bOutSuccess)
+	{
+		ABILITY_LOG(Fatal, TEXT("FGameplayEffectSpecHandle should not be NetSerialized"));
+		return false;
+	}
+
+	/** Comparison operator */
+	bool operator==(FGameplayEffectSpecHandle const& Other) const
+	{
+		// Both invalid structs or both valid and Pointer compare (???) // deep comparison equality
+		bool bBothValid = IsValid() && Other.IsValid();
+		bool bBothInvalid = !IsValid() && !Other.IsValid();
+		return (bBothInvalid || (bBothValid && (Data.Get() == Other.Data.Get())));
+	}
+
+	/** Comparison operator */
+	bool operator!=(FGameplayEffectSpecHandle const& Other) const
+	{
+		return !(FGameplayEffectSpecHandle::operator==(Other));
+	}
+};
+
+template<>
+struct TStructOpsTypeTraits<FGameplayEffectSpecHandle> : public TStructOpsTypeTraitsBase
+{
+	enum
+	{
+		WithCopy = true,		// Necessary so that TSharedPtr<FGameplayAbilityTargetData> Data is copied around
+		WithNetSerializer = true,
+		WithIdenticalViaEquality = true,
+	};
+};
