Fix build failure on Clang due to precedence warning.

diff --git a/Engine/Source/Runtime/GameplayAbilities/Private/Abilities/GameplayAbility.cpp b/Engine/Source/Runtime/GameplayAbilities/Private/Abilities/GameplayAbility.cpp
--- a/Engine/Source/Runtime/GameplayAbilities/Private/Abilities/GameplayAbility.cpp
+++ b/Engine/Source/Runtime/GameplayAbilities/Private/Abilities/GameplayAbility.cpp
@@ -925,7 +925,7 @@ FActiveGameplayEffectHandle UGameplayAbility::K2_ApplyGameplayEffectSpecToTarget
 FActiveGameplayEffectHandle UGameplayAbility::ApplyGameplayEffectSpecToTarget(const FGameplayAbilitySpecHandle AbilityHandle, const FGameplayAbilityActorInfo* ActorInfo, const FGameplayAbilityActivationInfo ActivationInfo, const FGameplayEffectSpecHandle SpecHandle, FGameplayAbilityTargetDataHandle TargetData)
 {
 	FActiveGameplayEffectHandle EffectHandle;
-	if (SpecHandle.IsValid() && ActivationInfo.ActivationMode == EGameplayAbilityActivationMode::Authority || ActivationInfo.ActivationMode == EGameplayAbilityActivationMode::Predicting)
+	if (SpecHandle.IsValid() && (ActivationInfo.ActivationMode == EGameplayAbilityActivationMode::Authority || ActivationInfo.ActivationMode == EGameplayAbilityActivationMode::Predicting))
 	{
 		for (auto Data : TargetData.Data)
 		{
