Do RemoveTags on application of gameplayeffect, not on execute.

diff --git a/Engine/Source/Runtime/GameplayAbilities/Private/AbilitySystemComponent.cpp b/Engine/Source/Runtime/GameplayAbilities/Private/AbilitySystemComponent.cpp
--- a/Engine/Source/Runtime/GameplayAbilities/Private/AbilitySystemComponent.cpp
+++ b/Engine/Source/Runtime/GameplayAbilities/Private/AbilitySystemComponent.cpp
@@ -580,6 +580,16 @@ FActiveGameplayEffectHandle UAbilitySystemComponent::ApplyGameplayEffectSpecToSe
 	{
 		ABILITY_LOG(Warning, TEXT("%s is periodic but also applies GameplayEffects to its target. GameplayEffects will only be applied once, not every period."), *Spec.Def->GetPathName());
 	}
+
+	// ------------------------------------------------------
+	//	Remove gameplay effects with tags
+	//		Remove any active gameplay effects that match the RemoveGameplayEffectsWithTags in the definition for this spec
+	// ------------------------------------------------------
+	if (IsOwnerActorAuthoritative())
+	{
+		ActiveGameplayEffects.RemoveActiveEffects(FActiveGameplayEffectQuery(&Spec.Def->RemoveGameplayEffectsWithTags.CombinedTags));
+	}
+
 	// todo: this is ignoring the returned handles, should we put them into a TArray and return all of the handles?
 	for (const FGameplayEffectSpecHandle TargetSpec: Spec.TargetEffectSpecs)
 	{
diff --git a/Engine/Source/Runtime/GameplayAbilities/Private/GameplayEffect.cpp b/Engine/Source/Runtime/GameplayAbilities/Private/GameplayEffect.cpp
--- a/Engine/Source/Runtime/GameplayAbilities/Private/GameplayEffect.cpp
+++ b/Engine/Source/Runtime/GameplayAbilities/Private/GameplayEffect.cpp
@@ -990,16 +990,9 @@ void FActiveGameplayEffectsContainer::ExecuteActiveEffectsFrom(FGameplayEffectSp
 		}
 	}
 
-	// ------------------------------------------------------
-	//	Remove gameplay effects with tags
-	//		Remove any active gameplay effects that match the RemoveGameplayEffectsWithTags in the definition for this spec
-	// ------------------------------------------------------
-	RemoveActiveEffects(FActiveGameplayEffectQuery(&Spec.Def->RemoveGameplayEffectsWithTags.CombinedTags));
-
 	// ------------------------------------------------------
 	//	Invoke GameplayCue events
 	// ------------------------------------------------------
-
 	
 	// Prune the modified attributes before we replicate
 	Spec.PruneModifiedAttributes();
