Fixed assert caused by appling a gameplay effect with a NULL modifier attribute.

diff --git a/Engine/Source/Runtime/GameplayAbilities/Private/AbilitySystemComponent.cpp b/Engine/Source/Runtime/GameplayAbilities/Private/AbilitySystemComponent.cpp
--- a/Engine/Source/Runtime/GameplayAbilities/Private/AbilitySystemComponent.cpp
+++ b/Engine/Source/Runtime/GameplayAbilities/Private/AbilitySystemComponent.cpp
@@ -482,6 +482,12 @@ FActiveGameplayEffectHandle UAbilitySystemComponent::ApplyGameplayEffectSpecToSe
 	// We also need to handle things in the execution list
 	for (const FGameplayModifierInfo& Mod : Spec.Def->Modifiers)
 	{
+		if (!Mod.Attribute.IsValid())
+		{
+			ABILITY_LOG(Warning, TEXT("%s has a null modifier attribute."), *Spec.Def->GetPathName());
+			return FActiveGameplayEffectHandle();
+		}
+
 		if (HasAttributeSetForAttribute(Mod.Attribute) == false)
 		{
 			return FActiveGameplayEffectHandle();
