- Change gameplay effect execution to make a duplicate of the GE spec for non-instant effects, as execution may dirty the spec otherwise, which would result in incorrect calculations for DoTs, etc.
- Change executions to take non-const params so that they can also modify specs as necessary

diff --git a/Engine/Source/Runtime/GameplayAbilities/Private/AbilitySystemComponent.cpp b/Engine/Source/Runtime/GameplayAbilities/Private/AbilitySystemComponent.cpp
--- a/Engine/Source/Runtime/GameplayAbilities/Private/AbilitySystemComponent.cpp
+++ b/Engine/Source/Runtime/GameplayAbilities/Private/AbilitySystemComponent.cpp
@@ -456,7 +456,7 @@ FActiveGameplayEffectHandle UAbilitySystemComponent::ApplyGameplayEffectSpecToSe
 		return FActiveGameplayEffectHandle();
 	}
 
-	// Don't allow prediciton of periodic effects
+	// Don't allow prediction of periodic effects
 	if(PredictionKey.IsValidKey() && Spec.GetPeriod() > 0.f)
 	{
 		if(IsOwnerActorAuthoritative())
@@ -471,7 +471,7 @@ FActiveGameplayEffectHandle UAbilitySystemComponent::ApplyGameplayEffectSpecToSe
 		}
 	}
 
-	// Are we currently immunte to this? (ApplicationImmunity)
+	// Are we currently immune to this? (ApplicationImmunity)
 	if (ActiveGameplayEffects.CheckApplicationImmunity(Spec))
 	{
 		return FActiveGameplayEffectHandle();
@@ -502,7 +502,7 @@ FActiveGameplayEffectHandle UAbilitySystemComponent::ApplyGameplayEffectSpecToSe
 	}
 
 	// Get MyTags.
-	//	We may want to cache off a GameplayTagContainer instead of rebuilding it everytime.
+	//	We may want to cache off a GameplayTagContainer instead of rebuilding it every time.
 	//	But this will also be where we need to merge in context tags? (Headshot, executing ability, etc?)
 	//	Or do we push these tags into (our copy of the spec)?
 
diff --git a/Engine/Source/Runtime/GameplayAbilities/Private/GameplayEffect.cpp b/Engine/Source/Runtime/GameplayAbilities/Private/GameplayEffect.cpp
--- a/Engine/Source/Runtime/GameplayAbilities/Private/GameplayEffect.cpp
+++ b/Engine/Source/Runtime/GameplayAbilities/Private/GameplayEffect.cpp
@@ -940,28 +940,38 @@ void FActiveGameplayEffectsContainer::ExecuteActiveEffectsFrom(FGameplayEffectSp
 		}
 	}
 
-	float ChanceToExecute = Spec.GetChanceToExecuteOnGameplayEffect();		// Not implemented? Should we just remove ChanceToExecute?
+	// The process of executing a gameplay spec may cause modifications to the spec, however those modifications
+	// should not persist across multiple executions on the same spec. Therefore, if the spec isn't instantaneous, make a 
+	// copy that can be modified and pass it along to each part of the execution process.
+	TSharedPtr<FGameplayEffectSpec> StackSpec;
+	if (Spec.GetDuration() != UGameplayEffect::INSTANT_APPLICATION)
+	{
+		StackSpec = TSharedPtr<FGameplayEffectSpec>(new FGameplayEffectSpec(Spec));
+	}
+	FGameplayEffectSpec& SpecToUse = StackSpec.IsValid() ? *(StackSpec.Get()) : Spec;
+
+	float ChanceToExecute = SpecToUse.GetChanceToExecuteOnGameplayEffect();		// Not implemented? Should we just remove ChanceToExecute?
 
 	// Capture our own tags.
 	// TODO: We should only capture them if we need to. We may have snapshotted target tags (?) (in the case of dots with exotic setups?)
 
-	Spec.CapturedTargetTags.GetActorTags().RemoveAllTags();
-	Owner->GetOwnedGameplayTags(Spec.CapturedTargetTags.GetActorTags());
+	SpecToUse.CapturedTargetTags.GetActorTags().RemoveAllTags();
+	Owner->GetOwnedGameplayTags(SpecToUse.CapturedTargetTags.GetActorTags());
 
-	Spec.CalculateModifierMagnitudes();
+	SpecToUse.CalculateModifierMagnitudes();
 
 	// ------------------------------------------------------
 	//	Modifiers
 	//		These will modify the base value of attributes
 	// ------------------------------------------------------
 
-	for(int32 ModIdx = 0; ModIdx < Spec.Modifiers.Num(); ++ModIdx)
+	for (int32 ModIdx = 0; ModIdx < SpecToUse.Modifiers.Num(); ++ModIdx)
 	{
-		const FGameplayModifierInfo& ModDef = Spec.Def->Modifiers[ModIdx];
-		const FModifierSpec& ModSpec = Spec.Modifiers[ModIdx];
+		const FGameplayModifierInfo& ModDef = SpecToUse.Def->Modifiers[ModIdx];
+		const FModifierSpec& ModSpec = SpecToUse.Modifiers[ModIdx];
 
 		FGameplayModifierEvaluatedData EvalData(ModDef.Attribute, ModDef.ModifierOp, ModSpec.GetEvaluatedMagnitude());
-		InvokeGameplayCueExecute |= InternalExecuteMod(Spec, EvalData);
+		InvokeGameplayCueExecute |= InternalExecuteMod(SpecToUse, EvalData);
 	}
 
 	// ------------------------------------------------------
@@ -969,7 +979,7 @@ void FActiveGameplayEffectsContainer::ExecuteActiveEffectsFrom(FGameplayEffectSp
 	//		This will run custom code to 'do stuff'
 	// ------------------------------------------------------
 	
-	for (const FGameplayEffectExecutionDefinition& CurExecDef : Spec.Def->Executions)
+	for (const FGameplayEffectExecutionDefinition& CurExecDef : SpecToUse.Def->Executions)
 	{
 		if (CurExecDef.CalculationClass)
 		{
@@ -979,13 +989,13 @@ void FActiveGameplayEffectsContainer::ExecuteActiveEffectsFrom(FGameplayEffectSp
 			TArray<FGameplayModifierEvaluatedData> OutModifiers;
 
 			// Run the custom execution
-			FGameplayEffectCustomExecutionParameters ExecutionParams(Spec, CurExecDef.CalculationModifiers, Owner);
+			FGameplayEffectCustomExecutionParameters ExecutionParams(SpecToUse, CurExecDef.CalculationModifiers, Owner);
 			ExecCDO->Execute(ExecutionParams, OutModifiers);
 
 			// Execute any mods the custom execution yielded
 			for (FGameplayModifierEvaluatedData& CurExecMod : OutModifiers)
 			{
-				InvokeGameplayCueExecute |= InternalExecuteMod(Spec, CurExecMod);
+				InvokeGameplayCueExecute |= InternalExecuteMod(SpecToUse, CurExecMod);
 			}
 		}
 	}
@@ -995,15 +1005,15 @@ void FActiveGameplayEffectsContainer::ExecuteActiveEffectsFrom(FGameplayEffectSp
 	// ------------------------------------------------------
 	
 	// Prune the modified attributes before we replicate
-	Spec.PruneModifiedAttributes();
+	SpecToUse.PruneModifiedAttributes();
 
-	if (InvokeGameplayCueExecute && Spec.Def->GameplayCues.Num())
+	if (InvokeGameplayCueExecute && SpecToUse.Def->GameplayCues.Num())
 	{
 		// TODO: check replication policy. Right now we will replicate every execute via a multicast RPC
 
-		ABILITY_LOG(Log, TEXT("Invoking Execute GameplayCue for %s"), *Spec.ToSimpleString() );
+		ABILITY_LOG(Log, TEXT("Invoking Execute GameplayCue for %s"), *SpecToUse.ToSimpleString());
 		Owner->ForceReplication();
-		Owner->NetMulticast_InvokeGameplayCueExecuted_FromSpec(Spec, PredictionKey);
+		Owner->NetMulticast_InvokeGameplayCueExecuted_FromSpec(SpecToUse, PredictionKey);
 	}
 }
 
diff --git a/Engine/Source/Runtime/GameplayAbilities/Private/GameplayEffectExecutionCalculation.cpp b/Engine/Source/Runtime/GameplayAbilities/Private/GameplayEffectExecutionCalculation.cpp
--- a/Engine/Source/Runtime/GameplayAbilities/Private/GameplayEffectExecutionCalculation.cpp
+++ b/Engine/Source/Runtime/GameplayAbilities/Private/GameplayEffectExecutionCalculation.cpp
@@ -10,7 +10,7 @@ FGameplayEffectCustomExecutionParameters::FGameplayEffectCustomExecutionParamete
 {
 }
 
-FGameplayEffectCustomExecutionParameters::FGameplayEffectCustomExecutionParameters(const FGameplayEffectSpec& InOwningSpec, const TArray<FGameplayEffectExecutionScopedModifierInfo>& InScopedMods, UAbilitySystemComponent* InTargetAbilityComponent)
+FGameplayEffectCustomExecutionParameters::FGameplayEffectCustomExecutionParameters(FGameplayEffectSpec& InOwningSpec, const TArray<FGameplayEffectExecutionScopedModifierInfo>& InScopedMods, UAbilitySystemComponent* InTargetAbilityComponent)
 	: OwningSpec(&InOwningSpec)
 	, TargetAbilitySystemComponent(InTargetAbilityComponent)
 {
@@ -50,6 +50,13 @@ const FGameplayEffectSpec& FGameplayEffectCustomExecutionParameters::GetOwningSp
 	return *OwningSpec;
 }
 
+FGameplayEffectSpec& FGameplayEffectCustomExecutionParameters::GetOwningSpec()
+{
+	check(OwningSpec);
+	return *OwningSpec;
+}
+
+
 UAbilitySystemComponent* FGameplayEffectCustomExecutionParameters::GetTargetAbilitySystemComponent() const
 {
 	return TargetAbilitySystemComponent.Get();
@@ -192,6 +199,6 @@ void UGameplayEffectExecutionCalculation::GetValidScopedModifierAttributeCapture
 }
 #endif // #if WITH_EDITORONLY_DATA
 
-void UGameplayEffectExecutionCalculation::Execute_Implementation(const FGameplayEffectCustomExecutionParameters& ExecutionParams, OUT TArray<FGameplayModifierEvaluatedData>& OutAdditionalModifiers) const
+void UGameplayEffectExecutionCalculation::Execute_Implementation(FGameplayEffectCustomExecutionParameters& ExecutionParams, OUT TArray<FGameplayModifierEvaluatedData>& OutAdditionalModifiers) const
 {
 }
diff --git a/Engine/Source/Runtime/GameplayAbilities/Public/GameplayEffectExecutionCalculation.h b/Engine/Source/Runtime/GameplayAbilities/Public/GameplayEffectExecutionCalculation.h
--- a/Engine/Source/Runtime/GameplayAbilities/Public/GameplayEffectExecutionCalculation.h
+++ b/Engine/Source/Runtime/GameplayAbilities/Public/GameplayEffectExecutionCalculation.h
@@ -20,11 +20,14 @@ public:
 
 	// Constructors
 	FGameplayEffectCustomExecutionParameters();
-	FGameplayEffectCustomExecutionParameters(const FGameplayEffectSpec& InOwningSpec, const TArray<FGameplayEffectExecutionScopedModifierInfo>& InScopedMods, UAbilitySystemComponent* InTargetAbilityComponent);
+	FGameplayEffectCustomExecutionParameters(FGameplayEffectSpec& InOwningSpec, const TArray<FGameplayEffectExecutionScopedModifierInfo>& InScopedMods, UAbilitySystemComponent* InTargetAbilityComponent);
 
 	/** Simple accessor to owning gameplay spec */
 	const FGameplayEffectSpec& GetOwningSpec() const;
 
+	/** Simple non-const accessor to owning gameplay spec */
+	FGameplayEffectSpec& GetOwningSpec();
+
 	/** Simple accessor to target ability system component */
 	UAbilitySystemComponent* GetTargetAbilitySystemComponent() const;
 
@@ -96,7 +99,7 @@ private:
 	TMap<FGameplayEffectAttributeCaptureDefinition, FAggregator> ScopedModifierAggregators;
 
 	/** Owning gameplay effect spec */
-	const FGameplayEffectSpec* OwningSpec;
+	FGameplayEffectSpec* OwningSpec;
 
 	/** Target ability system component of the execution */
 	TWeakObjectPtr<UAbilitySystemComponent> TargetAbilitySystemComponent;
@@ -135,5 +138,5 @@ public:
 	 * @param OutAdditionalModifiers	[OUT] Additional modifiers the custom execution has generated and would like executed upon the target
 	 */
 	UFUNCTION(BlueprintNativeEvent, Category="Calculation")
-	void Execute(const FGameplayEffectCustomExecutionParameters& ExecutionParams, TArray<FGameplayModifierEvaluatedData>& OutAdditionalModifiers) const;
+	void Execute(FGameplayEffectCustomExecutionParameters& ExecutionParams, TArray<FGameplayModifierEvaluatedData>& OutAdditionalModifiers) const;
 };
