- Make gameplay abilities also add their source object's tags to outgoing spec tags when appropriate; Note that this still doesn't encompass game-specific, additional tag containers, such as tags on "damage sources," which will need to be added later, as necessary

diff --git a/Engine/Source/Runtime/GameplayAbilities/Private/Abilities/GameplayAbility.cpp b/Engine/Source/Runtime/GameplayAbilities/Private/Abilities/GameplayAbility.cpp
--- a/Engine/Source/Runtime/GameplayAbilities/Private/Abilities/GameplayAbility.cpp
+++ b/Engine/Source/Runtime/GameplayAbilities/Private/Abilities/GameplayAbility.cpp
@@ -542,14 +542,28 @@ FGameplayEffectSpecHandle UGameplayAbility::MakeOutgoingGameplayEffectSpec(const
 	FGameplayEffectSpecHandle NewHandle = ActorInfo->AbilitySystemComponent->MakeOutgoingSpec(GameplayEffectClass, Level, GetEffectContext(ActorInfo));
 	if (NewHandle.IsValid())
 	{
-		ApplyAbilityTagsToGameplayEffectSpec(*NewHandle.Data.Get());
+		FGameplayAbilitySpec* AbilitySpec =  ActorInfo->AbilitySystemComponent->FindAbilitySpecFromHandle(Handle);
+		ApplyAbilityTagsToGameplayEffectSpec(*NewHandle.Data.Get(), AbilitySpec);
 	}
 	return NewHandle;
 }
 
-void UGameplayAbility::ApplyAbilityTagsToGameplayEffectSpec(FGameplayEffectSpec& Spec) const
+void UGameplayAbility::ApplyAbilityTagsToGameplayEffectSpec(FGameplayEffectSpec& Spec, FGameplayAbilitySpec* AbilitySpec) const
 {
 	Spec.CapturedSourceTags.GetSpecTags().AppendTags(AbilityTags);
+
+	// Allow the source object of the ability to propagate tags along as well
+	if (AbilitySpec)
+	{
+		const IGameplayTagAssetInterface* SourceObjAsTagInterface = Cast<IGameplayTagAssetInterface>(AbilitySpec->SourceObject);
+		if (SourceObjAsTagInterface)
+		{
+			FGameplayTagContainer SourceObjTags;
+			SourceObjAsTagInterface->GetOwnedGameplayTags(SourceObjTags);
+
+			Spec.CapturedSourceTags.GetSpecTags().AppendTags(SourceObjTags);
+		}
+	}
 }
 
 /** Fixme: Naming is confusing here */
diff --git a/Engine/Source/Runtime/GameplayAbilities/Public/Abilities/GameplayAbility.h b/Engine/Source/Runtime/GameplayAbilities/Public/Abilities/GameplayAbility.h
--- a/Engine/Source/Runtime/GameplayAbilities/Public/Abilities/GameplayAbility.h
+++ b/Engine/Source/Runtime/GameplayAbilities/Public/Abilities/GameplayAbility.h
@@ -455,8 +455,8 @@ public:
 
 	FGameplayEffectSpecHandle MakeOutgoingGameplayEffectSpec(const FGameplayAbilitySpecHandle Handle, const FGameplayAbilityActorInfo* ActorInfo, const FGameplayAbilityActivationInfo ActivationInfo, TSubclassOf<UGameplayEffect> GameplayEffectClass, float Level = 1.f) const;
 
-	/** Add the Ability's tags to the given GameplayEffectSpec. This is likely to be overriden per project. */
-	virtual void ApplyAbilityTagsToGameplayEffectSpec(FGameplayEffectSpec& Spec) const;	
+	/** Add the Ability's tags to the given GameplayEffectSpec. This is likely to be overridden per project. */
+	virtual void ApplyAbilityTagsToGameplayEffectSpec(FGameplayEffectSpec& Spec, FGameplayAbilitySpec* AbilitySpec) const;
 
 	/** Returns the currently playing montage for this ability, if any */
 	UFUNCTION(BlueprintCallable, Category = Animation)
