BP function to assist with "teleport back home" abilities.

diff --git a/Engine/Source/Runtime/GameplayAbilities/Private/AbilitySystemBlueprintLibrary.cpp b/Engine/Source/Runtime/GameplayAbilities/Private/AbilitySystemBlueprintLibrary.cpp
--- a/Engine/Source/Runtime/GameplayAbilities/Private/AbilitySystemBlueprintLibrary.cpp
+++ b/Engine/Source/Runtime/GameplayAbilities/Private/AbilitySystemBlueprintLibrary.cpp
@@ -20,6 +20,21 @@ UAbilitySystemComponent* UAbilitySystemBlueprintLibrary::GetAbilitySystemCompone
 	return UAbilitySystemGlobals::GetAbilitySystemComponentFromActor(Actor);
 }
 
+AActor* UAbilitySystemBlueprintLibrary::FindPlayerStart(AController* PlayerController)
+{
+	if (PlayerController)
+	{
+		if (UWorld* World = PlayerController->GetWorld())
+		{
+			if (AGameMode* GameMode = World->GetAuthGameMode())
+			{
+				return GameMode->FindPlayerStart(PlayerController);
+			}
+		}
+	}
+	return NULL;
+}
+
 FGameplayAbilityTargetDataHandle UAbilitySystemBlueprintLibrary::AppendTargetDataHandle(FGameplayAbilityTargetDataHandle TargetHandle, FGameplayAbilityTargetDataHandle HandleToAdd)
 {
 	TargetHandle.Append(&HandleToAdd);
diff --git a/Engine/Source/Runtime/GameplayAbilities/Public/AbilitySystemBlueprintLibrary.h b/Engine/Source/Runtime/GameplayAbilities/Public/AbilitySystemBlueprintLibrary.h
--- a/Engine/Source/Runtime/GameplayAbilities/Public/AbilitySystemBlueprintLibrary.h
+++ b/Engine/Source/Runtime/GameplayAbilities/Public/AbilitySystemBlueprintLibrary.h
@@ -23,6 +23,8 @@ class GAMEPLAYABILITIES_API UAbilitySystemBlueprintLibrary : public UBlueprintFu
 	UFUNCTION(BlueprintPure, Category = Ability)
 	static UAbilitySystemComponent* GetAbilitySystemComponent(AActor *Actor);
 
+	static AActor* FindPlayerStart(AController* PlayerController);
+
 	// -------------------------------------------------------------------------------
 	//		TargetData
 	// -------------------------------------------------------------------------------
