Code-only cleanup.

diff --git a/Engine/Source/Runtime/GameplayAbilities/Public/Abilities/GameplayAbilityTargetTypes.h b/Engine/Source/Runtime/GameplayAbilities/Public/Abilities/GameplayAbilityTargetTypes.h
--- a/Engine/Source/Runtime/GameplayAbilities/Public/Abilities/GameplayAbilityTargetTypes.h
+++ b/Engine/Source/Runtime/GameplayAbilities/Public/Abilities/GameplayAbilityTargetTypes.h
@@ -245,45 +245,6 @@ struct TStructOpsTypeTraits<FGameplayAbilityTargetDataHandle> : public TStructOp
 	};
 };
 
-/*
-USTRUCT(BlueprintType)
-struct FGameplayAbilityTargetDataActorFilter
-{
-	GENERATED_USTRUCT_BODY()
-
-	virtual bool FilterPassesForActor(const AActor* ActorToBeFiltered) const
-	{
-		return true;
-	}
-};
-
-USTRUCT(BlueprintType)
-struct FGameplayAbilityTargetDataActorFilterHandleBase
-{
-	GENERATED_USTRUCT_BODY()
-
-	TSharedPtr<FGameplayAbilityTargetDataActorFilter>	Filter;
-
-	bool operator()(const TWeakObjectPtr<AActor> A) const
-	{
-		if (Filter.IsValid())
-		{
-			return Filter.Get()->FilterPassesForActor(A.Get());
-		}
-		return true;
-	}
-
-	bool operator()(const AActor* A) const
-	{
-		if (Filter.IsValid())
-		{
-			return Filter.Get()->FilterPassesForActor(A);
-		}
-		return true;
-	}
-};
-*/
-
 USTRUCT(BlueprintType)
 struct GAMEPLAYABILITIES_API FGameplayAbilityTargetingLocationInfo
 {
@@ -335,7 +296,6 @@ public:
 
 	FGameplayAbilityTargetDataHandle MakeTargetDataHandleFromHitResult(TWeakObjectPtr<UGameplayAbility> Ability, FHitResult HitResult) const;
 	FGameplayAbilityTargetDataHandle MakeTargetDataHandleFromHitResults(TWeakObjectPtr<UGameplayAbility> Ability, const TArray<FHitResult>& HitResults) const;
-
 	FGameplayAbilityTargetDataHandle MakeTargetDataHandleFromActors(TArray<TWeakObjectPtr<AActor>> TargetActors, bool OneActorPerHandle = false) const;
 
 	/** Type of location used - will determine what data is transmitted over the network and what fields are used when calculating position. */
