Removed one-frame lag on Actor Visualization reticles.

diff --git a/Engine/Source/Runtime/GameplayAbilities/Private/Abilities/GameplayAbilityWorldReticle_ActorVisualization.cpp b/Engine/Source/Runtime/GameplayAbilities/Private/Abilities/GameplayAbilityWorldReticle_ActorVisualization.cpp
--- a/Engine/Source/Runtime/GameplayAbilities/Private/Abilities/GameplayAbilityWorldReticle_ActorVisualization.cpp
+++ b/Engine/Source/Runtime/GameplayAbilities/Private/Abilities/GameplayAbilityWorldReticle_ActorVisualization.cpp
@@ -48,6 +48,7 @@ void AGameplayAbilityWorldReticle_ActorVisualization::InitializeReticleVisualiza
 		check(MyRoot);
 
 		TargetingActor = InTargetingActor;
+		AddTickPrerequisiteActor(TargetingActor);		//We want the reticle to tick after the targeting actor so that designers have the final say on the position
 
 		for (UMeshComponent* MeshComp : MeshComps)
 		{
