fix ApplyGameplayEffect*ToOwner not returning proper handle

diff --git a/Engine/Source/Runtime/GameplayAbilities/Private/Abilities/GameplayAbility.cpp b/Engine/Source/Runtime/GameplayAbilities/Private/Abilities/GameplayAbility.cpp
--- a/Engine/Source/Runtime/GameplayAbilities/Private/Abilities/GameplayAbility.cpp
+++ b/Engine/Source/Runtime/GameplayAbilities/Private/Abilities/GameplayAbility.cpp
@@ -882,7 +882,7 @@ FActiveGameplayEffectHandle UGameplayAbility::ApplyGameplayEffectSpecToOwner(con
 {
 	if (SpecHandle.IsValid() && (ActivationInfo.ActivationMode == EGameplayAbilityActivationMode::Authority || ActivationInfo.ActivationMode == EGameplayAbilityActivationMode::Predicting))
 	{
-		ActorInfo->AbilitySystemComponent->ApplyGameplayEffectSpecToSelf(*SpecHandle.Data.Get(), ActivationInfo.GetPredictionKeyForNewAction());
+		return ActorInfo->AbilitySystemComponent->ApplyGameplayEffectSpecToSelf(*SpecHandle.Data.Get(), ActivationInfo.GetPredictionKeyForNewAction());
 
 	}
 	return FActiveGameplayEffectHandle();
