Prediction fixes
-Fix cases of predicting duration based GEs

diff --git a/Engine/Source/Runtime/GameplayAbilities/Private/AttributeSet.cpp b/Engine/Source/Runtime/GameplayAbilities/Private/AttributeSet.cpp
--- a/Engine/Source/Runtime/GameplayAbilities/Private/AttributeSet.cpp
+++ b/Engine/Source/Runtime/GameplayAbilities/Private/AttributeSet.cpp
@@ -104,6 +104,16 @@ void UAttributeSet::PrintDebug()
 	
 }
 
+void UAttributeSet::PreNetReceive()
+{
+	FScopedAggregatorOnDirtyBatch::BeginLock();
+}
+	
+void UAttributeSet::PostNetReceive()
+{
+	FScopedAggregatorOnDirtyBatch::EndLock(true);
+}
+
 FAttributeMetaData::FAttributeMetaData()
 	: MinValue(0.f)
 	, MaxValue(1.f)
@@ -329,4 +339,4 @@ void FAttributeSetInitter::InitAttributeSetDefaults(UAbilitySystemComponent* Abi
 			}
 		}		
 	}
-}
\ No newline at end of file
+}
diff --git a/Engine/Source/Runtime/GameplayAbilities/Private/GameplayEffect.cpp b/Engine/Source/Runtime/GameplayAbilities/Private/GameplayEffect.cpp
--- a/Engine/Source/Runtime/GameplayAbilities/Private/GameplayEffect.cpp
+++ b/Engine/Source/Runtime/GameplayAbilities/Private/GameplayEffect.cpp
@@ -1083,7 +1083,6 @@ FAggregatorRef& FActiveGameplayEffectsContainer::FindOrCreateAttributeAggregator
 	return AttributeAggregatorMap.Add(Attribute, FAggregatorRef(NewAttributeAggregator));
 }
 
-// Move this to UAttributeComponent?
 void FActiveGameplayEffectsContainer::OnAttributeAggregatorDirty(FAggregator* Aggregator, FGameplayAttribute Attribute)
 {
 	ABILITY_LOG_SCOPE(TEXT("FActiveGameplayEffectsContainer::OnAttributeAggregatorDirty"));
@@ -1098,19 +1097,38 @@ void FActiveGameplayEffectsContainer::OnAttributeAggregatorDirty(FAggregator* Ag
 
 	if (Owner->IsNetSimulating())
 	{
-		// We are a client. The current value of this attribute is the replicated server's "final" value. We dont actually know what the 
-		// server's base value is. But we can calculate it with ReverseEvaluate(). Then, we can call Evaluate with IncludePredictiveMods=true
-		// to apply our mods and get an accurate predicted value.
-		
-		float FinalValue = Owner->GetNumericAttribute(Attribute);
-		float BaseValue = Aggregator->ReverseEvaluate(FinalValue, EvaluationParameters);
-		Aggregator->SetBaseValue(BaseValue, false);
+		if (FScopedAggregatorOnDirtyBatch::GlobalFromNetworkUpdate && Aggregator->NetUpdateID != FScopedAggregatorOnDirtyBatch::NetUpdateID)
+		{
+			// We are a client. The current value of this attribute is the replicated server's "final" value. We dont actually know what the 
+			// server's base value is. But we can calculate it with ReverseEvaluate(). Then, we can call Evaluate with IncludePredictiveMods=true
+			// to apply our mods and get an accurate predicted value.
+			// 
+			// It is very important that we only do this exactly one time when we get a new value from the server. Once we set the new local value for this
+			// attribute below, recalculating the base would give us the wrong server value. We should only do this when we are coming directly from a network update.
+			// 
+			// Unfortunately there are two ways we could get here from a network update: from the ActiveGameplayEffect container being updated or from a traditional
+			// OnRep on the actual attribute uproperty. Both of these could happen in a single network update, or potentially only one could happen
+			// (and in fact it could be either one! the AGE container could change in a way that doesnt change the final attribute value, or we could have the base value
+			// of the attribute actually be modified (e.g,. losing health or mana which only results in an OnRep and not in a AGE being applied).
+			// 
+			// So both paths need to lead to this function, but we should only do it one time per update. Once we update the base value, we need to make sure we dont do it again
+			// until we get a new network update. GlobalFromNetworkUpdate and NetUpdateID are what do this.
+			// 
+			// GlobalFromNetworkUpdate - only set to true when we are coming from an OnRep or when we are coming from an ActiveGameplayEffect container net update.
+			// NetUpdateID - updated once whenever an AttributeSet is received over the network. It will be incremented one time per actor that gets an update.
+					
+
+			float FinalValue = Owner->GetNumericAttribute(Attribute);
+			float BaseValue = Aggregator->ReverseEvaluate(FinalValue, EvaluationParameters);
+			Aggregator->SetBaseValue(BaseValue, false);
+			Aggregator->NetUpdateID = FScopedAggregatorOnDirtyBatch::NetUpdateID;
+
+			ABILITY_LOG(Log, TEXT("Reverse Evaluated %s. FinalValue: %.2f  BaseValue: %.2f "), *Attribute.GetName(), FinalValue, BaseValue);
+		}
 
 		EvaluationParameters.IncludePredictiveMods = true;
-		ABILITY_LOG(Log, TEXT("Reverse Evaluated %s. FinalValue: %.2f  BaseValue: %.2f "), *Attribute.GetName(), FinalValue, BaseValue);
 	}
 	
-	
 	float NewValue = Aggregator->Evaluate(EvaluationParameters);
 
 	if (EvaluationParameters.IncludePredictiveMods)
@@ -1186,10 +1204,10 @@ void FActiveGameplayEffectsContainer::SetBaseAttributeValueFromReplication(FGame
 	if (RefPtr && RefPtr->Get())
 	{
 		FAggregator* Aggregator =  RefPtr->Get();
-		if (Aggregator->HasPredictedMods())
-		{
-			OnAttributeAggregatorDirty(Aggregator, Attribute);
-		}
+		
+		FScopedAggregatorOnDirtyBatch::GlobalFromNetworkUpdate = true;
+		OnAttributeAggregatorDirty(Aggregator, Attribute);
+		FScopedAggregatorOnDirtyBatch::GlobalFromNetworkUpdate = false;
 	}
 }
 
diff --git a/Engine/Source/Runtime/GameplayAbilities/Private/GameplayEffectAggregator.cpp b/Engine/Source/Runtime/GameplayAbilities/Private/GameplayEffectAggregator.cpp
--- a/Engine/Source/Runtime/GameplayAbilities/Private/GameplayEffectAggregator.cpp
+++ b/Engine/Source/Runtime/GameplayAbilities/Private/GameplayEffectAggregator.cpp
@@ -293,22 +293,37 @@ void FAggregatorRef::TakeSnapshotOf(const FAggregatorRef& RefToSnapshot)
 
 int32 FScopedAggregatorOnDirtyBatch::GlobalBatchCount = 0;
 TSet<FAggregator*> FScopedAggregatorOnDirtyBatch::DirtyAggregators;
+bool FScopedAggregatorOnDirtyBatch::GlobalFromNetworkUpdate = false;
+int32 FScopedAggregatorOnDirtyBatch::NetUpdateID = 1;
 
 FScopedAggregatorOnDirtyBatch::FScopedAggregatorOnDirtyBatch()
 {
-	GlobalBatchCount++;
+	BeginLock();
 }
 
 FScopedAggregatorOnDirtyBatch::~FScopedAggregatorOnDirtyBatch()
+{
+	EndLock(false);
+}
+
+void FScopedAggregatorOnDirtyBatch::BeginLock()
+{
+	GlobalBatchCount++;
+}
+void FScopedAggregatorOnDirtyBatch::EndLock(bool FromNetworkUpdate)
 {
 	GlobalBatchCount--;
 	if (GlobalBatchCount == 0)
 	{
+		if (FromNetworkUpdate)
+			NetUpdateID++;
+
+		GlobalFromNetworkUpdate = FromNetworkUpdate;
 		for (FAggregator* Agg : DirtyAggregators)
 		{
 			Agg->BroadcastOnDirty();
-
 		}
 		DirtyAggregators.Empty();
+		GlobalFromNetworkUpdate = false;
 	}
 }
\ No newline at end of file
diff --git a/Engine/Source/Runtime/GameplayAbilities/Public/AttributeSet.h b/Engine/Source/Runtime/GameplayAbilities/Public/AttributeSet.h
--- a/Engine/Source/Runtime/GameplayAbilities/Public/AttributeSet.h
+++ b/Engine/Source/Runtime/GameplayAbilities/Public/AttributeSet.h
@@ -131,6 +131,10 @@ public:
 
 	virtual void PrintDebug();
 
+	virtual void PreNetReceive();
+	
+	virtual void PostNetReceive();
+
 protected:
 	/** Is this attribute set safe to ID over the network by name?  */
 	uint32 bNetAddressable : 1;
diff --git a/Engine/Source/Runtime/GameplayAbilities/Public/GameplayEffectAggregator.h b/Engine/Source/Runtime/GameplayAbilities/Public/GameplayEffectAggregator.h
--- a/Engine/Source/Runtime/GameplayAbilities/Public/GameplayEffectAggregator.h
+++ b/Engine/Source/Runtime/GameplayAbilities/Public/GameplayEffectAggregator.h
@@ -38,7 +38,7 @@ struct GAMEPLAYABILITIES_API FAggregator : public TSharedFromThis<FAggregator>
 {
 	DECLARE_MULTICAST_DELEGATE_OneParam(FOnAggregatorDirty, FAggregator*);
 
-	FAggregator(float InBaseValue=0.f) : BaseValue(InBaseValue), IsBroadcastingDirty(false), NumPredictiveMods(0) { }
+	FAggregator(float InBaseValue=0.f) : BaseValue(InBaseValue), NetUpdateID(0), IsBroadcastingDirty(false), NumPredictiveMods(0) { }
 
 	/** Simple accessor to base value */
 	float GetBaseValue() const;
@@ -70,6 +70,9 @@ struct GAMEPLAYABILITIES_API FAggregator : public TSharedFromThis<FAggregator>
 	void AddDependant(FActiveGameplayEffectHandle Handle);
 
 	bool HasPredictedMods() const;
+	
+	/** NetworkID that we had our last update from. Will only be set on clients and is only used to ensure we dont recompute server base value twice. */
+	int32 NetUpdateID;
 
 private:
 
@@ -114,6 +117,12 @@ struct GAMEPLAYABILITIES_API FScopedAggregatorOnDirtyBatch
 	FScopedAggregatorOnDirtyBatch();
 	~FScopedAggregatorOnDirtyBatch();
 
+	static void BeginLock();
+	static void EndLock(bool FromNetworkUpdate);
+
 	static int32	GlobalBatchCount;
 	static TSet<FAggregator*>	DirtyAggregators;
+
+	static bool		GlobalFromNetworkUpdate;
+	static int32	NetUpdateID;
 };
\ No newline at end of file
