Fixed Clang compile error

diff --git a/Engine/Source/Runtime/GameplayAbilities/Public/GameplayEffectAggregator.h b/Engine/Source/Runtime/GameplayAbilities/Public/GameplayEffectAggregator.h
--- a/Engine/Source/Runtime/GameplayAbilities/Public/GameplayEffectAggregator.h
+++ b/Engine/Source/Runtime/GameplayAbilities/Public/GameplayEffectAggregator.h
@@ -38,7 +38,7 @@ struct GAMEPLAYABILITIES_API FAggregator : public TSharedFromThis<FAggregator>
 {
 	DECLARE_MULTICAST_DELEGATE_OneParam(FOnAggregatorDirty, FAggregator*);
 
-	FAggregator(float InBaseValue=0.f) : BaseValue(InBaseValue), NetUpdateID(0), IsBroadcastingDirty(false), NumPredictiveMods(0) { }
+	FAggregator(float InBaseValue=0.f) : NetUpdateID(0), BaseValue(InBaseValue), IsBroadcastingDirty(false), NumPredictiveMods(0) { }
 
 	/** Simple accessor to base value */
 	float GetBaseValue() const;
