#GamplayEffects: Added ctors to init FGameplayEffectModifierMagnitude via code, and adjusted Automated Tests to use that so it passes now.

diff --git a/Engine/Source/Runtime/GameplayAbilities/Private/GameplayEffectTests.cpp b/Engine/Source/Runtime/GameplayAbilities/Private/GameplayEffectTests.cpp
--- a/Engine/Source/Runtime/GameplayAbilities/Private/GameplayEffectTests.cpp
+++ b/Engine/Source/Runtime/GameplayAbilities/Private/GameplayEffectTests.cpp
@@ -92,14 +92,14 @@ bool GameplayEffectsTest_InstantDamage(UWorld *World, FAutomationTestBase * Test
 
 		UGameplayEffect * BaseDmgEffect = Cast<UGameplayEffect>(StaticConstructObject(UGameplayEffect::StaticClass(), GetTransientPackage(), FName(TEXT("BaseDmgEffect"))));
 		BaseDmgEffect->Modifiers.SetNum(1);
-		BaseDmgEffect->Modifiers[0].Magnitude.SetValue(DamageValue);
+		BaseDmgEffect->Modifiers[0].ModifierMagnitude = FScalableFloat(DamageValue);
 		BaseDmgEffect->Modifiers[0].ModifierOp = EGameplayModOp::Additive;
 		BaseDmgEffect->Modifiers[0].Attribute.SetUProperty(HealthProperty);
 		BaseDmgEffect->Duration.Value = UGameplayEffect::INSTANT_APPLICATION;
 
 		SourceComponent->ApplyGameplayEffectToTarget(BaseDmgEffect, DestComponent, 1.f);
 		
-		Test->TestTrue(SKILL_TEST_TEXT("Basic Instant Damage Applied"), (DestComponent->GetSet<UAbilitySystemTestAttributeSet>()->Health == (StartHealth + DamageValue)));
+		Test->TestEqual(SKILL_TEST_TEXT("Basic Instant Damage Applied"), DestComponent->GetSet<UAbilitySystemTestAttributeSet>()->Health, StartHealth + DamageValue);
 		ABILITY_LOG(Log, TEXT("Final Health: %.2f"), DestComponent->GetSet<UAbilitySystemTestAttributeSet>()->Health);
 	}
 
diff --git a/Engine/Source/Runtime/GameplayAbilities/Public/GameplayEffect.h b/Engine/Source/Runtime/GameplayAbilities/Public/GameplayEffect.h
--- a/Engine/Source/Runtime/GameplayAbilities/Public/GameplayEffect.h
+++ b/Engine/Source/Runtime/GameplayAbilities/Public/GameplayEffect.h
@@ -178,15 +178,34 @@ struct FGameplayEffectModifierMagnitude
 
 public:
 
-	/** Constructor */
+	/** Default Constructor */
 	FGameplayEffectModifierMagnitude()
 		: MagnitudeCalculationType(EGameplayEffectMagnitudeCalculation::ScalableFloat)
-		, ScalableFloatMagnitude()
-		, AttributeBasedMagnitude()
-		, CustomMagnitude()
 	{
 	}
 
+	/** Constructors for setting value in code (for automation tests) */
+	FGameplayEffectModifierMagnitude(const FScalableFloat& Value)
+		: MagnitudeCalculationType(EGameplayEffectMagnitudeCalculation::ScalableFloat)
+		, ScalableFloatMagnitude(Value)
+	{
+	}
+	FGameplayEffectModifierMagnitude(const FAttributeBasedFloat& Value)
+		: MagnitudeCalculationType(EGameplayEffectMagnitudeCalculation::AttributeBased)
+		, AttributeBasedMagnitude(Value)
+	{
+	}
+	FGameplayEffectModifierMagnitude(const FCustomCalculationBasedFloat& Value)
+		: MagnitudeCalculationType(EGameplayEffectMagnitudeCalculation::CustomCalculationClass)
+		, CustomMagnitude(Value)
+	{
+	}
+	FGameplayEffectModifierMagnitude(const FSetByCallerFloat& Value)
+		: MagnitudeCalculationType(EGameplayEffectMagnitudeCalculation::SetByCaller)
+		, SetByCallerMagnitude(Value)
+	{
+	}
+ 
 	/**
 	 * Determines if the magnitude can be properly calculated with the specified gameplay effect spec (could fail if relying on an attribute not present, etc.)
 	 * 
