#GameplayEffects #Tests: Refactored test slightly o pull out boilerplate code in prep for adding more tests

diff --git a/Engine/Source/Runtime/GameplayAbilities/Private/GameplayEffectTests.cpp b/Engine/Source/Runtime/GameplayAbilities/Private/GameplayEffectTests.cpp
--- a/Engine/Source/Runtime/GameplayAbilities/Private/GameplayEffectTests.cpp
+++ b/Engine/Source/Runtime/GameplayAbilities/Private/GameplayEffectTests.cpp
@@ -72,42 +72,90 @@ static UDataTable* CreateGameplayDataTable()
 	return DataTable;
 }
 
-bool GameplayEffectsTest_InstantDamage(UWorld *World, FAutomationTestBase * Test)
+#define GET_FIELD_CHECKED(Class, Field) FindFieldChecked<UProperty>(Class::StaticClass(), GET_MEMBER_NAME_CHECKED(Class, Field))
+#define CONSTRUCT_CLASS(Class, Name) Class * Name = Cast<Class>(StaticConstructObject(Class::StaticClass(), GetTransientPackage(), FName(TEXT(#Name))))
+
+class GameplayEffectsTestSuite
 {
-	const float StartHealth = 100.f;
-	const float DamageValue = -5.f;
+public:
+	GameplayEffectsTestSuite(UWorld* WorldIn, FAutomationTestBase* TestIn)
+	: World(WorldIn)
+	, Test(TestIn)
+	{
+		// run before each test
 
-	AAbilitySystemTestPawn *SourceActor = World->SpawnActor<AAbilitySystemTestPawn>();
-	AAbilitySystemTestPawn *DestActor = World->SpawnActor<AAbilitySystemTestPawn>();
+		const float StartingHealth = 100.f;
 
-	UProperty *HealthProperty = FindFieldChecked<UProperty>(UAbilitySystemTestAttributeSet::StaticClass(), GET_MEMBER_NAME_CHECKED(UAbilitySystemTestAttributeSet, Health));
+		// set up the source actor
+		SourceActor = World->SpawnActor<AAbilitySystemTestPawn>();
+		SourceComponent = SourceActor->GetAbilitySystemComponent();
+		SourceComponent->GetSet<UAbilitySystemTestAttributeSet>()->Health = StartingHealth;
 
-	UAbilitySystemComponent * SourceComponent = SourceActor->GetAbilitySystemComponent();
-	UAbilitySystemComponent * DestComponent = DestActor->GetAbilitySystemComponent();
-	SourceComponent->GetSet<UAbilitySystemTestAttributeSet>()->Health = StartHealth;
-	DestComponent->GetSet<UAbilitySystemTestAttributeSet>()->Health = StartHealth;
+		// set up the destination actor
+		DestActor = World->SpawnActor<AAbilitySystemTestPawn>();
+		DestComponent = DestActor->GetAbilitySystemComponent();
+		DestComponent->GetSet<UAbilitySystemTestAttributeSet>()->Health = StartingHealth;
+	}
 
+	~GameplayEffectsTestSuite()
 	{
-		ABILITY_LOG_SCOPE(TEXT("Apply InstantDamage"))
+		// run after each test
+
+		// destroy the actors
+		if (SourceActor)
+		{
+			World->EditorDestroyActor(SourceActor, false);
+		}
+		if (DestActor)
+		{
+			World->EditorDestroyActor(DestActor, false);
+		}
+	}
+
+public: // the tests
 
-		UGameplayEffect * BaseDmgEffect = Cast<UGameplayEffect>(StaticConstructObject(UGameplayEffect::StaticClass(), GetTransientPackage(), FName(TEXT("BaseDmgEffect"))));
+	void Test_InstantDamage()
+	{
+		ABILITY_LOG_SCOPE(TEXT("Apply InstantDamage"));
+
+		const float DamageValue = -5.f;
+		const float StartingHealth = DestComponent->GetSet<UAbilitySystemTestAttributeSet>()->Health;
+
+		CONSTRUCT_CLASS(UGameplayEffect, BaseDmgEffect);
 		BaseDmgEffect->Modifiers.SetNum(1);
 		BaseDmgEffect->Modifiers[0].ModifierMagnitude = FScalableFloat(DamageValue);
 		BaseDmgEffect->Modifiers[0].ModifierOp = EGameplayModOp::Additive;
-		BaseDmgEffect->Modifiers[0].Attribute.SetUProperty(HealthProperty);
+		BaseDmgEffect->Modifiers[0].Attribute.SetUProperty(GET_FIELD_CHECKED(UAbilitySystemTestAttributeSet, Health));
 		BaseDmgEffect->Duration.Value = UGameplayEffect::INSTANT_APPLICATION;
 
 		SourceComponent->ApplyGameplayEffectToTarget(BaseDmgEffect, DestComponent, 1.f);
-		
-		Test->TestEqual(SKILL_TEST_TEXT("Basic Instant Damage Applied"), DestComponent->GetSet<UAbilitySystemTestAttributeSet>()->Health, StartHealth + DamageValue);
-		ABILITY_LOG(Log, TEXT("Final Health: %.2f"), DestComponent->GetSet<UAbilitySystemTestAttributeSet>()->Health);
+
+		Test->TestEqual(SKILL_TEST_TEXT("Basic Instant Damage Applied"), DestComponent->GetSet<UAbilitySystemTestAttributeSet>()->Health, StartingHealth + DamageValue);
 	}
 
-	World->EditorDestroyActor(SourceActor, false);
-	World->EditorDestroyActor(DestActor, false);
+public:
 
-	return true;
-}
+	// used to run tests consistently
+	static void RunTest(UWorld* WorldIn, FAutomationTestBase* TestIn, void (GameplayEffectsTestSuite::*TestFunc)())
+	{
+		uint64 InitialFrameCounter = GFrameCounter;
+		{
+			GameplayEffectsTestSuite Tester(WorldIn, TestIn);
+			(Tester.*TestFunc)();
+		}
+		GFrameCounter = InitialFrameCounter;
+	}
+
+private:
+	UWorld* World;
+	FAutomationTestBase* Test;
+
+	AAbilitySystemTestPawn* SourceActor;
+	UAbilitySystemComponent* SourceComponent;
+
+	AAbilitySystemTestPawn* DestActor;
+	UAbilitySystemComponent* DestComponent;
+};
 
 #endif //WITH_EDITOR
 
@@ -131,8 +179,7 @@ bool FGameplayEffectsTest::RunTest( const FString& Parameters )
 	World->InitializeActorsForPlay(URL);
 	World->BeginPlay();
 	
-	GameplayEffectsTest_InstantDamage(World, this);
-	
+	GameplayEffectsTestSuite::RunTest(World, this, &GameplayEffectsTestSuite::Test_InstantDamage);
 
 	GEngine->DestroyWorldContext(World);
 	World->DestroyWorld(false);
