PredictionKey fixes
-Make ScopedPrediction key live on the AbilitySystemComponent. It is now 'global' for that ASC and only ever valid during a function scope.
-ActivationInfo only keeps track of the prediction key it was activated on.
-Callsites that initiaite predictive actions now check if they are authority or if they have a valid prediction key.

diff --git a/Engine/Source/Runtime/GameplayAbilities/Private/Abilities/GameplayAbility.cpp b/Engine/Source/Runtime/GameplayAbilities/Private/Abilities/GameplayAbility.cpp
--- a/Engine/Source/Runtime/GameplayAbilities/Private/Abilities/GameplayAbility.cpp
+++ b/Engine/Source/Runtime/GameplayAbilities/Private/Abilities/GameplayAbility.cpp
@@ -75,19 +75,11 @@ bool UGameplayAbility::CallRemoteFunction(UFunction* Function, void* Parameters,
 // TODO: polymorphic payload
 void UGameplayAbility::SendGameplayEvent(FGameplayTag EventTag, FGameplayEventData Payload)
 {
-	if (NetExecutionPolicy == EGameplayAbilityNetExecutionPolicy::Predictive)
+	UAbilitySystemComponent* AbilitySystemComponent = CurrentActorInfo->AbilitySystemComponent.Get();
+	if (ensure(AbilitySystemComponent))
 	{
-		Payload.PredictionKey = CurrentActivationInfo.GetPredictionKeyForNewAction();
-	}
-
-	AActor *OwnerActor = Cast<AActor>(GetOuter());
-	if (OwnerActor)
-	{
-		UAbilitySystemComponent* AbilitySystemComponent = OwnerActor->FindComponentByClass<UAbilitySystemComponent>();
-		if (ensure(AbilitySystemComponent))
-		{
-			AbilitySystemComponent->HandleGameplayEvent(EventTag, &Payload);
-		}
+		FScopedPredictionWindow NewScopedWindow(AbilitySystemComponent, true);
+		AbilitySystemComponent->HandleGameplayEvent(EventTag, &Payload);
 	}
 }
 
@@ -740,19 +732,19 @@ void UGameplayAbility::TaskEnded(UAbilityTask* Task)
 void UGameplayAbility::K2_ExecuteGameplayCue(FGameplayTag GameplayCueTag, FGameplayEffectContextHandle Context)
 {
 	check(CurrentActorInfo);
-	CurrentActorInfo->AbilitySystemComponent->ExecuteGameplayCue(GameplayCueTag, CurrentActivationInfo.GetPredictionKeyForNewAction(), Context);
+	CurrentActorInfo->AbilitySystemComponent->ExecuteGameplayCue(GameplayCueTag, Context);
 }
 
 void UGameplayAbility::K2_AddGameplayCue(FGameplayTag GameplayCueTag, FGameplayEffectContextHandle Context)
 {
 	check(CurrentActorInfo);
-	CurrentActorInfo->AbilitySystemComponent->AddGameplayCue(GameplayCueTag, CurrentActivationInfo.GetPredictionKeyForNewAction(), Context);
+	CurrentActorInfo->AbilitySystemComponent->AddGameplayCue(GameplayCueTag, Context);
 }
 
 void UGameplayAbility::K2_RemoveGameplayCue(FGameplayTag GameplayCueTag)
 {
 	check(CurrentActorInfo);
-	CurrentActorInfo->AbilitySystemComponent->RemoveGameplayCue(GameplayCueTag, CurrentActivationInfo.GetPredictionKeyForNewAction());
+	CurrentActorInfo->AbilitySystemComponent->RemoveGameplayCue(GameplayCueTag);
 }
 
 FGameplayEffectContextHandle UGameplayAbility::GetContextFromOwner(FGameplayAbilityTargetDataHandle OptionalTargetData) const
@@ -830,6 +822,11 @@ bool UGameplayAbility::IsTriggered() const
 	return AbilityTriggers.Num() > 0;
 }
 
+bool UGameplayAbility::HasAuthorityOrPredictionKey(const FGameplayAbilityActorInfo* ActorInfo, const FGameplayAbilityActivationInfo* ActivationInfo) const
+{
+	return ActorInfo->AbilitySystemComponent->HasAuthorityOrPredictionKey(ActivationInfo);
+}
+
 // -------------------------------------------------------
 
 FActiveGameplayEffectHandle UGameplayAbility::BP_ApplyGameplayEffectToOwner(TSubclassOf<UGameplayEffect> GameplayEffectClass, int32 GameplayEffectLevel)
@@ -860,7 +857,7 @@ FActiveGameplayEffectHandle UGameplayAbility::K2_ApplyGameplayEffectToOwner(cons
 
 FActiveGameplayEffectHandle UGameplayAbility::ApplyGameplayEffectToOwner(const FGameplayAbilitySpecHandle Handle, const FGameplayAbilityActorInfo* ActorInfo, const FGameplayAbilityActivationInfo ActivationInfo, const UGameplayEffect* GameplayEffect, float GameplayEffectLevel)
 {
-	if (GameplayEffect && (ActivationInfo.ActivationMode == EGameplayAbilityActivationMode::Authority || ActivationInfo.ActivationMode == EGameplayAbilityActivationMode::Predicting))
+	if (GameplayEffect && (HasAuthorityOrPredictionKey(ActorInfo, &ActivationInfo)))
 	{
 		FGameplayEffectSpecHandle SpecHandle = MakeOutgoingGameplayEffectSpec(Handle, ActorInfo, ActivationInfo, GameplayEffect->GetClass(), GameplayEffectLevel);
 		if (SpecHandle.IsValid())
@@ -880,9 +877,9 @@ FActiveGameplayEffectHandle UGameplayAbility::K2_ApplyGameplayEffectSpecToOwner(
 
 FActiveGameplayEffectHandle UGameplayAbility::ApplyGameplayEffectSpecToOwner(const FGameplayAbilitySpecHandle AbilityHandle, const FGameplayAbilityActorInfo* ActorInfo, const FGameplayAbilityActivationInfo ActivationInfo, const FGameplayEffectSpecHandle SpecHandle)
 {
-	if (SpecHandle.IsValid() && (ActivationInfo.ActivationMode == EGameplayAbilityActivationMode::Authority || ActivationInfo.ActivationMode == EGameplayAbilityActivationMode::Predicting))
+	if (SpecHandle.IsValid() && (HasAuthorityOrPredictionKey(ActorInfo, &ActivationInfo)))
 	{
-		return ActorInfo->AbilitySystemComponent->ApplyGameplayEffectSpecToSelf(*SpecHandle.Data.Get(), ActivationInfo.GetPredictionKeyForNewAction());
+		return ActorInfo->AbilitySystemComponent->ApplyGameplayEffectSpecToSelf(*SpecHandle.Data.Get(), ActorInfo->AbilitySystemComponent->GetPredictionKeyForNewAction());
 
 	}
 	return FActiveGameplayEffectHandle();
@@ -922,7 +919,7 @@ FActiveGameplayEffectHandle UGameplayAbility::ApplyGameplayEffectToTarget(const
 
 	FActiveGameplayEffectHandle EffectHandle;
 
-	if (ActivationInfo.ActivationMode == EGameplayAbilityActivationMode::Authority || ActivationInfo.ActivationMode == EGameplayAbilityActivationMode::Predicting)
+	if (HasAuthorityOrPredictionKey(ActorInfo, &ActivationInfo))
 	{
 		FGameplayEffectSpecHandle SpecHandle = MakeOutgoingGameplayEffectSpec(Handle, ActorInfo, ActivationInfo, GameplayEffectClass, GameplayEffectLevel);
 		EffectHandle = ApplyGameplayEffectSpecToTarget(Handle, ActorInfo, ActivationInfo, SpecHandle, Target);
@@ -939,11 +936,12 @@ FActiveGameplayEffectHandle UGameplayAbility::K2_ApplyGameplayEffectSpecToTarget
 FActiveGameplayEffectHandle UGameplayAbility::ApplyGameplayEffectSpecToTarget(const FGameplayAbilitySpecHandle AbilityHandle, const FGameplayAbilityActorInfo* ActorInfo, const FGameplayAbilityActivationInfo ActivationInfo, const FGameplayEffectSpecHandle SpecHandle, FGameplayAbilityTargetDataHandle TargetData)
 {
 	FActiveGameplayEffectHandle EffectHandle;
-	if (SpecHandle.IsValid() && (ActivationInfo.ActivationMode == EGameplayAbilityActivationMode::Authority || ActivationInfo.ActivationMode == EGameplayAbilityActivationMode::Predicting))
+	
+	if (SpecHandle.IsValid() && HasAuthorityOrPredictionKey(ActorInfo, &ActivationInfo))
 	{
 		for (auto Data : TargetData.Data)
 		{
-			TArray<FActiveGameplayEffectHandle> EffectHandles = Data->ApplyGameplayEffectSpec(*SpecHandle.Data.Get(), ActivationInfo.GetPredictionKeyForNewAction());
+			TArray<FActiveGameplayEffectHandle> EffectHandles = Data->ApplyGameplayEffectSpec(*SpecHandle.Data.Get(), ActorInfo->AbilitySystemComponent->GetPredictionKeyForNewAction());
 			if (EffectHandles.Num() > 0)
 			{
 				EffectHandle = EffectHandles[0];
diff --git a/Engine/Source/Runtime/GameplayAbilities/Private/Abilities/GameplayAbility_CharacterJump.cpp b/Engine/Source/Runtime/GameplayAbilities/Private/Abilities/GameplayAbility_CharacterJump.cpp
--- a/Engine/Source/Runtime/GameplayAbilities/Private/Abilities/GameplayAbility_CharacterJump.cpp
+++ b/Engine/Source/Runtime/GameplayAbilities/Private/Abilities/GameplayAbility_CharacterJump.cpp
@@ -18,8 +18,8 @@ UGameplayAbility_CharacterJump::UGameplayAbility_CharacterJump(const FObjectInit
 
 void UGameplayAbility_CharacterJump::ActivateAbility(const FGameplayAbilitySpecHandle Handle, const FGameplayAbilityActorInfo* ActorInfo, const FGameplayAbilityActivationInfo ActivationInfo)
 {
-	if (ActivationInfo.ActivationMode == EGameplayAbilityActivationMode::Authority ||
-		ActivationInfo.ActivationMode == EGameplayAbilityActivationMode::Predicting)
+	
+	if (HasAuthorityOrPredictionKey(ActorInfo, &ActivationInfo))
 	{
 		if (!CommitAbility(Handle, ActorInfo, ActivationInfo))
 		{
diff --git a/Engine/Source/Runtime/GameplayAbilities/Private/Abilities/Tasks/AbilityTask_WaitTargetData.cpp b/Engine/Source/Runtime/GameplayAbilities/Private/Abilities/Tasks/AbilityTask_WaitTargetData.cpp
--- a/Engine/Source/Runtime/GameplayAbilities/Private/Abilities/Tasks/AbilityTask_WaitTargetData.cpp
+++ b/Engine/Source/Runtime/GameplayAbilities/Private/Abilities/Tasks/AbilityTask_WaitTargetData.cpp
@@ -162,8 +162,6 @@ void UAbilityTask_WaitTargetData::FinishSpawningActor(UObject* WorldContextObjec
 void UAbilityTask_WaitTargetData::OnTargetDataReplicatedCallback(FGameplayAbilityTargetDataHandle Data)
 {
 	check(AbilitySystemComponent.IsValid());
-	
-	FScopedPredictionWindow	ScopedPrediction(Ability.Get());
 
 	/** 
 	 *  Call into the TargetActor to sanitize/verify the data. If this returns false, we are rejecting
@@ -204,11 +202,11 @@ void UAbilityTask_WaitTargetData::OnTargetDataReadyCallback(FGameplayAbilityTarg
 {
 	check(AbilitySystemComponent.IsValid());
 
-	FScopedPredictionWindow	ScopedPrediction(Ability.Get(), ShouldReplicateDataToServer());
+	FScopedPredictionWindow	ScopedPrediction(AbilitySystemComponent.Get(), ShouldReplicateDataToServer());
 	
 	if (ShouldReplicateDataToServer())
 	{
-		AbilitySystemComponent->ServerSetReplicatedTargetData(Data, ScopedPrediction.ScopedPredictionKey);
+		AbilitySystemComponent->ServerSetReplicatedTargetData(Data, AbilitySystemComponent->ScopedPredictionKey);
 	}
 
 	ValidData.Broadcast(Data);
diff --git a/Engine/Source/Runtime/GameplayAbilities/Private/AbilitySystemComponent.cpp b/Engine/Source/Runtime/GameplayAbilities/Private/AbilitySystemComponent.cpp
--- a/Engine/Source/Runtime/GameplayAbilities/Private/AbilitySystemComponent.cpp
+++ b/Engine/Source/Runtime/GameplayAbilities/Private/AbilitySystemComponent.cpp
@@ -718,20 +718,20 @@ void UAbilitySystemComponent::InvokeGameplayCueEvent(const FGameplayTag Gameplay
 	UAbilitySystemGlobals::Get().GetGameplayCueManager()->HandleGameplayCues(ActorAvatar, GameplayCueTag, EventType, CueParameters);
 }
 
-void UAbilitySystemComponent::ExecuteGameplayCue(const FGameplayTag GameplayCueTag, FPredictionKey PredictionKey, FGameplayEffectContextHandle EffectContext)
+void UAbilitySystemComponent::ExecuteGameplayCue(const FGameplayTag GameplayCueTag, FGameplayEffectContextHandle EffectContext)
 {
 	if (IsOwnerActorAuthoritative())
 	{
 		ForceReplication();
-		NetMulticast_InvokeGameplayCueExecuted(GameplayCueTag, PredictionKey, EffectContext);
+		NetMulticast_InvokeGameplayCueExecuted(GameplayCueTag, ScopedPredictionKey, EffectContext);
 	}
-	else if (PredictionKey.IsValidKey())
+	else if (ScopedPredictionKey.IsValidKey())
 	{
 		InvokeGameplayCueEvent(GameplayCueTag, EGameplayCueEvent::Executed, EffectContext);
 	}
 }
 
-void UAbilitySystemComponent::AddGameplayCue(const FGameplayTag GameplayCueTag, FPredictionKey PredictionKey, FGameplayEffectContextHandle EffectContext)
+void UAbilitySystemComponent::AddGameplayCue(const FGameplayTag GameplayCueTag, FGameplayEffectContextHandle EffectContext)
 {
 	if (IsOwnerActorAuthoritative())
 	{
@@ -739,7 +739,7 @@ void UAbilitySystemComponent::AddGameplayCue(const FGameplayTag GameplayCueTag,
 
 		ForceReplication();
 		ActiveGameplayCues.AddCue(GameplayCueTag);
-		NetMulticast_InvokeGameplayCueAdded(GameplayCueTag, PredictionKey, EffectContext);
+		NetMulticast_InvokeGameplayCueAdded(GameplayCueTag, ScopedPredictionKey, EffectContext);
 
 		if (!bWasInList)
 		{
@@ -747,7 +747,7 @@ void UAbilitySystemComponent::AddGameplayCue(const FGameplayTag GameplayCueTag,
 			InvokeGameplayCueEvent(GameplayCueTag, EGameplayCueEvent::WhileActive);
 		}
 	}
-	else if (PredictionKey.IsValidKey())
+	else if (ScopedPredictionKey.IsValidKey())
 	{
 		// Allow for predictive gameplaycue events? Needs more thought
 		InvokeGameplayCueEvent(GameplayCueTag, EGameplayCueEvent::OnActive, EffectContext);
@@ -755,16 +755,16 @@ void UAbilitySystemComponent::AddGameplayCue(const FGameplayTag GameplayCueTag,
 	}
 }
 
-void UAbilitySystemComponent::RemoveGameplayCue(const FGameplayTag GameplayCueTag, FPredictionKey PredictionKey)
+void UAbilitySystemComponent::RemoveGameplayCue(const FGameplayTag GameplayCueTag)
 {
 	if (IsOwnerActorAuthoritative())
 	{
 		bool bWasInList = HasMatchingGameplayTag(GameplayCueTag);
 
 		ActiveGameplayCues.RemoveCue(GameplayCueTag);
-		NetMulticast_InvokeGameplayCueRemoved(GameplayCueTag, PredictionKey);
+		NetMulticast_InvokeGameplayCueRemoved(GameplayCueTag, ScopedPredictionKey);
 	}
-	else if (PredictionKey.IsValidKey())
+	else if (ScopedPredictionKey.IsValidKey())
 	{
 		InvokeGameplayCueEvent(GameplayCueTag, EGameplayCueEvent::Removed);
 	}
@@ -973,6 +973,11 @@ void UAbilitySystemComponent::OnRep_PredictionKey()
 	FPredictionKeyDelegates::CatchUpTo(ReplicatedPredictionKey.Current);
 }
 
+bool UAbilitySystemComponent::HasAuthorityOrPredictionKey(const FGameplayAbilityActivationInfo* ActivationInfo) const
+{
+	return ((ActivationInfo->ActivationMode == EGameplayAbilityActivationMode::Authority) || CanPredict());
+}
+
 // ---------------------------------------------------------------------------------------
 
 void UAbilitySystemComponent::PrintAllGameplayEffects() const
diff --git a/Engine/Source/Runtime/GameplayAbilities/Private/AbilitySystemComponent_Abilities.cpp b/Engine/Source/Runtime/GameplayAbilities/Private/AbilitySystemComponent_Abilities.cpp
--- a/Engine/Source/Runtime/GameplayAbilities/Private/AbilitySystemComponent_Abilities.cpp
+++ b/Engine/Source/Runtime/GameplayAbilities/Private/AbilitySystemComponent_Abilities.cpp
@@ -558,7 +558,7 @@ bool UAbilitySystemComponent::TryActivateAbility(FGameplayAbilitySpecHandle Hand
 	Spec->ActiveCount++;
 
 	// Setup a fresh ActivationInfo for this AbilitySpec.
-	Spec->ActivationInfo = FGameplayAbilityActivationInfo(ActorInfo->OwnerActor.Get(), InPredictionKey);
+	Spec->ActivationInfo = FGameplayAbilityActivationInfo(ActorInfo->OwnerActor.Get());
 	FGameplayAbilityActivationInfo &ActivationInfo = Spec->ActivationInfo;
 
 	UGameplayAbility* InstancedAbility = Ability;
@@ -588,15 +588,15 @@ bool UAbilitySystemComponent::TryActivateAbility(FGameplayAbilitySpecHandle Hand
 	else if (Ability->GetNetExecutionPolicy() == EGameplayAbilityNetExecutionPolicy::Predictive)
 	{
 		// This execution is now officially EGameplayAbilityActivationMode:Predicting and has a PredictionKey
-		ActivationInfo.GenerateNewPredictionKey();
+		FScopedPredictionWindow ScopedPredictionWindow(this, true);
 
-		FPredictionKey ThisPredictionKey = ActivationInfo.GetPredictionKeyForNewAction();
+		ActivationInfo.SetPredicting(ScopedPredictionKey);
 		
 		// This must be called immediately after GeneratePredictionKey to prevent problems with recursively activating abilities
-		ServerTryActivateAbility(Handle, Spec->InputPressed, ThisPredictionKey);
+		ServerTryActivateAbility(Handle, Spec->InputPressed, ScopedPredictionKey);
 
 		// If this PredictionKey is rejected, we will call OnClientActivateAbilityFailed.
-		ThisPredictionKey.NewRejectedDelegate().BindUObject(this, &UAbilitySystemComponent::OnClientActivateAbilityFailed, Handle, ThisPredictionKey.Current);
+		ScopedPredictionKey.NewRejectedDelegate().BindUObject(this, &UAbilitySystemComponent::OnClientActivateAbilityFailed, Handle, ScopedPredictionKey.Current);
 
 		if (Ability->GetInstancingPolicy() == EGameplayAbilityInstancingPolicy::InstancedPerExecution)
 		{
@@ -623,11 +623,7 @@ bool UAbilitySystemComponent::TryActivateAbility(FGameplayAbilitySpecHandle Hand
 			Ability->CallActivateAbility(Handle, ActorInfo, ActivationInfo, OnGameplayAbilityEndedDelegate);
 		}
 	}
-
-	// Anything after this can no longer be predicted with the prediction key we just generated.
-	// The key we generated and sent in ServerTryActivateAbility will be confirmed/denied after the server run CallActivateAbility.
-	// If game code tries to do something predictively after this (latent call, response to input, etc) it will not be able to use this key.
-	ActivationInfo.SetPredictionStale();
+	
 	InstancedAbility->SetCurrentActivationInfo(ActivationInfo);	// Need to push this to the ability if it was instanced.
 
 	return true;
@@ -651,6 +647,8 @@ void UAbilitySystemComponent::ServerTryActivateAbility_Implementation(FGameplayA
 		return;
 	}
 
+	FScopedPredictionWindow ScopedPredictionWindow(this, PredictionKey);
+
 	UGameplayAbility* AbilityToActivate = Spec->Ability;
 
 	ensure(AbilityToActivate);
@@ -713,14 +711,6 @@ void UAbilitySystemComponent::ServerTryActivateAbility_Implementation(FGameplayA
 	{
 		ClientActivateAbilityFailed(Handle, PredictionKey.Current);
 	}
-
-	// Update our ReplicatedPredictionKey. When the client gets value, he will know his state (actor+all components/subobjects) are up to do date and he can
-	// remove any necessary predictive work.
-	if (PredictionKey.IsValidKey())
-	{
-		ensure(PredictionKey.Current > ReplicatedPredictionKey.Current);
-		ReplicatedPredictionKey = PredictionKey;
-	}
 }
 
 bool UAbilitySystemComponent::ServerTryActivateAbility_Validate(FGameplayAbilitySpecHandle Handle, bool InputPressed, FPredictionKey PredictionKey)
@@ -779,7 +769,7 @@ void UAbilitySystemComponent::OnClientActivateAbilityFailed(FGameplayAbilitySpec
 	TArray<UGameplayAbility*> Instances = Spec->GetAbilityInstances();
 	for (UGameplayAbility* Ability : Instances)
 	{
-		if (Ability->CurrentActivationInfo.GetPredictionKey().Current == PredictionKey)
+		if (ScopedPredictionKey.Current == PredictionKey)
 		{
 			ABILITY_LOG(Warning, TEXT("Ending Ability %s"), *Ability->GetName());
 			Ability->K2_EndAbility();
@@ -821,7 +811,7 @@ void UAbilitySystemComponent::ClientActivateAbilitySucceed_Implementation(FGamep
 			bool found = false;
 			for (UGameplayAbility* LocalAbility : Spec->NonReplicatedInstances)				// Fixme: this has to be updated once predictive abilities can replicate
 			{
-				if (LocalAbility->GetCurrentActivationInfo().GetPredictionKey().Current == PredictionKey)
+				if (LocalAbility->GetCurrentActivationInfo().GetActivationPredictionKey().Current == PredictionKey)
 				{
 					LocalAbility->ConfirmActivateSucceed();
 					found = true;
@@ -873,17 +863,15 @@ void UAbilitySystemComponent::TriggerAbilityFromGameplayEvent(FGameplayAbilitySp
 		int32 ExecutingAbilityIndex = -1;
 
 		// if we're the server and this is coming from a predicted event we should check if the client has already predicted it
-		if (Payload->PredictionKey.Current > 0
+		if (ScopedPredictionKey.IsValidKey()
 			&& Ability->GetNetExecutionPolicy() == EGameplayAbilityNetExecutionPolicy::Predictive
 			&& ActorInfo->OwnerActor->Role == ROLE_Authority)
 		{
 			bool bPendingClientAbilityFound = false;
 			for (auto PendingAbilityInfo : Component.PendingClientAbilities)
 			{
-				if (Payload->PredictionKey.Current == PendingAbilityInfo.PredictionKey.Base && Handle == PendingAbilityInfo.Handle) // found a match
+				if (ScopedPredictionKey.Current == PendingAbilityInfo.PredictionKey.Base && Handle == PendingAbilityInfo.Handle) // found a match
 				{
-					Payload->PredictionKey = PendingAbilityInfo.PredictionKey;
-
 					Component.PendingClientAbilities.RemoveSingleSwap(PendingAbilityInfo);
 					bPendingClientAbilityFound = true;
 					break;
@@ -895,14 +883,14 @@ void UAbilitySystemComponent::TriggerAbilityFromGameplayEvent(FGameplayAbilitySp
 			if (bPendingClientAbilityFound == false)
 			{
 				UAbilitySystemComponent::FExecutingAbilityInfo Info;
-				Info.PredictionKey = Payload->PredictionKey;
+				Info.PredictionKey = ScopedPredictionKey;
 				Info.Handle = Handle;
 
 				ExecutingAbilityIndex = Component.ExecutingServerAbilities.Add(Info);
 			}
 		}
 
-		if (TryActivateAbility(Handle, Payload->PredictionKey))
+		if (TryActivateAbility(Handle, ScopedPredictionKey))
 		{
 			if (ExecutingAbilityIndex >= 0)
 			{
@@ -1055,13 +1043,13 @@ void UAbilitySystemComponent::AbilityInputPressed(int32 InputID)
 							Instance->InputPressed(Spec.Handle, AbilityActorInfo.Get(), Spec.ActivationInfo);
 						}						
 					}
-					FPredictionKey NewKey = FPredictionKey::CreateNewPredictionKey(this);
-					FScopedPredictionWindow ScopedPrediction(this, NewKey);
+
+					FScopedPredictionWindow ScopedPrediction(this, true);
 					//We don't require (AbilityKeyPressCallbacks.IsBound() || AbilityKeyReleaseCallbacks.IsBound())), because we don't necessarily know now if we'll need this data later.
 					if (GetOwnerRole() != ROLE_Authority)
 					{
 						// Tell the server we pressed input.
-						ServerSetReplicatedAbilityKeyState(InputID, true, NewKey);
+						ServerSetReplicatedAbilityKeyState(InputID, true, ScopedPredictionKey);
 					}
 					AbilityKeyPressCallbacks.Broadcast(InputID);
 				}
@@ -1086,12 +1074,12 @@ void UAbilitySystemComponent::AbilityInputReleased(int32 InputID)
 			{
 				AbilitySpecInputReleased(Spec);
 				FPredictionKey NewKey = FPredictionKey::CreateNewPredictionKey(this);
-				FScopedPredictionWindow ScopedPrediction(this, NewKey);
+				FScopedPredictionWindow ScopedPrediction(this, true);
 				//We don't require (AbilityKeyPressCallbacks.IsBound() || AbilityKeyReleaseCallbacks.IsBound())), because we don't necessarily know now if we'll need this data later.
 				if (GetOwnerRole() != ROLE_Authority)
 				{
 					// Tell the server we released input.
-					ServerSetReplicatedAbilityKeyState(InputID, false, NewKey);
+					ServerSetReplicatedAbilityKeyState(InputID, false, ScopedPredictionKey);
 				}
 				AbilityKeyReleaseCallbacks.Broadcast(InputID);
 			}
@@ -1143,13 +1131,12 @@ void UAbilitySystemComponent::AbilitySpecInputReleased(FGameplayAbilitySpec& Spe
 
 void UAbilitySystemComponent::InputConfirm()
 {
-	FPredictionKey NewKey = FPredictionKey::CreateNewPredictionKey(this);
-	FScopedPredictionWindow ScopedPrediction(this, NewKey);
+	FScopedPredictionWindow ScopedPrediction(this, true);
 
 	if (GetOwnerRole() != ROLE_Authority && ConfirmCallbacks.IsBound())
 	{
 		// Tell the server we confirmed input.
-		ServerSetReplicatedConfirm(true, NewKey);
+		ServerSetReplicatedConfirm(true, ScopedPredictionKey);
 	}
 	
 	ConfirmCallbacks.Broadcast();
@@ -1157,13 +1144,12 @@ void UAbilitySystemComponent::InputConfirm()
 
 void UAbilitySystemComponent::InputCancel()
 {
-	FPredictionKey NewKey = FPredictionKey::CreateNewPredictionKey(this);
-	FScopedPredictionWindow ScopedPrediction(this, NewKey);
+	FScopedPredictionWindow ScopedPrediction(this, true);
 
 	if (GetOwnerRole() != ROLE_Authority && CancelCallbacks.IsBound())
 	{
 		// Tell the server we confirmed input.
-		ServerSetReplicatedConfirm(false, NewKey);
+		ServerSetReplicatedConfirm(false, ScopedPredictionKey);
 	}
 
 	CancelCallbacks.Broadcast();
@@ -1275,7 +1261,7 @@ void UAbilitySystemComponent::ServerSetReplicatedAbilityKeyState_Implementation(
 	{
 		AbilityKeyReleaseCallbacks.Broadcast(InputID);
 	}
-}
+} 
 
 bool UAbilitySystemComponent::ServerSetReplicatedAbilityKeyState_Validate(int32 InputID, bool Pressed, FPredictionKey PredictionKey)
 {
@@ -1411,7 +1397,7 @@ float UAbilitySystemComponent::PlayMontage(UGameplayAbility* InAnimatingAbility,
 			else
 			{
 				// If this prediction key is rejected, we need to end the preview
-				FPredictionKey PredictionKey = ActivationInfo.GetPredictionKeyForNewAction();
+				FPredictionKey PredictionKey = GetPredictionKeyForNewAction();
 				if (PredictionKey.IsValidKey())
 				{
 					PredictionKey.NewRejectedDelegate().BindUObject(this, &UAbilitySystemComponent::OnPredictiveMontageRejected, NewAnimMontage);
diff --git a/Engine/Source/Runtime/GameplayAbilities/Private/AttributeSet.cpp b/Engine/Source/Runtime/GameplayAbilities/Private/AttributeSet.cpp
--- a/Engine/Source/Runtime/GameplayAbilities/Private/AttributeSet.cpp
+++ b/Engine/Source/Runtime/GameplayAbilities/Private/AttributeSet.cpp
@@ -106,12 +106,15 @@ void UAttributeSet::PrintDebug()
 
 void UAttributeSet::PreNetReceive()
 {
-	FScopedAggregatorOnDirtyBatch::BeginLock();
+	// During the scope of this entire actor's network update, we need to lock our attribute aggregators.
+	FScopedAggregatorOnDirtyBatch::BeginNetReceiveLock();
 }
 	
 void UAttributeSet::PostNetReceive()
 {
-	FScopedAggregatorOnDirtyBatch::EndLock(true);
+	// Once we are done receiving properties, we can unlock the attribute aggregators and flag them that the 
+	// current property values are from the server.
+	FScopedAggregatorOnDirtyBatch::EndNetReceiveLock();
 }
 
 FAttributeMetaData::FAttributeMetaData()
diff --git a/Engine/Source/Runtime/GameplayAbilities/Private/GameplayAbilityTypes.cpp b/Engine/Source/Runtime/GameplayAbilities/Private/GameplayAbilityTypes.cpp
--- a/Engine/Source/Runtime/GameplayAbilities/Private/GameplayAbilityTypes.cpp
+++ b/Engine/Source/Runtime/GameplayAbilities/Private/GameplayAbilityTypes.cpp
@@ -93,10 +93,10 @@ bool FGameplayAbilityActorInfo::IsNetAuthority() const
 	return false;
 }
 
-void FGameplayAbilityActivationInfo::GenerateNewPredictionKey() const
-{	
-	PredictionKey.GenerateDependantPredictionKey();
+void FGameplayAbilityActivationInfo::SetPredicting(FPredictionKey PredictionKey)
+{
 	ActivationMode = EGameplayAbilityActivationMode::Predicting;
+	PredictionKeyWhenActivated = PredictionKey;
 }
 
 void FGameplayAbilityActivationInfo::SetActivationConfirmed()
@@ -106,11 +106,6 @@ void FGameplayAbilityActivationInfo::SetActivationConfirmed()
 	bCanBeEndedByOtherInstance = true;
 }
 
-void FGameplayAbilityActivationInfo::SetPredictionStale()
-{
-	PredictionKey.bIsStale = true;
-}
-
 bool FGameplayAbilitySpec::IsActive() const
 {
 	return ActiveCount > 0;
diff --git a/Engine/Source/Runtime/GameplayAbilities/Private/GameplayEffectAggregator.cpp b/Engine/Source/Runtime/GameplayAbilities/Private/GameplayEffectAggregator.cpp
--- a/Engine/Source/Runtime/GameplayAbilities/Private/GameplayEffectAggregator.cpp
+++ b/Engine/Source/Runtime/GameplayAbilities/Private/GameplayEffectAggregator.cpp
@@ -303,27 +303,41 @@ FScopedAggregatorOnDirtyBatch::FScopedAggregatorOnDirtyBatch()
 
 FScopedAggregatorOnDirtyBatch::~FScopedAggregatorOnDirtyBatch()
 {
-	EndLock(false);
+	EndLock();
 }
 
 void FScopedAggregatorOnDirtyBatch::BeginLock()
 {
 	GlobalBatchCount++;
 }
-void FScopedAggregatorOnDirtyBatch::EndLock(bool FromNetworkUpdate)
+void FScopedAggregatorOnDirtyBatch::EndLock()
 {
 	GlobalBatchCount--;
 	if (GlobalBatchCount == 0)
 	{
-		if (FromNetworkUpdate)
-			NetUpdateID++;
-
-		GlobalFromNetworkUpdate = FromNetworkUpdate;
 		for (FAggregator* Agg : DirtyAggregators)
 		{
 			Agg->BroadcastOnDirty();
 		}
 		DirtyAggregators.Empty();
+	}
+}
+
+
+void FScopedAggregatorOnDirtyBatch::BeginNetReceiveLock()
+{
+	BeginLock();
+}
+void FScopedAggregatorOnDirtyBatch::EndNetReceiveLock()
+{
+	// The network lock must end the first time it is called.
+	// Subsequent calls to EndNetReceiveLock() should not trigger a full EndLock, only the first one.
+	if (GlobalBatchCount > 0)
+	{
+		GlobalBatchCount = 1;
+		NetUpdateID++;
+		GlobalFromNetworkUpdate = true;
+		EndLock();
 		GlobalFromNetworkUpdate = false;
 	}
 }
\ No newline at end of file
diff --git a/Engine/Source/Runtime/GameplayAbilities/Private/GameplayPrediction.cpp b/Engine/Source/Runtime/GameplayAbilities/Private/GameplayPrediction.cpp
--- a/Engine/Source/Runtime/GameplayAbilities/Private/GameplayPrediction.cpp
+++ b/Engine/Source/Runtime/GameplayAbilities/Private/GameplayPrediction.cpp
@@ -199,74 +199,50 @@ FScopedPredictionWindow::FScopedPredictionWindow(UAbilitySystemComponent* Abilit
 	// This is used to set an already generated prediction key as the current scoped prediction key.
 	// Should be called on the server for logical scopes where a given key is valid. E.g, "client gave me this key, we both are going to run Foo()".
 
+	// If you are hitting this, this FScopedPredictionWindow constructor should only be called from Server RPCs where the client has given us a PredictionKey.
+	ensure(AbilitySystemComponent->IsNetSimulating() == false);
+
 	Owner = AbilitySystemComponent;
-	Owner->ScopedPedictionKey = InPredictionKey;
+	Owner->ScopedPredictionKey = InPredictionKey;
 	ClearScopedPredictionKey = true;
-
+	SetReplicatedPredictionKey = true;
 }
 
-FScopedPredictionWindow::FScopedPredictionWindow(UGameplayAbility* GameplayAbilityInstance, bool bCanGenerateNewKey)
+FScopedPredictionWindow::FScopedPredictionWindow(UAbilitySystemComponent* InAbilitySystemComponent, bool bCanGenerateNewKey)
 {
-	// "Gets" the current/"best" prediction key and sets it as the prediciton key for this given abilities CurrentActivationInfo.
-	//	-On the server, it means we look for the prediction key that was given to us and either stored on the ActivationIfno or on the ScopedPredictionKey.
-	//  -On the client, we generate a new prediction  key if we don't already have one.
-
-	Ability = GameplayAbilityInstance;
-	check(Ability.IsValid());
-
-	Owner = GameplayAbilityInstance->GetCurrentActorInfo()->AbilitySystemComponent.Get();
-	check(Owner.IsValid());
+	// On the server, this will do nothing since he is authorative and doesnt need a prediction key for anything.
+	// On the client, this will generate a new prediction key if bCanGenerateNewKey is true, and we have a invalid prediction key.
 
 	ClearScopedPredictionKey = false;
+	SetReplicatedPredictionKey = false;
+	Owner = InAbilitySystemComponent;
 
-	FGameplayAbilityActivationInfo& ActivationInfo = Ability->GetCurrentActivationInfoRef();
-
-	if (ActivationInfo.ActivationMode == EGameplayAbilityActivationMode::Authority)
+	if (Owner->IsNetSimulating() == false)
 	{
-		// We are the server, if we already have a valid prediction key in our ActivationInfo, just use it.
-		if (ActivationInfo.GetPredictionKey().IsValidForMorePrediction() == false)
-		{
-			// If we don't, then we fall back to the ScopedPredictionKey that would have been explcitly set by another ScopedPredictionKey that used the other constructor.
-			ActivationInfo.SetPredictionKey(Owner->ScopedPedictionKey);
-		}
+		return;
 	}
-	else
+	
+	// InAbilitySystemComponent->GetPredictionKey().IsValidForMorePrediction() == false && 
+	if (bCanGenerateNewKey)
 	{
-		// We are the client, first save off our current activation mode (E.g, if we were previously 'Confirmed' we need to back to 'Predicting' then back to 'Confirmed' after this scope.
-		ClientPrevActivationMode = static_cast<int8>(ActivationInfo.ActivationMode);
-
-		if (Owner->ScopedPedictionKey.IsValidKey())
-		{
-			// We already have a scoped prediction key set locally, so use that.
-			ActivationInfo.SetPredictionKey(Owner->ScopedPedictionKey);
-			ActivationInfo.ActivationMode = EGameplayAbilityActivationMode::Predicting;
-			ScopedPredictionKey = Owner->ScopedPedictionKey;
-		}
-		else if(bCanGenerateNewKey)
-		{
-			// We don't have a valid key, so generate a new one.
-			ActivationInfo.GenerateNewPredictionKey();
-			ScopedPredictionKey = ActivationInfo.GetPredictionKeyForNewAction();
-		}
+		ClearScopedPredictionKey = true;
+		RestoreKey = InAbilitySystemComponent->ScopedPredictionKey;
+		InAbilitySystemComponent->ScopedPredictionKey.GenerateDependantPredictionKey();
+		
 	}
 }
 
 FScopedPredictionWindow::~FScopedPredictionWindow()
 {
-	if (Ability.IsValid())
+	if (Owner.IsValid())
 	{
-		// Restore old activation info settings
-		FGameplayAbilityActivationInfo& ActivationInfo = Ability->GetCurrentActivationInfoRef();
-		if (ActivationInfo.ActivationMode != EGameplayAbilityActivationMode::Authority)
+		if (SetReplicatedPredictionKey)
 		{
-			ActivationInfo.ActivationMode = static_cast< TEnumAsByte<EGameplayAbilityActivationMode::Type> > (ClientPrevActivationMode);
+			Owner->ReplicatedPredictionKey = Owner->ScopedPredictionKey;
+		}
+		if (ClearScopedPredictionKey)
+		{
+			Owner->ScopedPredictionKey = RestoreKey;
 		}
-
-		ActivationInfo.SetPredictionStale();
-	}
-
-	if (ClearScopedPredictionKey && Owner.IsValid())
-	{
-		Owner->ScopedPedictionKey = FPredictionKey();
 	}
 }
diff --git a/Engine/Source/Runtime/GameplayAbilities/Public/Abilities/GameplayAbility.h b/Engine/Source/Runtime/GameplayAbilities/Public/Abilities/GameplayAbility.h
--- a/Engine/Source/Runtime/GameplayAbilities/Public/Abilities/GameplayAbility.h
+++ b/Engine/Source/Runtime/GameplayAbilities/Public/Abilities/GameplayAbility.h
@@ -216,6 +216,8 @@ public:
 	/** Is this ability triggered from TriggerData (or is it triggered explicitly through input/game code) */
 	bool IsTriggered() const;
 
+	bool HasAuthorityOrPredictionKey(const FGameplayAbilityActorInfo* ActorInfo, const FGameplayAbilityActivationInfo* ActivationInfo) const;
+
 protected:
 
 	// --------------------------------------
diff --git a/Engine/Source/Runtime/GameplayAbilities/Public/Abilities/GameplayAbilityTypes.h b/Engine/Source/Runtime/GameplayAbilities/Public/Abilities/GameplayAbilityTypes.h
--- a/Engine/Source/Runtime/GameplayAbilities/Public/Abilities/GameplayAbilityTypes.h
+++ b/Engine/Source/Runtime/GameplayAbilities/Public/Abilities/GameplayAbilityTypes.h
@@ -218,42 +218,19 @@ struct GAMEPLAYABILITIES_API FGameplayAbilityActivationInfo
 
 	}
 
-	FGameplayAbilityActivationInfo(AActor* InActor, FPredictionKey InPredictionKey)
-		: bCanBeEndedByOtherInstance(false)
-		, PredictionKey(InPredictionKey)
+	FGameplayAbilityActivationInfo(AActor* InActor)
+		: bCanBeEndedByOtherInstance(false)	
 	{
 		// On Init, we are either Authority or NonAuthority. We haven't been given a PredictionKey and we haven't been confirmed.
 		// NonAuthority essentially means 'I'm not sure what how I'm going to do this yet'.
 		ActivationMode = (InActor->Role == ROLE_Authority ? EGameplayAbilityActivationMode::Authority : EGameplayAbilityActivationMode::NonAuthority);
 	}
 
-	FGameplayAbilityActivationInfo(EGameplayAbilityActivationMode::Type InType, FPredictionKey InPredictionKey = FPredictionKey())
+	FGameplayAbilityActivationInfo(EGameplayAbilityActivationMode::Type InType)
 		: ActivationMode(InType)
 		, bCanBeEndedByOtherInstance(false)
-		, PredictionKey(InPredictionKey)
 	{
-	}
-
-	void GenerateNewPredictionKey() const;
-
-	void SetActivationConfirmed();
-
-	void SetPredictionStale();
-
-	FPredictionKey GetPredictionKeyForNewAction() const
-	{
-		return PredictionKey.IsValidForMorePrediction() ? PredictionKey : FPredictionKey();
-	}
-
-	FPredictionKey GetPredictionKey() const
-	{
-		return PredictionKey;
-	}
-
-	void SetPredictionKey(FPredictionKey NewKey) const
-	{
-		PredictionKey = NewKey;
-	}
+	}	
 
 	UPROPERTY(BlueprintReadOnly, Category = "ActorInfo")
 	mutable TEnumAsByte<EGameplayAbilityActivationMode::Type>	ActivationMode;
@@ -262,10 +239,17 @@ struct GAMEPLAYABILITIES_API FGameplayAbilityActivationInfo
 	UPROPERTY()
 	bool bCanBeEndedByOtherInstance;
 
+	void SetActivationConfirmed();
+
+	void SetPredicting(FPredictionKey PredictionKey);
+
+	FPredictionKey GetActivationPredictionKey() { return PredictionKeyWhenActivated; }
+
 private:
 
-	UPROPERTY()
-	mutable FPredictionKey	PredictionKey;
+	// This was the prediction key used to activate this ability. It does not get updated
+	// if new prediction keys are generated over the course of the ability.
+	FPredictionKey PredictionKeyWhenActivated;
 	
 };
 
@@ -434,8 +418,6 @@ struct GAMEPLAYABILITIES_API FGameplayEventData
 	float Var2;
 	UPROPERTY(EditAnywhere, BlueprintReadWrite, Category = GameplayAbilityTriggerPayload)
 	float Var3;
-
-	FPredictionKey PredictionKey;
 };
 
 /** 
diff --git a/Engine/Source/Runtime/GameplayAbilities/Public/AbilitySystemComponent.h b/Engine/Source/Runtime/GameplayAbilities/Public/AbilitySystemComponent.h
--- a/Engine/Source/Runtime/GameplayAbilities/Public/AbilitySystemComponent.h
+++ b/Engine/Source/Runtime/GameplayAbilities/Public/AbilitySystemComponent.h
@@ -141,7 +141,20 @@ class GAMEPLAYABILITIES_API UAbilitySystemComponent : public UActorComponent, pu
 	UPROPERTY(ReplicatedUsing=OnRep_PredictionKey)
 	FPredictionKey	ReplicatedPredictionKey;
 
-	FPredictionKey	ScopedPedictionKey;
+	FPredictionKey	ScopedPredictionKey;
+
+	FPredictionKey GetPredictionKeyForNewAction() const
+	{
+		return ScopedPredictionKey.IsValidForMorePrediction() ? ScopedPredictionKey : FPredictionKey();
+	}
+
+	/** Do we have a valid prediction key to do more predictive actions with */
+	bool CanPredict() const
+	{
+		return ScopedPredictionKey.IsValidForMorePrediction();
+	}
+
+	bool HasAuthorityOrPredictionKey(const FGameplayAbilityActivationInfo* ActivationInfo) const;	
 
 	UFUNCTION()
 	void OnRep_PredictionKey();
@@ -380,11 +393,11 @@ class GAMEPLAYABILITIES_API UAbilitySystemComponent : public UActorComponent, pu
 
 	// GameplayCues can also come on their own. These take an optional effect context to pass through hit result, etc
 
-	void ExecuteGameplayCue(const FGameplayTag GameplayCueTag, FPredictionKey PredictionKey, FGameplayEffectContextHandle EffectContext = FGameplayEffectContextHandle());
+	void ExecuteGameplayCue(const FGameplayTag GameplayCueTag, FGameplayEffectContextHandle EffectContext = FGameplayEffectContextHandle());
 
-	void AddGameplayCue(const FGameplayTag GameplayCueTag, FPredictionKey PredictionKey, FGameplayEffectContextHandle EffectContext = FGameplayEffectContextHandle());
+	void AddGameplayCue(const FGameplayTag GameplayCueTag, FGameplayEffectContextHandle EffectContext = FGameplayEffectContextHandle());
 	
-	void RemoveGameplayCue(const FGameplayTag GameplayCueTag, FPredictionKey PredictionKey);
+	void RemoveGameplayCue(const FGameplayTag GameplayCueTag);
 
 	UFUNCTION(NetMulticast, unreliable)
 	void NetMulticast_InvokeGameplayCueExecuted(const FGameplayTag GameplayCueTag, FPredictionKey PredictionKey, FGameplayEffectContextHandle EffectContext);
diff --git a/Engine/Source/Runtime/GameplayAbilities/Public/GameplayEffectAggregator.h b/Engine/Source/Runtime/GameplayAbilities/Public/GameplayEffectAggregator.h
--- a/Engine/Source/Runtime/GameplayAbilities/Public/GameplayEffectAggregator.h
+++ b/Engine/Source/Runtime/GameplayAbilities/Public/GameplayEffectAggregator.h
@@ -118,7 +118,10 @@ struct GAMEPLAYABILITIES_API FScopedAggregatorOnDirtyBatch
 	~FScopedAggregatorOnDirtyBatch();
 
 	static void BeginLock();
-	static void EndLock(bool FromNetworkUpdate);
+	static void EndLock();
+
+	static void BeginNetReceiveLock();
+	static void EndNetReceiveLock();
 
 	static int32	GlobalBatchCount;
 	static TSet<FAggregator*>	DirtyAggregators;
diff --git a/Engine/Source/Runtime/GameplayAbilities/Public/GameplayPrediction.h b/Engine/Source/Runtime/GameplayAbilities/Public/GameplayPrediction.h
--- a/Engine/Source/Runtime/GameplayAbilities/Public/GameplayPrediction.h
+++ b/Engine/Source/Runtime/GameplayAbilities/Public/GameplayPrediction.h
@@ -450,7 +450,7 @@ struct GAMEPLAYABILITIES_API FScopedPredictionWindow
 	FScopedPredictionWindow(UAbilitySystemComponent* AbilitySystemComponent, FPredictionKey InPredictionKey);
 
 	/** To be called in the callsite where the predictive code will take place. This generates a new PredictionKey and acts as a synchonization point between client and server for that key.  */
-	FScopedPredictionWindow(UGameplayAbility* GameplayAbilityInstance, bool CanGenerateNewKey=true);
+	FScopedPredictionWindow(UAbilitySystemComponent* AbilitySystemComponent, bool CanGenerateNewKey=true);
 
 	~FScopedPredictionWindow();
 
@@ -459,8 +459,7 @@ struct GAMEPLAYABILITIES_API FScopedPredictionWindow
 private:
 
 	TWeakObjectPtr<UAbilitySystemComponent> Owner;
-	TWeakObjectPtr<UGameplayAbility> Ability;
-
-	int8 ClientPrevActivationMode;
 	bool ClearScopedPredictionKey;
+	bool SetReplicatedPredictionKey;
+	FPredictionKey RestoreKey;
 };
\ No newline at end of file
