#GameplayEffects #Tests: Added a couple more tests, and split up the boilerplate so they're all nicely listed individually in the editor.

diff --git a/Engine/Source/Runtime/GameplayAbilities/Private/GameplayEffectTests.cpp b/Engine/Source/Runtime/GameplayAbilities/Private/GameplayEffectTests.cpp
--- a/Engine/Source/Runtime/GameplayAbilities/Private/GameplayEffectTests.cpp
+++ b/Engine/Source/Runtime/GameplayAbilities/Private/GameplayEffectTests.cpp
@@ -12,26 +12,10 @@
 #include "GameplayEffectStackingExtension_CappedNumberTest.h"
 #include "GameplayEffectStackingExtension_DiminishingReturnsTest.h"
 
-IMPLEMENT_SIMPLE_AUTOMATION_TEST(FGameplayEffectsTest, "AbilitySystem.GameplayEffects", EAutomationTestFlags::ATF_Editor)
-
 #define SKILL_TEST_TEXT( Format, ... ) FString::Printf(TEXT("%s - %d: %s"), TEXT(__FILE__) , __LINE__ , *FString::Printf(TEXT(Format), ##__VA_ARGS__) )
 
 #if WITH_EDITOR
 
-void GameplayTest_TickWorld(UWorld *World, float Time)
-{
-	const float step = 0.1f;
-	while(Time > 0.f)
-	{
-		World->Tick(ELevelTick::LEVELTICK_All, FMath::Min(Time, step) );
-		Time-=step;
-
-		// This is terrible but required for subticking like this.
-		// we could always cache the real GFrameCounter at the start of our tests and restore it when finished.
-		GFrameCounter++;
-	}
-}
-
 static UDataTable* CreateGameplayDataTable()
 {
 	FString CSV(TEXT(",Tag,CategoryText,"));
@@ -85,16 +69,23 @@ public:
 		// run before each test
 
 		const float StartingHealth = 100.f;
+		const float StartingMana = 200.f;
 
 		// set up the source actor
 		SourceActor = World->SpawnActor<AAbilitySystemTestPawn>();
 		SourceComponent = SourceActor->GetAbilitySystemComponent();
 		SourceComponent->GetSet<UAbilitySystemTestAttributeSet>()->Health = StartingHealth;
+		SourceComponent->GetSet<UAbilitySystemTestAttributeSet>()->MaxHealth = StartingHealth;
+		SourceComponent->GetSet<UAbilitySystemTestAttributeSet>()->Mana = StartingMana;
+		SourceComponent->GetSet<UAbilitySystemTestAttributeSet>()->MaxMana = StartingMana;
 
 		// set up the destination actor
 		DestActor = World->SpawnActor<AAbilitySystemTestPawn>();
 		DestComponent = DestActor->GetAbilitySystemComponent();
 		DestComponent->GetSet<UAbilitySystemTestAttributeSet>()->Health = StartingHealth;
+		DestComponent->GetSet<UAbilitySystemTestAttributeSet>()->MaxHealth = StartingHealth;
+		DestComponent->GetSet<UAbilitySystemTestAttributeSet>()->Mana = StartingMana;
+		DestComponent->GetSet<UAbilitySystemTestAttributeSet>()->MaxMana = StartingMana;
 	}
 
 	~GameplayEffectsTestSuite()
@@ -116,34 +107,110 @@ public: // the tests
 
 	void Test_InstantDamage()
 	{
-		ABILITY_LOG_SCOPE(TEXT("Apply InstantDamage"));
+		const float DamageValue = 5.f;
+		const float StartingHealth = DestComponent->GetSet<UAbilitySystemTestAttributeSet>()->Health;
 
-		const float DamageValue = -5.f;
+		// just try and reduce the health attribute
+		{
+			ABILITY_LOG_SCOPE(TEXT("Apply InstantDamage"));
+			
+			CONSTRUCT_CLASS(UGameplayEffect, BaseDmgEffect);
+			AddModifier(BaseDmgEffect, GET_FIELD_CHECKED(UAbilitySystemTestAttributeSet, Health), EGameplayModOp::Additive, FScalableFloat(-DamageValue));
+			BaseDmgEffect->Duration.Value = UGameplayEffect::INSTANT_APPLICATION;
+			
+			SourceComponent->ApplyGameplayEffectToTarget(BaseDmgEffect, DestComponent, 1.f);
+		}
+
+		// make sure health was reduced
+		TestEqual(SKILL_TEST_TEXT("Health Reduced"), DestComponent->GetSet<UAbilitySystemTestAttributeSet>()->Health, StartingHealth - DamageValue);
+	}
+
+	void Test_InstantDamageRemap()
+	{
+		const float DamageValue = 5.f;
 		const float StartingHealth = DestComponent->GetSet<UAbilitySystemTestAttributeSet>()->Health;
 
-		CONSTRUCT_CLASS(UGameplayEffect, BaseDmgEffect);
-		BaseDmgEffect->Modifiers.SetNum(1);
-		BaseDmgEffect->Modifiers[0].ModifierMagnitude = FScalableFloat(DamageValue);
-		BaseDmgEffect->Modifiers[0].ModifierOp = EGameplayModOp::Additive;
-		BaseDmgEffect->Modifiers[0].Attribute.SetUProperty(GET_FIELD_CHECKED(UAbilitySystemTestAttributeSet, Health));
-		BaseDmgEffect->Duration.Value = UGameplayEffect::INSTANT_APPLICATION;
+		// This is the same as GameplayEffectsTest_InstantDamage but modifies the Damage attribute and confirms it is remapped to -Health by UAbilitySystemTestAttributeSet::PostAttributeModify
+		{
+			ABILITY_LOG_SCOPE(TEXT("Apply InstantDamage"));
 
-		SourceComponent->ApplyGameplayEffectToTarget(BaseDmgEffect, DestComponent, 1.f);
+			CONSTRUCT_CLASS(UGameplayEffect, BaseDmgEffect);
+			AddModifier(BaseDmgEffect, GET_FIELD_CHECKED(UAbilitySystemTestAttributeSet, Damage), EGameplayModOp::Additive, FScalableFloat(DamageValue));
+			BaseDmgEffect->Duration.Value = UGameplayEffect::INSTANT_APPLICATION;
 
-		Test->TestEqual(SKILL_TEST_TEXT("Basic Instant Damage Applied"), DestComponent->GetSet<UAbilitySystemTestAttributeSet>()->Health, StartingHealth + DamageValue);
+			SourceComponent->ApplyGameplayEffectToTarget(BaseDmgEffect, DestComponent, 1.f);
+		}
+
+		// Now we should have lost some health
+		TestEqual(SKILL_TEST_TEXT("Health Reduced"), DestComponent->GetSet<UAbilitySystemTestAttributeSet>()->Health, StartingHealth - DamageValue);
+
+		// Confirm the damage attribute itself was reset to 0 when it was applied to health
+		TestEqual(SKILL_TEST_TEXT("Damage Applied"), DestComponent->GetSet<UAbilitySystemTestAttributeSet>()->Damage, 0.f);
 	}
 
-public:
+	void Test_ManaBuff()
+	{
+		const float BuffValue = 30.f;
+		const float StartingMana = DestComponent->GetSet<UAbilitySystemTestAttributeSet>()->Mana;
+
+		FActiveGameplayEffectHandle BuffHandle;
+
+		// apply the buff
+		{
+			ABILITY_LOG_SCOPE(TEXT("Apply Buff"));
+
+			CONSTRUCT_CLASS(UGameplayEffect, DamageBuffEffect);
+			AddModifier(DamageBuffEffect, GET_FIELD_CHECKED(UAbilitySystemTestAttributeSet, Mana), EGameplayModOp::Additive, FScalableFloat(BuffValue));
+			DamageBuffEffect->Duration.Value = UGameplayEffect::INFINITE_DURATION;
+
+			BuffHandle = SourceComponent->ApplyGameplayEffectToTarget(DamageBuffEffect, DestComponent, 1.f);
+		}
+
+		// check that the value changed
+		TestEqual(SKILL_TEST_TEXT("Mana Buffed"), DestComponent->GetSet<UAbilitySystemTestAttributeSet>()->Mana, StartingMana + BuffValue);
+
+		// remove the effect
+		{
+			ABILITY_LOG_SCOPE(TEXT("Remove Buff"));
+
+			DestComponent->RemoveActiveGameplayEffect(BuffHandle);
+		}
+
+		// check that the value changed back
+		TestEqual(SKILL_TEST_TEXT("Mana Restored"), DestComponent->GetSet<UAbilitySystemTestAttributeSet>()->Mana, StartingMana);
+	}
 
-	// used to run tests consistently
-	static void RunTest(UWorld* WorldIn, FAutomationTestBase* TestIn, void (GameplayEffectsTestSuite::*TestFunc)())
+private: // test helpers
+
+	void TestEqual(const FString& TestText, float A, float B)
 	{
-		uint64 InitialFrameCounter = GFrameCounter;
+		Test->TestEqual(FString::Printf(TEXT("%s: %f == %f"), *TestText, A, B), A, B);
+	}
+
+	template<typename MODIFIER_T>
+	FGameplayModifierInfo& AddModifier(UGameplayEffect* Effect, UProperty* Property, EGameplayModOp::Type Op, const MODIFIER_T& Magnitude)
+	{
+		int32 Idx = Effect->Modifiers.Num();
+		Effect->Modifiers.SetNum(Idx+1);
+		FGameplayModifierInfo& Info = Effect->Modifiers[Idx];
+		Info.ModifierMagnitude = Magnitude;
+		Info.ModifierOp = Op;
+		Info.Attribute.SetUProperty(Property);
+		return Info;
+	}
+
+	void TickWorld(float Time)
+	{
+		const float step = 0.1f;
+		while (Time > 0.f)
 		{
-			GameplayEffectsTestSuite Tester(WorldIn, TestIn);
-			(Tester.*TestFunc)();
+			World->Tick(ELevelTick::LEVELTICK_All, FMath::Min(Time, step));
+			Time -= step;
+
+			// This is terrible but required for subticking like this.
+			// we could always cache the real GFrameCounter at the start of our tests and restore it when finished.
+			GFrameCounter++;
 		}
-		GFrameCounter = InitialFrameCounter;
 	}
 
 private:
@@ -157,38 +224,98 @@ private:
 	UAbilitySystemComponent* DestComponent;
 };
 
-#endif //WITH_EDITOR
+#define ADD_TEST(Name) \
+	TestFunctions.Add(&GameplayEffectsTestSuite::##Name); \
+	TestFunctionNames.Add(TEXT(#Name))
 
-bool FGameplayEffectsTest::RunTest( const FString& Parameters )
+class FGameplayEffectsTest : public FAutomationTestBase
 {
-#if WITH_EDITOR
+public:
+	typedef void (GameplayEffectsTestSuite::*TestFunc)();
 
-	UCurveTable *CurveTable = IGameplayAbilitiesModule::Get().GetAbilitySystemGlobals()->GetGlobalCurveTable();
-	UDataTable *DataTable = IGameplayAbilitiesModule::Get().GetAbilitySystemGlobals()->GetGlobalAttributeMetaDataTable();
+	FGameplayEffectsTest(const FString& InName)
+	: FAutomationTestBase(InName, false)
+	{
+		// list all test functions here
+		ADD_TEST(Test_InstantDamage);
+		ADD_TEST(Test_InstantDamageRemap);
+		ADD_TEST(Test_ManaBuff);
+	}
 
-	// setup required GameplayTags
-	UDataTable* TagTable = CreateGameplayDataTable();
+	virtual uint32 GetTestFlags() const { return EAutomationTestFlags::ATF_Editor; }
+	virtual bool IsStressTest() const { return false; }
+	virtual uint32 GetRequiredDeviceNum() const { return 1; }
 
-	IGameplayTagsModule::Get().GetGameplayTagsManager().PopulateTreeFromDataTable(TagTable);
+protected:
+	virtual FString GetBeautifiedTestName() const override { return "AbilitySystem.GameplayEffects"; }
+	virtual void GetTests(TArray<FString>& OutBeautifiedNames, TArray <FString>& OutTestCommands) const override
+	{
+		for (const FString& TestName : TestFunctionNames)
+		{
+			OutBeautifiedNames.Add(TestName);
+			OutTestCommands.Add(TestName);
+		}
+	}
 
-	UWorld *World = UWorld::CreateWorld(EWorldType::Game, false);
-	FWorldContext &WorldContext = GEngine->CreateNewWorldContext(EWorldType::Game);
-	WorldContext.SetCurrentWorld(World);
-	
-	FURL URL;
-	World->InitializeActorsForPlay(URL);
-	World->BeginPlay();
-	
-	GameplayEffectsTestSuite::RunTest(World, this, &GameplayEffectsTestSuite::Test_InstantDamage);
+	bool RunTest(const FString& Parameters)
+	{
+		// find the matching test
+		TestFunc TestFunction = nullptr;
+		for (int32 i = 0; i < TestFunctionNames.Num(); ++i)
+		{
+			if (TestFunctionNames[i] == Parameters)
+			{
+				TestFunction = TestFunctions[i];
+				break;
+			}
+		}
+		if (TestFunction == nullptr)
+		{
+			return false;
+		}
 
-	GEngine->DestroyWorldContext(World);
-	World->DestroyWorld(false);
+		// get the current curve and data table (to restore later)
+		UCurveTable *CurveTable = IGameplayAbilitiesModule::Get().GetAbilitySystemGlobals()->GetGlobalCurveTable();
+		UDataTable *DataTable = IGameplayAbilitiesModule::Get().GetAbilitySystemGlobals()->GetGlobalAttributeMetaDataTable();
 
-	IGameplayAbilitiesModule::Get().GetAbilitySystemGlobals()->AutomationTestOnly_SetGlobalCurveTable(CurveTable);
-	IGameplayAbilitiesModule::Get().GetAbilitySystemGlobals()->AutomationTestOnly_SetGlobalAttributeDataTable(DataTable);
+		// setup required GameplayTags
+		UDataTable* TagTable = CreateGameplayDataTable();
 
-#endif //WITH_EDITOR
-	return true;
+		IGameplayTagsModule::Get().GetGameplayTagsManager().PopulateTreeFromDataTable(TagTable);
+
+		UWorld *World = UWorld::CreateWorld(EWorldType::Game, false);
+		FWorldContext &WorldContext = GEngine->CreateNewWorldContext(EWorldType::Game);
+		WorldContext.SetCurrentWorld(World);
+
+		FURL URL;
+		World->InitializeActorsForPlay(URL);
+		World->BeginPlay();
+
+		// run the matching test
+		uint64 InitialFrameCounter = GFrameCounter;
+		{
+			GameplayEffectsTestSuite Tester(World, this);
+			(Tester.*TestFunction)();
+		}
+		GFrameCounter = InitialFrameCounter;
+
+		GEngine->DestroyWorldContext(World);
+		World->DestroyWorld(false);
+
+		IGameplayAbilitiesModule::Get().GetAbilitySystemGlobals()->AutomationTestOnly_SetGlobalCurveTable(CurveTable);
+		IGameplayAbilitiesModule::Get().GetAbilitySystemGlobals()->AutomationTestOnly_SetGlobalAttributeDataTable(DataTable);
+		return true;
+	}
+
+	TArray<TestFunc> TestFunctions;
+	TArray<FString> TestFunctionNames;
+};
+
+namespace
+{
+	FGameplayEffectsTest FGameplayEffectsTestAutomationTestInstance(TEXT("FGameplayEffectsTest"));
 }
 
+#endif //WITH_EDITOR
+
 
