Make  GetAbilitySystemComponentFromActor NOT look through component list by default for an ASC

diff --git a/Engine/Source/Runtime/GameplayAbilities/Public/AbilitySystemGlobals.h b/Engine/Source/Runtime/GameplayAbilities/Public/AbilitySystemGlobals.h
--- a/Engine/Source/Runtime/GameplayAbilities/Public/AbilitySystemGlobals.h
+++ b/Engine/Source/Runtime/GameplayAbilities/Public/AbilitySystemGlobals.h
@@ -73,7 +73,7 @@ class GAMEPLAYABILITIES_API UAbilitySystemGlobals : public UObject
 
 	static UAbilitySystemGlobals& Get();
 
-	static UAbilitySystemComponent* GetAbilitySystemComponentFromActor(AActor* Actor, bool LookForComponent=true);
+	static UAbilitySystemComponent* GetAbilitySystemComponentFromActor(AActor* Actor, bool LookForComponent=false);
 
 	/** Should allocate a project specific AbilityActorInfo struct. Caller is responsible for deallocation */
 	virtual FGameplayAbilityActorInfo * AllocAbilityActorInfo() const;
