lower log verbosity

diff --git a/Engine/Source/Runtime/GameplayAbilities/Private/AttributeSet.cpp b/Engine/Source/Runtime/GameplayAbilities/Private/AttributeSet.cpp
--- a/Engine/Source/Runtime/GameplayAbilities/Private/AttributeSet.cpp
+++ b/Engine/Source/Runtime/GameplayAbilities/Private/AttributeSet.cpp
@@ -328,7 +328,7 @@ void FAttributeSetInitter::InitAttributeSetDefaults(UAbilitySystemComponent* Abi
 		const FAttributeDefaultValueList* DefaultDataList = SetDefaults.DataMap.Find(Set->GetClass());
 		if (DefaultDataList)
 		{
-			ABILITY_LOG(Warning, TEXT("Initializing Set %s"), *Set->GetName());
+			ABILITY_LOG(Log, TEXT("Initializing Set %s"), *Set->GetName());
 
 			for (auto& DataPair : DefaultDataList->List)
 			{
