adding null check

diff --git a/Engine/Source/Runtime/GameplayAbilities/Private/GameplayCueManager.cpp b/Engine/Source/Runtime/GameplayAbilities/Private/GameplayCueManager.cpp
--- a/Engine/Source/Runtime/GameplayAbilities/Private/GameplayCueManager.cpp
+++ b/Engine/Source/Runtime/GameplayAbilities/Private/GameplayCueManager.cpp
@@ -32,6 +32,13 @@ void UGameplayCueManager::HandleGameplayCues(AActor* TargetActor, const FGamepla
 
 void UGameplayCueManager::HandleGameplayCue(AActor* TargetActor, FGameplayTag GameplayCueTag, EGameplayCueEvent::Type EventType, FGameplayCueParameters Parameters)
 {
+	if (TargetActor == nullptr)
+	{
+		ABILITY_LOG(Warning, TEXT("UGameplayCueManager::HandleGameplayCue called on null TargetActor. GameplayCueTag: %s."), *GameplayCueTag.ToString());
+		return;
+	}
+
+
 	// GameplayCueTags could have been removed from the dictionary but not content. When the content is resaved the old tag will be cleaned up, but it could still come through here
 	// at runtime. Since we only populate the map with dictionary gameplaycue tags, we may not find it here.
 	int32* Ptr=GameplayCueDataMap.Find(GameplayCueTag);
