#GameplayEffects #Test: Fix for clang

diff --git a/Engine/Source/Runtime/GameplayAbilities/Private/GameplayEffectTests.cpp b/Engine/Source/Runtime/GameplayAbilities/Private/GameplayEffectTests.cpp
--- a/Engine/Source/Runtime/GameplayAbilities/Private/GameplayEffectTests.cpp
+++ b/Engine/Source/Runtime/GameplayAbilities/Private/GameplayEffectTests.cpp
@@ -225,7 +225,7 @@ private:
 };
 
 #define ADD_TEST(Name) \
-	TestFunctions.Add(&GameplayEffectsTestSuite::##Name); \
+	TestFunctions.Add(&GameplayEffectsTestSuite::Name); \
 	TestFunctionNames.Add(TEXT(#Name))
 
 class FGameplayEffectsTest : public FAutomationTestBase
