#GameplayEffects #Test: Added PeriodicDamage automation  test

diff --git a/Engine/Source/Runtime/GameplayAbilities/Private/GameplayEffectTests.cpp b/Engine/Source/Runtime/GameplayAbilities/Private/GameplayEffectTests.cpp
--- a/Engine/Source/Runtime/GameplayAbilities/Private/GameplayEffectTests.cpp
+++ b/Engine/Source/Runtime/GameplayAbilities/Private/GameplayEffectTests.cpp
@@ -180,11 +180,49 @@ public: // the tests
 		TestEqual(SKILL_TEST_TEXT("Mana Restored"), DestComponent->GetSet<UAbilitySystemTestAttributeSet>()->Mana, StartingMana);
 	}
 
+	void Test_PeriodicDamage()
+	{
+		const int32 NumPeriods = 10;
+		const float PeriodSecs = 1.0f;
+		const float DamagePerPeriod = 5.f; 
+		const float StartingHealth = DestComponent->GetSet<UAbilitySystemTestAttributeSet>()->Health;
+
+		// just try and reduce the health attribute
+		{
+			CONSTRUCT_CLASS(UGameplayEffect, BaseDmgEffect);
+			AddModifier(BaseDmgEffect, GET_FIELD_CHECKED(UAbilitySystemTestAttributeSet, Health), EGameplayModOp::Additive, FScalableFloat(-DamagePerPeriod));
+			BaseDmgEffect->Duration.Value = NumPeriods * PeriodSecs;
+			BaseDmgEffect->Period.Value = PeriodSecs;
+
+			SourceComponent->ApplyGameplayEffectToTarget(BaseDmgEffect, DestComponent, 1.f);
+		}
+
+		// make sure health was not reduced on the first tick
+		TestEqual(SKILL_TEST_TEXT("Health Reduced"), DestComponent->GetSet<UAbilitySystemTestAttributeSet>()->Health, StartingHealth);
+
+		for (int32 i=1;i<=NumPeriods;++i)
+		{
+			// advance time by one period
+			TickWorld(PeriodSecs);
+
+			// check that health has been reduced
+			TestEqual(SKILL_TEST_TEXT("Health Reduced"), DestComponent->GetSet<UAbilitySystemTestAttributeSet>()->Health, StartingHealth - i*DamagePerPeriod);
+		}
+
+		// advance time by one extra period
+		TickWorld(PeriodSecs);
+
+		// should not have reduced further
+		TestEqual(SKILL_TEST_TEXT("Health Reduced"), DestComponent->GetSet<UAbilitySystemTestAttributeSet>()->Health, StartingHealth - NumPeriods*DamagePerPeriod);
+
+		// TODO: test that the effect is no longer applied
+	}
+
 private: // test helpers
 
-	void TestEqual(const FString& TestText, float A, float B)
+	void TestEqual(const FString& TestText, float Actual, float Expected)
 	{
-		Test->TestEqual(FString::Printf(TEXT("%s: %f == %f"), *TestText, A, B), A, B);
+		Test->TestEqual(FString::Printf(TEXT("%s: %f (actual) != %f (expected)"), *TestText, Actual, Expected), Actual, Expected);
 	}
 
 	template<typename MODIFIER_T>
@@ -240,6 +278,7 @@ public:
 		ADD_TEST(Test_InstantDamage);
 		ADD_TEST(Test_InstantDamageRemap);
 		ADD_TEST(Test_ManaBuff);
+		ADD_TEST(Test_PeriodicDamage);
 	}
 
 	virtual uint32 GetTestFlags() const { return EAutomationTestFlags::ATF_Editor; }
