Fixing (moving) notification for rejected ability confirmation.

diff --git a/Engine/Source/Runtime/Engine/Classes/GameFramework/PlayerController.h b/Engine/Source/Runtime/Engine/Classes/GameFramework/PlayerController.h
--- a/Engine/Source/Runtime/Engine/Classes/GameFramework/PlayerController.h
+++ b/Engine/Source/Runtime/Engine/Classes/GameFramework/PlayerController.h
@@ -648,10 +648,6 @@ public:
 	UFUNCTION(Reliable, Client)
 	void ClientIgnoreMoveInput(bool bIgnore);
 
-	/** Informs the client of a rejected attempt to confirm an ability. */
-	UFUNCTION(Unreliable, Client)
-	void ClientNotifyRejectedAbilityConfirmation(int32 InputID);
-
 	/**
 	 * Outputs a message to HUD
 	 * @param S - message to display
diff --git a/Engine/Source/Runtime/Engine/Private/PlayerController.cpp b/Engine/Source/Runtime/Engine/Private/PlayerController.cpp
--- a/Engine/Source/Runtime/Engine/Private/PlayerController.cpp
+++ b/Engine/Source/Runtime/Engine/Private/PlayerController.cpp
@@ -2754,11 +2754,6 @@ void APlayerController::ClientIgnoreLookInput_Implementation(bool bIgnore)
 	SetIgnoreLookInput(bIgnore);
 }
 
-void APlayerController::ClientNotifyRejectedAbilityConfirmation_Implementation(int32 InputID)
-{
-	//TODO: Hook UI notification here. Might be good to include a cooldown on this.
-}
-
 
 void APlayerController::DisplayDebug(class UCanvas* Canvas, const FDebugDisplayInfo& DebugDisplay, float& YL, float& YPos)
 {
diff --git a/Engine/Source/Runtime/GameplayAbilities/Private/Abilities/GameplayAbilityTargetActor.cpp b/Engine/Source/Runtime/GameplayAbilities/Private/Abilities/GameplayAbilityTargetActor.cpp
--- a/Engine/Source/Runtime/GameplayAbilities/Private/Abilities/GameplayAbilityTargetActor.cpp
+++ b/Engine/Source/Runtime/GameplayAbilities/Private/Abilities/GameplayAbilityTargetActor.cpp
@@ -71,13 +71,13 @@ void AGameplayAbilityTargetActor::ConfirmTargeting()
 
 void AGameplayAbilityTargetActor::NotifyPlayerControllerOfRejectedConfirmation()
 {
-	if (OwningAbility && OwningAbility->IsInstantiated())
+	if (OwningAbility && OwningAbility->IsInstantiated() && OwningAbility->GetCurrentActorInfo())
 	{
-		if (APlayerController* PC = OwningAbility->GetActorInfo().PlayerController.Get())
+		if (UAbilitySystemComponent* ASC = OwningAbility->GetCurrentActorInfo()->AbilitySystemComponent.Get())
 		{
 			if (FGameplayAbilitySpec* AbilitySpec = OwningAbility->GetCurrentAbilitySpec())
 			{
-				PC->ClientNotifyRejectedAbilityConfirmation(AbilitySpec->InputID);
+				ASC->ClientAbilityNotifyRejected(AbilitySpec->InputID);
 			}
 		}
 	}
diff --git a/Engine/Source/Runtime/GameplayAbilities/Private/AbilitySystemComponent_Abilities.cpp b/Engine/Source/Runtime/GameplayAbilities/Private/AbilitySystemComponent_Abilities.cpp
--- a/Engine/Source/Runtime/GameplayAbilities/Private/AbilitySystemComponent_Abilities.cpp
+++ b/Engine/Source/Runtime/GameplayAbilities/Private/AbilitySystemComponent_Abilities.cpp
@@ -842,6 +842,11 @@ void UAbilitySystemComponent::ClientActivateAbilitySucceed_Implementation(FGamep
 	}
 }
 
+void UAbilitySystemComponent::ClientAbilityNotifyRejected_Implementation(int32 InputID)
+{
+	TargetingRejectedConfirmationDelegate.Broadcast(InputID);
+}
+
 void UAbilitySystemComponent::TriggerAbilityFromGameplayEvent(FGameplayAbilitySpecHandle Handle, FGameplayAbilityActorInfo* ActorInfo, FGameplayTag EventTag, FGameplayEventData* Payload, UAbilitySystemComponent& Component)
 {
 	FGameplayAbilitySpec* Spec = FindAbilitySpecFromHandle(Handle);
diff --git a/Engine/Source/Runtime/GameplayAbilities/Public/AbilitySystemComponent.h b/Engine/Source/Runtime/GameplayAbilities/Public/AbilitySystemComponent.h
--- a/Engine/Source/Runtime/GameplayAbilities/Public/AbilitySystemComponent.h
+++ b/Engine/Source/Runtime/GameplayAbilities/Public/AbilitySystemComponent.h
@@ -39,6 +39,8 @@
  */
 
 
+DECLARE_MULTICAST_DELEGATE_OneParam(FTargetingRejectedConfirmation, int32);
+
 /**
  *	The core ActorComponent for interfacing with the GameplayAbilities System
  */
@@ -512,7 +514,11 @@ class GAMEPLAYABILITIES_API UAbilitySystemComponent : public UActorComponent, pu
 
 	UFUNCTION(Client, Reliable)
 	void	ClientActivateAbilitySucceed(FGameplayAbilitySpecHandle AbilityToActivate,int16 PredictionKey);
-	
+
+	/** Attempted to confirm targeting, but the targeting actor rejected it. */
+	UFUNCTION(Client, Unreliable)
+	void	ClientAbilityNotifyRejected(int32 InputID);
+
 
 	// ----------------------------------------------------------------------------------------------------------------
 
@@ -738,6 +744,9 @@ public:
 	/** ReplicatedTargetData was 'cancelled' for this activation */
 	FAbilityConfirmOrCancel	ReplicatedTargetDataCancelledDelegate;
 
+	/** Targeting actor rejected a confirmation attempt */
+	FTargetingRejectedConfirmation TargetingRejectedConfirmationDelegate;
+
 	/** Tasks that run on simulated proxies */
 	UPROPERTY(ReplicatedUsing=OnRep_SimulatedTasks)
 	TArray<UAbilityTask*> SimulatedTasks;
