DeriveGameplayCueTagFromAssetName made into a static function on AbilitySystemGlobals instead of being duplicated on GCNA and GCNS classes.

diff --git a/Engine/Source/Runtime/GameplayAbilities/Private/AbilitySystemGlobals.cpp b/Engine/Source/Runtime/GameplayAbilities/Private/AbilitySystemGlobals.cpp
--- a/Engine/Source/Runtime/GameplayAbilities/Private/AbilitySystemGlobals.cpp
+++ b/Engine/Source/Runtime/GameplayAbilities/Private/AbilitySystemGlobals.cpp
@@ -63,6 +63,33 @@ T* UAbilitySystemGlobals::InternalGetLoadTable(T*& Table, FString TableName)
 	return Table;
 }
 
+void UAbilitySystemGlobals::DeriveGameplayCueTagFromAssetName(FString AssetName, FGameplayTag& GameplayCueTag, FName& GameplayCueName)
+{
+	// In the editor, attempt to infer GameplayCueTag from our asset name (if there is no valid GameplayCueTag already).
+	#if WITH_EDITOR
+	if (GIsEditor)
+	{
+		if (GameplayCueTag.IsValid() == false)
+		{
+			AssetName.RemoveFromStart(TEXT("Default__"));
+			AssetName.RemoveFromStart(TEXT("GC_"));		// allow GC_ prefix in asset name
+			AssetName.RemoveFromEnd(TEXT("_c"));
+
+			AssetName.ReplaceInline(TEXT("_"), TEXT("."));
+
+			if (!AssetName.Contains(TEXT("GameplayCue")))
+			{
+				AssetName = FString(TEXT("GameplayCue.")) + AssetName;
+			}
+
+			IGameplayTagsModule& GameplayTagsModule = IGameplayTagsModule::Get();
+			GameplayCueTag = GameplayTagsModule.GetGameplayTagsManager().RequestGameplayTag(FName(*AssetName), false);
+		}
+		GameplayCueName = GameplayCueTag.GetTagName();
+	}
+	#endif
+}
+
 
 #if WITH_EDITOR
 
diff --git a/Engine/Source/Runtime/GameplayAbilities/Private/GameplayCueNotify_Actor.cpp b/Engine/Source/Runtime/GameplayAbilities/Private/GameplayCueNotify_Actor.cpp
--- a/Engine/Source/Runtime/GameplayAbilities/Private/GameplayCueNotify_Actor.cpp
+++ b/Engine/Source/Runtime/GameplayAbilities/Private/GameplayCueNotify_Actor.cpp
@@ -22,31 +22,7 @@ void AGameplayCueNotify_Actor::PostEditChangeProperty(FPropertyChangedEvent& Pro
 
 void AGameplayCueNotify_Actor::DeriveGameplayCueTagFromAssetName()
 {
-	// In the editor, attempt to infer GameplayCueTag from our asset name (if there is no valid GameplayCueTag already).
-	#if WITH_EDITOR
-	if (GIsEditor)
-	{
-		if (GameplayCueTag.IsValid() == false)
-		{
-			FString MyName = GetName();
-
-			MyName.RemoveFromStart(TEXT("Default__"));
-			MyName.RemoveFromStart(TEXT("GC_"));		// allow GC_ prefix in asset name
-			MyName.RemoveFromEnd(TEXT("_c"));
-
-			MyName.ReplaceInline(TEXT("_"), TEXT("."));
-
-			if (!MyName.Contains(TEXT("GameplayCue")))
-			{
-				MyName = FString(TEXT("GameplayCue.")) + MyName;
-			}
-
-			IGameplayTagsModule& GameplayTagsModule = IGameplayTagsModule::Get();
-			GameplayCueTag = GameplayTagsModule.GetGameplayTagsManager().RequestGameplayTag(FName(*MyName), false);
-		}
-		GameplayCueName = GameplayCueTag.GetTagName();
-	}
-	#endif
+	UAbilitySystemGlobals::DeriveGameplayCueTagFromAssetName(GetName(), GameplayCueTag, GameplayCueName);
 }
 
 void AGameplayCueNotify_Actor::Serialize(FArchive& Ar)
diff --git a/Engine/Source/Runtime/GameplayAbilities/Private/GameplayCueNotify_Static.cpp b/Engine/Source/Runtime/GameplayAbilities/Private/GameplayCueNotify_Static.cpp
--- a/Engine/Source/Runtime/GameplayAbilities/Private/GameplayCueNotify_Static.cpp
+++ b/Engine/Source/Runtime/GameplayAbilities/Private/GameplayCueNotify_Static.cpp
@@ -21,31 +21,7 @@ void UGameplayCueNotify_Static::PostEditChangeProperty(FPropertyChangedEvent& Pr
 
 void UGameplayCueNotify_Static::DeriveGameplayCueTagFromAssetName()
 {
-	// In the editor, attempt to infer GameplayCueTag from our asset name (if there is no valid GameplayCueTag already).
-	#if WITH_EDITOR
-	if (GIsEditor)
-	{
-		if (GameplayCueTag.IsValid() == false)
-		{
-			FString MyName = GetName();
-
-			MyName.RemoveFromStart(TEXT("Default__"));
-			MyName.RemoveFromStart(TEXT("GC_"));		// allow GC_ prefix in asset name
-			MyName.RemoveFromEnd(TEXT("_c"));
-
-			MyName.ReplaceInline(TEXT("_"), TEXT("."));
-
-			if (MyName.Contains(TEXT("GameplayCue")) == false)
-			{
-				MyName = FString(TEXT("GameplayCue.")) + MyName;
-			}
-
-			IGameplayTagsModule& GameplayTagsModule = IGameplayTagsModule::Get();
-			GameplayCueTag = GameplayTagsModule.GetGameplayTagsManager().RequestGameplayTag(FName(*MyName), false);
-		}
-		GameplayCueName = GameplayCueTag.GetTagName();
-	}
-	#endif
+	UAbilitySystemGlobals::DeriveGameplayCueTagFromAssetName(GetName(), GameplayCueTag, GameplayCueName);
 }
 
 void UGameplayCueNotify_Static::Serialize(FArchive& Ar)
diff --git a/Engine/Source/Runtime/GameplayAbilities/Public/AbilitySystemGlobals.h b/Engine/Source/Runtime/GameplayAbilities/Public/AbilitySystemGlobals.h
--- a/Engine/Source/Runtime/GameplayAbilities/Public/AbilitySystemGlobals.h
+++ b/Engine/Source/Runtime/GameplayAbilities/Public/AbilitySystemGlobals.h
@@ -56,6 +56,8 @@ class GAMEPLAYABILITIES_API UAbilitySystemGlobals : public UObject
 
 	UDataTable* GetGlobalAttributeMetaDataTable();
 
+	static void DeriveGameplayCueTagFromAssetName(FString AssetName, FGameplayTag& GameplayCueTag, FName& GameplayCueName);
+
 	bool IsAbilitySystemGlobalsInitialized()
 	{
 		return GlobalAttributeSetInitter.IsValid();
diff --git a/Engine/Source/Runtime/GameplayAbilities/Public/GameplayCueNotify_Actor.h b/Engine/Source/Runtime/GameplayAbilities/Public/GameplayCueNotify_Actor.h
--- a/Engine/Source/Runtime/GameplayAbilities/Public/GameplayCueNotify_Actor.h
+++ b/Engine/Source/Runtime/GameplayAbilities/Public/GameplayCueNotify_Actor.h
@@ -64,5 +64,5 @@ class GAMEPLAYABILITIES_API AGameplayCueNotify_Actor : public AActor
 	bool IsOverride;
 
 private:
-	void	DeriveGameplayCueTagFromAssetName();
+	FORCEINLINE void DeriveGameplayCueTagFromAssetName();
 };
\ No newline at end of file
diff --git a/Engine/Source/Runtime/GameplayAbilities/Public/GameplayCueNotify_Static.h b/Engine/Source/Runtime/GameplayAbilities/Public/GameplayCueNotify_Static.h
--- a/Engine/Source/Runtime/GameplayAbilities/Public/GameplayCueNotify_Static.h
+++ b/Engine/Source/Runtime/GameplayAbilities/Public/GameplayCueNotify_Static.h
@@ -64,5 +64,5 @@ class GAMEPLAYABILITIES_API UGameplayCueNotify_Static : public UObject
 	bool IsOverride;
 
 private:
-	void	DeriveGameplayCueTagFromAssetName();
+	FORCEINLINE void DeriveGameplayCueTagFromAssetName();
 };
\ No newline at end of file
