Fixing inlined functions that should have been made as virtuals. Also changing GCNS class to use BP const functions now that the feature is working.

diff --git a/Engine/Source/Runtime/GameplayAbilities/Public/GameplayCueNotify_Actor.h b/Engine/Source/Runtime/GameplayAbilities/Public/GameplayCueNotify_Actor.h
--- a/Engine/Source/Runtime/GameplayAbilities/Public/GameplayCueNotify_Actor.h
+++ b/Engine/Source/Runtime/GameplayAbilities/Public/GameplayCueNotify_Actor.h
@@ -64,5 +64,5 @@ class GAMEPLAYABILITIES_API AGameplayCueNotify_Actor : public AActor
 	bool IsOverride;
 
 private:
-	FORCEINLINE void DeriveGameplayCueTagFromAssetName();
+	virtual void DeriveGameplayCueTagFromAssetName();
 };
\ No newline at end of file
diff --git a/Engine/Source/Runtime/GameplayAbilities/Public/GameplayCueNotify_Static.h b/Engine/Source/Runtime/GameplayAbilities/Public/GameplayCueNotify_Static.h
--- a/Engine/Source/Runtime/GameplayAbilities/Public/GameplayCueNotify_Static.h
+++ b/Engine/Source/Runtime/GameplayAbilities/Public/GameplayCueNotify_Static.h
@@ -44,13 +44,13 @@ class GAMEPLAYABILITIES_API UGameplayCueNotify_Static : public UObject
 	void K2_HandleGameplayCue(TWeakObjectPtr<AActor> MyTarget, EGameplayCueEvent::Type EventType, FGameplayCueParameters Parameters);
 
 	UFUNCTION(BlueprintImplementableEvent, BlueprintPure, Category = "GameplayCueNotify")
-	/*static*/ bool OnExecute(TWeakObjectPtr<AActor> MyTarget, FGameplayCueParameters Parameters);
+	bool OnExecute(TWeakObjectPtr<AActor> MyTarget, FGameplayCueParameters Parameters) const;
 
 	UFUNCTION(BlueprintImplementableEvent, BlueprintPure, Category = "GameplayCueNotify")
-	/*static*/ bool OnActive(TWeakObjectPtr<AActor> MyTarget, FGameplayCueParameters Parameters);
+	bool OnActive(TWeakObjectPtr<AActor> MyTarget, FGameplayCueParameters Parameters) const;
 
 	UFUNCTION(BlueprintImplementableEvent, BlueprintPure, Category = "GameplayCueNotify")
-	/*static*/ bool OnRemove(TWeakObjectPtr<AActor> MyTarget, FGameplayCueParameters Parameters);
+	bool OnRemove(TWeakObjectPtr<AActor> MyTarget, FGameplayCueParameters Parameters) const;
 
 	UPROPERTY(EditDefaultsOnly, Category = GameplayCue)
 	FGameplayTag	GameplayCueTag;
@@ -64,5 +64,5 @@ class GAMEPLAYABILITIES_API UGameplayCueNotify_Static : public UObject
 	bool IsOverride;
 
 private:
-	FORCEINLINE void DeriveGameplayCueTagFromAssetName();
+	virtual void DeriveGameplayCueTagFromAssetName();
 };
\ No newline at end of file
