Exposing BP function.

diff --git a/Engine/Source/Runtime/GameplayAbilities/Public/AbilitySystemBlueprintLibrary.h b/Engine/Source/Runtime/GameplayAbilities/Public/AbilitySystemBlueprintLibrary.h
--- a/Engine/Source/Runtime/GameplayAbilities/Public/AbilitySystemBlueprintLibrary.h
+++ b/Engine/Source/Runtime/GameplayAbilities/Public/AbilitySystemBlueprintLibrary.h
@@ -23,6 +23,7 @@ class GAMEPLAYABILITIES_API UAbilitySystemBlueprintLibrary : public UBlueprintFu
 	UFUNCTION(BlueprintPure, Category = Ability)
 	static UAbilitySystemComponent* GetAbilitySystemComponent(AActor *Actor);
 
+	UFUNCTION(BlueprintPure, Category = Ability)
 	static AActor* FindPlayerStart(AController* PlayerController);
 
 	// -------------------------------------------------------------------------------
