Made GCNA and GCNS child classes for Orion.

diff --git a/Engine/Source/Runtime/GameplayAbilities/Public/GameplayCueNotify_Actor.h b/Engine/Source/Runtime/GameplayAbilities/Public/GameplayCueNotify_Actor.h
--- a/Engine/Source/Runtime/GameplayAbilities/Public/GameplayCueNotify_Actor.h
+++ b/Engine/Source/Runtime/GameplayAbilities/Public/GameplayCueNotify_Actor.h
@@ -41,16 +41,16 @@ class GAMEPLAYABILITIES_API AGameplayCueNotify_Actor : public AActor
 
 	/** Generic Event Graph event that will get called for every event type */
 	UFUNCTION(BlueprintImplementableEvent, Category = "GameplayCueNotify", FriendlyName = "HandleGameplayCue")
-	void K2_HandleGameplayCue(TWeakObjectPtr<AActor> MyTarget, EGameplayCueEvent::Type EventType, FGameplayCueParameters Parameters);
+	virtual void K2_HandleGameplayCue(TWeakObjectPtr<AActor> MyTarget, EGameplayCueEvent::Type EventType, FGameplayCueParameters Parameters);
 
 	UFUNCTION(BlueprintImplementableEvent)
-	bool OnExecute(TWeakObjectPtr<AActor> MyTarget, FGameplayCueParameters Parameters);
+	virtual bool OnExecute(TWeakObjectPtr<AActor> MyTarget, FGameplayCueParameters Parameters);
 
 	UFUNCTION(BlueprintImplementableEvent)
-	bool OnActive(TWeakObjectPtr<AActor> MyTarget, FGameplayCueParameters Parameters);
+	virtual bool OnActive(TWeakObjectPtr<AActor> MyTarget, FGameplayCueParameters Parameters);
 
 	UFUNCTION(BlueprintImplementableEvent)
-	bool OnRemove(TWeakObjectPtr<AActor> MyTarget, FGameplayCueParameters Parameters);
+	virtual bool OnRemove(TWeakObjectPtr<AActor> MyTarget, FGameplayCueParameters Parameters);
 
 	UPROPERTY(EditDefaultsOnly, Category = GameplayCue)
 	FGameplayTag	GameplayCueTag;
diff --git a/Engine/Source/Runtime/GameplayAbilities/Public/GameplayCueNotify_Static.h b/Engine/Source/Runtime/GameplayAbilities/Public/GameplayCueNotify_Static.h
--- a/Engine/Source/Runtime/GameplayAbilities/Public/GameplayCueNotify_Static.h
+++ b/Engine/Source/Runtime/GameplayAbilities/Public/GameplayCueNotify_Static.h
@@ -41,16 +41,16 @@ class GAMEPLAYABILITIES_API UGameplayCueNotify_Static : public UObject
 
 	/** Generic Event Graph event that will get called for every event type */
 	UFUNCTION(BlueprintImplementableEvent, Category = "GameplayCueNotify", FriendlyName = "HandleGameplayCue")
-	void K2_HandleGameplayCue(TWeakObjectPtr<AActor> MyTarget, EGameplayCueEvent::Type EventType, FGameplayCueParameters Parameters);
+	virtual void K2_HandleGameplayCue(TWeakObjectPtr<AActor> MyTarget, EGameplayCueEvent::Type EventType, FGameplayCueParameters Parameters) const;
 
 	UFUNCTION(BlueprintImplementableEvent, BlueprintPure, Category = "GameplayCueNotify")
-	bool OnExecute(TWeakObjectPtr<AActor> MyTarget, FGameplayCueParameters Parameters) const;
+	virtual bool OnExecute(TWeakObjectPtr<AActor> MyTarget, FGameplayCueParameters Parameters) const;
 
 	UFUNCTION(BlueprintImplementableEvent, BlueprintPure, Category = "GameplayCueNotify")
-	bool OnActive(TWeakObjectPtr<AActor> MyTarget, FGameplayCueParameters Parameters) const;
+	virtual bool OnActive(TWeakObjectPtr<AActor> MyTarget, FGameplayCueParameters Parameters) const;
 
 	UFUNCTION(BlueprintImplementableEvent, BlueprintPure, Category = "GameplayCueNotify")
-	bool OnRemove(TWeakObjectPtr<AActor> MyTarget, FGameplayCueParameters Parameters) const;
+	virtual bool OnRemove(TWeakObjectPtr<AActor> MyTarget, FGameplayCueParameters Parameters) const;
 
 	UPROPERTY(EditDefaultsOnly, Category = GameplayCue)
 	FGameplayTag	GameplayCueTag;
