Fixed FGameplayEffectContext::AddActors() assert

* Fixed FGameplayEffectContext::AddActors() assert caused by adding multiple actors to the context.

diff --git a/Engine/Source/Runtime/GameplayAbilities/Private/GameplayEffectTypes.cpp b/Engine/Source/Runtime/GameplayAbilities/Private/GameplayEffectTypes.cpp
--- a/Engine/Source/Runtime/GameplayAbilities/Private/GameplayEffectTypes.cpp
+++ b/Engine/Source/Runtime/GameplayAbilities/Private/GameplayEffectTypes.cpp
@@ -48,8 +48,7 @@ void FGameplayEffectContext::AddActors(TArray<TWeakObjectPtr<AActor>> InActors,
 		Actors.Reset();
 	}
 
-	check(!Actors.Num());
-	Actors = InActors;
+	Actors.Append(InActors);
 }
 
 void FGameplayEffectContext::AddHitResult(const FHitResult InHitResult, bool bReset)
