PR #637: Move FindPlayerStart BP function to a proper location (Contributed by [REDACTED_PERSON])

diff --git a/Engine/Source/Runtime/Engine/Classes/GameFramework/GameMode.h b/Engine/Source/Runtime/Engine/Classes/GameFramework/GameMode.h
--- a/Engine/Source/Runtime/Engine/Classes/GameFramework/GameMode.h
+++ b/Engine/Source/Runtime/Engine/Classes/GameFramework/GameMode.h
@@ -502,6 +502,9 @@ public:
 	 */
 	virtual class AActor* FindPlayerStart( AController* Player, const FString& IncomingName = TEXT("") );
 
+	UFUNCTION(BlueprintPure, Category=Game, meta=(FriendlyName="FindPlayerStart"))
+	AActor* K2_FindPlayerStart(AController* Player);
+
 	/** 
 	* Return the 'best' player start for this player to start from.  
 	* Default implementation just returns the first PlayerStart found.
diff --git a/Engine/Source/Runtime/Engine/Private/GameMode.cpp b/Engine/Source/Runtime/Engine/Private/GameMode.cpp
--- a/Engine/Source/Runtime/Engine/Private/GameMode.cpp
+++ b/Engine/Source/Runtime/Engine/Private/GameMode.cpp
@@ -391,6 +391,11 @@ AActor* AGameMode::FindPlayerStart( AController* Player, const FString& Incoming
 	return BestStart;
 }
 
+AActor* AGameMode::K2_FindPlayerStart( AController* Player )
+{
+	return FindPlayerStart(Player);
+}
+
 void AGameMode::PreInitializeComponents()
 {
 	Super::PreInitializeComponents();
diff --git a/Engine/Source/Runtime/GameplayAbilities/Private/AbilitySystemBlueprintLibrary.cpp b/Engine/Source/Runtime/GameplayAbilities/Private/AbilitySystemBlueprintLibrary.cpp
--- a/Engine/Source/Runtime/GameplayAbilities/Private/AbilitySystemBlueprintLibrary.cpp
+++ b/Engine/Source/Runtime/GameplayAbilities/Private/AbilitySystemBlueprintLibrary.cpp
@@ -20,21 +20,6 @@ UAbilitySystemComponent* UAbilitySystemBlueprintLibrary::GetAbilitySystemCompone
 	return UAbilitySystemGlobals::GetAbilitySystemComponentFromActor(Actor);
 }
 
-AActor* UAbilitySystemBlueprintLibrary::FindPlayerStart(AController* PlayerController)
-{
-	if (PlayerController)
-	{
-		if (UWorld* World = PlayerController->GetWorld())
-		{
-			if (AGameMode* GameMode = World->GetAuthGameMode())
-			{
-				return GameMode->FindPlayerStart(PlayerController);
-			}
-		}
-	}
-	return NULL;
-}
-
 FGameplayAbilityTargetDataHandle UAbilitySystemBlueprintLibrary::AppendTargetDataHandle(FGameplayAbilityTargetDataHandle TargetHandle, FGameplayAbilityTargetDataHandle HandleToAdd)
 {
 	TargetHandle.Append(&HandleToAdd);
diff --git a/Engine/Source/Runtime/GameplayAbilities/Public/AbilitySystemBlueprintLibrary.h b/Engine/Source/Runtime/GameplayAbilities/Public/AbilitySystemBlueprintLibrary.h
--- a/Engine/Source/Runtime/GameplayAbilities/Public/AbilitySystemBlueprintLibrary.h
+++ b/Engine/Source/Runtime/GameplayAbilities/Public/AbilitySystemBlueprintLibrary.h
@@ -23,9 +23,6 @@ class GAMEPLAYABILITIES_API UAbilitySystemBlueprintLibrary : public UBlueprintFu
 	UFUNCTION(BlueprintPure, Category = Ability)
 	static UAbilitySystemComponent* GetAbilitySystemComponent(AActor *Actor);
 
-	UFUNCTION(BlueprintPure, Category = Ability)
-	static AActor* FindPlayerStart(AController* PlayerController);
-
 	// -------------------------------------------------------------------------------
 	//		TargetData
 	// -------------------------------------------------------------------------------
