- Make FActiveGameplayEffectsContainer::RemoveActiveEffects respect GE scope locking

diff --git a/Engine/Source/Runtime/GameplayAbilities/Private/GameplayEffect.cpp b/Engine/Source/Runtime/GameplayAbilities/Private/GameplayEffect.cpp
--- a/Engine/Source/Runtime/GameplayAbilities/Private/GameplayEffect.cpp
+++ b/Engine/Source/Runtime/GameplayAbilities/Private/GameplayEffect.cpp
@@ -2042,13 +2042,23 @@ TArray<float> FActiveGameplayEffectsContainer::GetActiveEffectsDuration(const FA
 
 void FActiveGameplayEffectsContainer::RemoveActiveEffects(const FActiveGameplayEffectQuery Query)
 {
+	const bool bLockInEffect = FScopedActiveGameplayEffectLock::IsLockInEffect();
+
 	for (int32 idx=0; idx < GameplayEffects.Num(); ++idx)
 	{
 		const FActiveGameplayEffect& Effect = GameplayEffects[idx];
 		if (Query.Matches(Effect))
 		{
-			InternalRemoveActiveGameplayEffect(idx);
-			idx--;
+			// Defer removal if we're scope-locked
+			if (bLockInEffect)
+			{
+				FScopedActiveGameplayEffectLock::AddAction(new FActiveGameplayEffectAction_Remove(TWeakObjectPtr<UAbilitySystemComponent>(Owner), Effect.Handle));
+			}
+			else
+			{
+				InternalRemoveActiveGameplayEffect(idx);
+				idx--;
+			}
 		}
 	}
 }
diff --git a/Engine/Source/Runtime/GameplayAbilities/Public/GameplayEffect.h b/Engine/Source/Runtime/GameplayAbilities/Public/GameplayEffect.h
--- a/Engine/Source/Runtime/GameplayAbilities/Public/GameplayEffect.h
+++ b/Engine/Source/Runtime/GameplayAbilities/Public/GameplayEffect.h
@@ -1320,7 +1320,7 @@ struct GAMEPLAYABILITIES_API FActiveGameplayEffectAction_Remove : public FActive
 	{
 	}
 
-	FActiveGameplayEffectAction_Remove(TWeakObjectPtr<UAbilitySystemComponent> InOwningASC, FActiveGameplayEffectHandle& InHandle)
+	FActiveGameplayEffectAction_Remove(TWeakObjectPtr<UAbilitySystemComponent> InOwningASC, const FActiveGameplayEffectHandle& InHandle)
 	: OwningASC(InOwningASC)
 	, Handle(InHandle)
 	{
