Added hitscan weapon effect GCNS, changed the base and other derived classes slightly.

diff --git a/Engine/Source/Runtime/GameplayAbilities/Private/GameplayCueNotify_Actor.cpp b/Engine/Source/Runtime/GameplayAbilities/Private/GameplayCueNotify_Actor.cpp
--- a/Engine/Source/Runtime/GameplayAbilities/Private/GameplayCueNotify_Actor.cpp
+++ b/Engine/Source/Runtime/GameplayAbilities/Private/GameplayCueNotify_Actor.cpp
@@ -84,3 +84,30 @@ void AGameplayCueNotify_Actor::OnOwnerDestroyed()
 	Destroy();
 	//MarkPendingKill();
 }
+
+bool AGameplayCueNotify_Actor::OnExecute_Implementation(AActor* MyTarget, FGameplayCueParameters Parameters)
+{
+	return OnExecuteNative(MyTarget, Parameters);
+}
+bool AGameplayCueNotify_Actor::OnExecuteNative(AActor* MyTarget, FGameplayCueParameters Parameters)
+{
+	return false;
+}
+
+bool AGameplayCueNotify_Actor::OnActive_Implementation(AActor* MyTarget, FGameplayCueParameters Parameters)
+{
+	return OnActiveNative(MyTarget, Parameters);
+}
+bool AGameplayCueNotify_Actor::OnActiveNative(AActor* MyTarget, FGameplayCueParameters Parameters)
+{
+	return false;
+}
+
+bool AGameplayCueNotify_Actor::OnRemove_Implementation(AActor* MyTarget, FGameplayCueParameters Parameters)
+{
+	return OnRemoveNative(MyTarget, Parameters);
+}
+bool AGameplayCueNotify_Actor::OnRemoveNative(AActor* MyTarget, FGameplayCueParameters Parameters)
+{
+	return false;
+}
diff --git a/Engine/Source/Runtime/GameplayAbilities/Private/GameplayCueNotify_Static.cpp b/Engine/Source/Runtime/GameplayAbilities/Private/GameplayCueNotify_Static.cpp
--- a/Engine/Source/Runtime/GameplayAbilities/Private/GameplayCueNotify_Static.cpp
+++ b/Engine/Source/Runtime/GameplayAbilities/Private/GameplayCueNotify_Static.cpp
@@ -80,3 +80,30 @@ void UGameplayCueNotify_Static::HandleGameplayCue(AActor* MyTarget, EGameplayCue
 void UGameplayCueNotify_Static::OnOwnerDestroyed()
 {
 }
+
+bool UGameplayCueNotify_Static::OnExecute_Implementation(AActor* MyTarget, FGameplayCueParameters Parameters) const
+{
+	return OnExecuteNative(MyTarget, Parameters);
+}
+bool UGameplayCueNotify_Static::OnExecuteNative(AActor* MyTarget, FGameplayCueParameters Parameters) const
+{
+	return false;
+}
+
+bool UGameplayCueNotify_Static::OnActive_Implementation(AActor* MyTarget, FGameplayCueParameters Parameters) const
+{
+	return OnActiveNative(MyTarget, Parameters);
+}
+bool UGameplayCueNotify_Static::OnActiveNative(AActor* MyTarget, FGameplayCueParameters Parameters) const
+{
+	return false;
+}
+
+bool UGameplayCueNotify_Static::OnRemove_Implementation(AActor* MyTarget, FGameplayCueParameters Parameters) const
+{
+	return OnRemoveNative(MyTarget, Parameters);
+}
+bool UGameplayCueNotify_Static::OnRemoveNative(AActor* MyTarget, FGameplayCueParameters Parameters) const
+{
+	return false;
+}
diff --git a/Engine/Source/Runtime/GameplayAbilities/Public/GameplayCueNotify_Actor.h b/Engine/Source/Runtime/GameplayAbilities/Public/GameplayCueNotify_Actor.h
--- a/Engine/Source/Runtime/GameplayAbilities/Public/GameplayCueNotify_Actor.h
+++ b/Engine/Source/Runtime/GameplayAbilities/Public/GameplayCueNotify_Actor.h
@@ -41,16 +41,19 @@ class GAMEPLAYABILITIES_API AGameplayCueNotify_Actor : public AActor
 
 	/** Generic Event Graph event that will get called for every event type */
 	UFUNCTION(BlueprintImplementableEvent, Category = "GameplayCueNotify", FriendlyName = "HandleGameplayCue")
-	virtual void K2_HandleGameplayCue(TWeakObjectPtr<AActor> MyTarget, EGameplayCueEvent::Type EventType, FGameplayCueParameters Parameters);
+	void K2_HandleGameplayCue(TWeakObjectPtr<AActor> MyTarget, EGameplayCueEvent::Type EventType, FGameplayCueParameters Parameters);
 
-	UFUNCTION(BlueprintImplementableEvent)
-	virtual bool OnExecute(TWeakObjectPtr<AActor> MyTarget, FGameplayCueParameters Parameters);
+	UFUNCTION(BlueprintNativeEvent, Category = "GameplayCueNotify")
+	bool OnExecute(AActor* MyTarget, FGameplayCueParameters Parameters);
+	virtual bool OnExecuteNative(AActor* MyTarget, FGameplayCueParameters Parameters);
 
-	UFUNCTION(BlueprintImplementableEvent)
-	virtual bool OnActive(TWeakObjectPtr<AActor> MyTarget, FGameplayCueParameters Parameters);
+	UFUNCTION(BlueprintNativeEvent, Category = "GameplayCueNotify")
+	bool OnActive(AActor* MyTarget, FGameplayCueParameters Parameters);
+	virtual bool OnActiveNative(AActor* MyTarget, FGameplayCueParameters Parameters);
 
-	UFUNCTION(BlueprintImplementableEvent)
-	virtual bool OnRemove(TWeakObjectPtr<AActor> MyTarget, FGameplayCueParameters Parameters);
+	UFUNCTION(BlueprintNativeEvent, Category = "GameplayCueNotify")
+	bool OnRemove(AActor* MyTarget, FGameplayCueParameters Parameters);
+	virtual bool OnRemoveNative(AActor* MyTarget, FGameplayCueParameters Parameters);
 
 	UPROPERTY(EditDefaultsOnly, Category = GameplayCue)
 	FGameplayTag	GameplayCueTag;
diff --git a/Engine/Source/Runtime/GameplayAbilities/Public/GameplayCueNotify_Static.h b/Engine/Source/Runtime/GameplayAbilities/Public/GameplayCueNotify_Static.h
--- a/Engine/Source/Runtime/GameplayAbilities/Public/GameplayCueNotify_Static.h
+++ b/Engine/Source/Runtime/GameplayAbilities/Public/GameplayCueNotify_Static.h
@@ -41,16 +41,19 @@ class GAMEPLAYABILITIES_API UGameplayCueNotify_Static : public UObject
 
 	/** Generic Event Graph event that will get called for every event type */
 	UFUNCTION(BlueprintImplementableEvent, Category = "GameplayCueNotify", FriendlyName = "HandleGameplayCue")
-	virtual void K2_HandleGameplayCue(TWeakObjectPtr<AActor> MyTarget, EGameplayCueEvent::Type EventType, FGameplayCueParameters Parameters) const;
+	void K2_HandleGameplayCue(TWeakObjectPtr<AActor> MyTarget, EGameplayCueEvent::Type EventType, FGameplayCueParameters Parameters) const;
 
-	UFUNCTION(BlueprintImplementableEvent, BlueprintPure, Category = "GameplayCueNotify")
-	virtual bool OnExecute(TWeakObjectPtr<AActor> MyTarget, FGameplayCueParameters Parameters) const;
+	UFUNCTION(BlueprintNativeEvent, BlueprintPure, Category = "GameplayCueNotify")
+	bool OnExecute(AActor* MyTarget, FGameplayCueParameters Parameters) const;
+	virtual bool OnExecuteNative(AActor* MyTarget, FGameplayCueParameters Parameters) const;
 
-	UFUNCTION(BlueprintImplementableEvent, BlueprintPure, Category = "GameplayCueNotify")
-	virtual bool OnActive(TWeakObjectPtr<AActor> MyTarget, FGameplayCueParameters Parameters) const;
+	UFUNCTION(BlueprintNativeEvent, BlueprintPure, Category = "GameplayCueNotify")
+	bool OnActive(AActor* MyTarget, FGameplayCueParameters Parameters) const;
+	virtual bool OnActiveNative(AActor* MyTarget, FGameplayCueParameters Parameters) const;
 
-	UFUNCTION(BlueprintImplementableEvent, BlueprintPure, Category = "GameplayCueNotify")
-	virtual bool OnRemove(TWeakObjectPtr<AActor> MyTarget, FGameplayCueParameters Parameters) const;
+	UFUNCTION(BlueprintNativeEvent, BlueprintPure, Category = "GameplayCueNotify")
+	bool OnRemove(AActor* MyTarget, FGameplayCueParameters Parameters) const;
+	virtual bool OnRemoveNative(AActor* MyTarget, FGameplayCueParameters Parameters) const;
 
 	UPROPERTY(EditDefaultsOnly, Category = GameplayCue)
 	FGameplayTag	GameplayCueTag;
