Trimmed down ActiveGameplayEffect scope-lock class.

diff --git a/Engine/Source/Runtime/GameplayAbilities/Private/GameplayEffect.cpp b/Engine/Source/Runtime/GameplayAbilities/Private/GameplayEffect.cpp
--- a/Engine/Source/Runtime/GameplayAbilities/Private/GameplayEffect.cpp
+++ b/Engine/Source/Runtime/GameplayAbilities/Private/GameplayEffect.cpp
@@ -1042,7 +1042,7 @@ FActiveGameplayEffect* FActiveGameplayEffectsContainer::GetActiveGameplayEffect(
 }
 
 int32 FScopedActiveGameplayEffectLock::AGELockCount = 0;
-FActiveGameplayEffectActionHandle FScopedActiveGameplayEffectLock::DeferredAGEActions;
+TArray<TSharedPtr<FActiveGameplayEffectAction>> FScopedActiveGameplayEffectLock::DeferredAGEActions;
 FScopedActiveGameplayEffectLock::FScopedActiveGameplayEffectLock()
 {
 	++AGELockCount;
@@ -1053,11 +1053,11 @@ FScopedActiveGameplayEffectLock::~FScopedActiveGameplayEffectLock()
 	--AGELockCount;
 	if (!IsLockInEffect())
 	{
-		for (int32 i = 0; i < DeferredAGEActions.Num(); ++i)
+		for (int32 i = 0; i < NumActions(); ++i)
 		{
-			DeferredAGEActions.Get(i)->PerformAction();
+			GetAction(i)->PerformAction();
 		}
-		DeferredAGEActions.Clear();
+		ClearActions();
 	}
 }
 
@@ -1590,20 +1590,19 @@ void FActiveGameplayEffectsContainer::AddActiveGameplayEffectGrantedTagsAndModif
 bool FActiveGameplayEffectsContainer::RemoveActiveGameplayEffect(FActiveGameplayEffectHandle Handle)
 {
 	// Could make this a map for quicker lookup
-	for(int32 Idx = 0; Idx < GameplayEffects.Num(); ++Idx)
+	for (int32 Idx = 0; Idx < GameplayEffects.Num(); ++Idx)
 	{
-		if(GameplayEffects[Idx].Handle == Handle)
+		if (GameplayEffects[Idx].Handle == Handle)
 		{
 			//Block removal if we're scope-locked
 			if (FScopedActiveGameplayEffectLock::IsLockInEffect())
 			{
-				FScopedActiveGameplayEffectLock::AddAction(new FActiveGameplayEffectAction_Remove(TWeakObjectPtr<UAbilitySystemComponent>(Owner), Handle));
+				FScopedActiveGameplayEffectLock::AddAction()->InitForRemoveGE(TWeakObjectPtr<UAbilitySystemComponent>(Owner), Handle);
 				return true;		//We are assuming the "ensure(Idx < GameplayEffects.Num())" would be passed.
 			}
 			return InternalRemoveActiveGameplayEffect(Idx);
 		}
 	}
-
 	ABILITY_LOG(Warning, TEXT("RemoveActiveGameplayEffect called with invalid Handle: %s"), *Handle.ToString());
 	return false;
 }
@@ -2052,7 +2051,7 @@ void FActiveGameplayEffectsContainer::RemoveActiveEffects(const FActiveGameplayE
 			// Defer removal if we're scope-locked
 			if (bLockInEffect)
 			{
-				FScopedActiveGameplayEffectLock::AddAction(new FActiveGameplayEffectAction_Remove(TWeakObjectPtr<UAbilitySystemComponent>(Owner), Effect.Handle));
+				FScopedActiveGameplayEffectLock::AddAction()->InitForRemoveGE(TWeakObjectPtr<UAbilitySystemComponent>(Owner), Effect.Handle);
 			}
 			else
 			{
@@ -2228,22 +2227,22 @@ void FInheritedTagContainer::RemoveTag(FGameplayTag TagToRemove)
 	CombinedTags.RemoveTag(TagToRemove);
 }
 
-void FActiveGameplayEffectAction_Remove::PerformAction()
-{
-	if (OwningASC.IsValid())
-	{
-		OwningASC.Get()->RemoveActiveGameplayEffect(Handle);
-	}
-}
-
-void FActiveGameplayEffectAction_Add::PerformAction()
+void FActiveGameplayEffectAction::PerformAction()
 {
+	check(bIsInitialized);
 	if (OwningASC.IsValid())
 	{
-		UAbilitySystemComponent* ASC = OwningASC.Get();
-		check(!ASC->ActiveGameplayEffects.GameplayEffectsPendingAdd.IsEmpty());
-		FActiveGameplayEffect TempEffect;
-		ASC->ActiveGameplayEffects.GameplayEffectsPendingAdd.Dequeue(TempEffect);
-		ASC->ActiveGameplayEffects.GameplayEffects.Add(TempEffect);
+		if (bRemove)
+		{
+			OwningASC.Get()->RemoveActiveGameplayEffect(Handle);
+		}
+		else
+		{
+			UAbilitySystemComponent* ASC = OwningASC.Get();
+			check(!ASC->ActiveGameplayEffects.GameplayEffectsPendingAdd.IsEmpty());
+			FActiveGameplayEffect TempEffect;
+			ASC->ActiveGameplayEffects.GameplayEffectsPendingAdd.Dequeue(TempEffect);
+			ASC->ActiveGameplayEffects.GameplayEffects.Add(TempEffect);
+		}
 	}
-}
+}
\ No newline at end of file
diff --git a/Engine/Source/Runtime/GameplayAbilities/Public/AbilitySystemComponent.h b/Engine/Source/Runtime/GameplayAbilities/Public/AbilitySystemComponent.h
--- a/Engine/Source/Runtime/GameplayAbilities/Public/AbilitySystemComponent.h
+++ b/Engine/Source/Runtime/GameplayAbilities/Public/AbilitySystemComponent.h
@@ -811,8 +811,9 @@ protected:
 	const UAttributeSet*	GetAttributeSubobjectChecked(const TSubclassOf<UAttributeSet> AttributeClass) const;
 	const UAttributeSet*	GetOrCreateAttributeSubobject(TSubclassOf<UAttributeSet> AttributeClass);
 
-	friend struct FActiveGameplayEffectsContainer;
 	friend struct FActiveGameplayEffect;
+	friend struct FActiveGameplayEffectAction;
+	friend struct FActiveGameplayEffectsContainer;
 	friend struct FActiveGameplayCue;
 	friend struct FActiveGameplayCueContainer;
 };
diff --git a/Engine/Source/Runtime/GameplayAbilities/Public/GameplayEffect.h b/Engine/Source/Runtime/GameplayAbilities/Public/GameplayEffect.h
--- a/Engine/Source/Runtime/GameplayAbilities/Public/GameplayEffect.h
+++ b/Engine/Source/Runtime/GameplayAbilities/Public/GameplayEffect.h
@@ -1280,118 +1280,79 @@ struct GAMEPLAYABILITIES_API FActiveGameplayEffectAction
 {
 	GENERATED_USTRUCT_BODY()
 	FActiveGameplayEffectAction()
+	: bIsInitialized(false)
 	{
 	}
 
-	virtual ~FActiveGameplayEffectAction()
+	void InitForAddGE(TWeakObjectPtr<UAbilitySystemComponent> InOwningASC)
 	{
+		OwningASC = InOwningASC;
+		bRemove = false;
+		bIsInitialized = true;
 	}
 
-	virtual void PerformAction()
+	void InitForRemoveGE(TWeakObjectPtr<UAbilitySystemComponent> InOwningASC, const FActiveGameplayEffectHandle& InHandle)
 	{
+		OwningASC = InOwningASC;
+		Handle = InHandle;
+		bRemove = true;
+		bIsInitialized = true;
 	}
-};
-
 
-USTRUCT()
-struct GAMEPLAYABILITIES_API FActiveGameplayEffectAction_Add : public FActiveGameplayEffectAction
-{
-	GENERATED_USTRUCT_BODY()
-	FActiveGameplayEffectAction_Add()
-	{
-	}
-
-	FActiveGameplayEffectAction_Add(TWeakObjectPtr<UAbilitySystemComponent> InOwningASC)
-	: OwningASC(InOwningASC)
-	{
-	}
-
-	virtual void PerformAction() override;
+	void PerformAction();
 
 	UPROPERTY()
 	TWeakObjectPtr<UAbilitySystemComponent> OwningASC;
-};
-
-USTRUCT()
-struct GAMEPLAYABILITIES_API FActiveGameplayEffectAction_Remove : public FActiveGameplayEffectAction
-{
-	GENERATED_USTRUCT_BODY()
-	FActiveGameplayEffectAction_Remove()
-	{
-	}
-
-	FActiveGameplayEffectAction_Remove(TWeakObjectPtr<UAbilitySystemComponent> InOwningASC, const FActiveGameplayEffectHandle& InHandle)
-	: OwningASC(InOwningASC)
-	, Handle(InHandle)
-	{
-	}
 
-	virtual void PerformAction() override;
+	/** This is needed for removals, and does not exist for additions. */
+	UPROPERTY()
+	FActiveGameplayEffectHandle Handle;
 
 	UPROPERTY()
-	TWeakObjectPtr<UAbilitySystemComponent> OwningASC;
+	bool bRemove;
 
 	UPROPERTY()
-	FActiveGameplayEffectHandle Handle;
+	bool bIsInitialized;
 };
 
-USTRUCT()
-struct GAMEPLAYABILITIES_API FActiveGameplayEffectActionHandle
-{
-	GENERATED_USTRUCT_BODY()
 
-	FActiveGameplayEffectActionHandle()
-	{
-	}
+struct GAMEPLAYABILITIES_API FScopedActiveGameplayEffectLock
+{
+	FScopedActiveGameplayEffectLock();
+	~FScopedActiveGameplayEffectLock();
 
-	FActiveGameplayEffectActionHandle(struct FActiveGameplayEffectAction* DataPtr)
-	{
-		Add(DataPtr);
-	}
+private:
+	static int32 AGELockCount;
+	static TArray<TSharedPtr<FActiveGameplayEffectAction>> DeferredAGEActions;
 
-	void Add(struct FActiveGameplayEffectAction* DataPtr)
+public:
+	static bool IsLockInEffect()
 	{
-		Data.Add(TSharedPtr<FActiveGameplayEffectAction>(DataPtr));
+		return (AGELockCount ? true : false);
 	}
 
-	void Clear()
+	static FActiveGameplayEffectAction* AddAction()
 	{
-		Data.Reset();
+		check(IsLockInEffect());
+		FActiveGameplayEffectAction* DataPtr = new FActiveGameplayEffectAction();
+		DeferredAGEActions.Add(TSharedPtr<FActiveGameplayEffectAction>(DataPtr));
+		return DataPtr;
 	}
 
-	int32 Num()
+	void ClearActions()
 	{
-		return Data.Num();
+		DeferredAGEActions.Reset();
 	}
 
-	FActiveGameplayEffectAction* Get(int32 Index)
+	int32 NumActions()
 	{
-		check((Index >= 0) && (Index < Data.Num()));
-		return Data[Index].Get();
+		return DeferredAGEActions.Num();
 	}
 
-	TArray<TSharedPtr<FActiveGameplayEffectAction>>	Data;
-};
-
-struct GAMEPLAYABILITIES_API FScopedActiveGameplayEffectLock
-{
-	FScopedActiveGameplayEffectLock();
-	~FScopedActiveGameplayEffectLock();
-
-private:
-	static int32 AGELockCount;
-	static FActiveGameplayEffectActionHandle DeferredAGEActions;
-
-public:
-	static void AddAction(FActiveGameplayEffectAction* NewAction)
+	FActiveGameplayEffectAction* GetAction(int32 Index)
 	{
-		check(IsLockInEffect());
-		DeferredAGEActions.Add(NewAction);
-	}
-
-	static bool IsLockInEffect()
-	{
-		return (AGELockCount ? true : false);
+		check((Index >= 0) && (Index < DeferredAGEActions.Num()));
+		return DeferredAGEActions[Index].Get();
 	}
 };
 
