Fixed a bug with map geometry being detected as actors and blocking ground placement.

diff --git a/Engine/Source/Runtime/GameplayAbilities/Private/Abilities/GameplayAbilityTargetActor_GroundTrace.cpp b/Engine/Source/Runtime/GameplayAbilities/Private/Abilities/GameplayAbilityTargetActor_GroundTrace.cpp
--- a/Engine/Source/Runtime/GameplayAbilities/Private/Abilities/GameplayAbilityTargetActor_GroundTrace.cpp
+++ b/Engine/Source/Runtime/GameplayAbilities/Private/Abilities/GameplayAbilityTargetActor_GroundTrace.cpp
@@ -84,7 +84,7 @@ bool AGameplayAbilityTargetActor_GroundTrace::AdjustCollisionResultForShape(cons
 		SweepWithFilter(LocalResult, ThisWorld, Filter, TraceStart, TraceEnd, FQuat::Identity, TraceChannel, CollisionShape, Params);
 		if (!LocalResult.bStartPenetrating)
 		{
-			if (!LocalResult.bBlockingHit || LocalResult.Actor.IsValid())
+			if (!LocalResult.bBlockingHit || (LocalResult.Actor.IsValid() && Cast<APawn>(LocalResult.Actor.Get())))
 			{
 				//Off the map, or hit an actor
 				if (bDebug)
