- adding in cheats, God and BadAss
- adding in Owning Tags and Effect Tags to the FActiveGameplayEffect
- fixing bugs with XP not giving XP to killers and not adding in assisters correctly

diff --git a/Engine/Source/Runtime/GameplayAbilities/Private/GameplayEffect.cpp b/Engine/Source/Runtime/GameplayAbilities/Private/GameplayEffect.cpp
--- a/Engine/Source/Runtime/GameplayAbilities/Private/GameplayEffect.cpp
+++ b/Engine/Source/Runtime/GameplayAbilities/Private/GameplayEffect.cpp
@@ -2147,15 +2147,34 @@ void FActiveGameplayEffectHandle::RemoveFromGlobalMap()
 
 bool FActiveGameplayEffectQuery::Matches(const FActiveGameplayEffect& Effect) const
 {
-	if (TagContainer)
+	bool bFoundMatches = false;
+
+	// if we are looking for owner tags check InheritableOwnedTagsContainer with OwningTagContainer
+	if (OwningTagContainer)
 	{
-		if (!Effect.Spec.Def->InheritableOwnedTagsContainer.CombinedTags.MatchesAny(*TagContainer, false) &&
-			!Effect.Spec.DynamicGrantedTags.MatchesAny(*TagContainer, false))
+		if (Effect.Spec.Def->InheritableOwnedTagsContainer.CombinedTags.MatchesAny(*OwningTagContainer, false) ||
+			Effect.Spec.DynamicGrantedTags.MatchesAny(*OwningTagContainer, false))
 		{
-			return false;
+			bFoundMatches = true;
 		}
 	}	
-	
+
+	// if we are looking for effect tags check InheritableGameplayEffectTags with EffectTagContainer
+	if (EffectTagContainer)
+	{
+		if (Effect.Spec.Def->InheritableGameplayEffectTags.CombinedTags.MatchesAny(*EffectTagContainer, false) ||
+			Effect.Spec.DynamicGrantedTags.MatchesAny(*EffectTagContainer, false))
+		{
+			bFoundMatches = true;
+		}
+	}
+
+	// if we didn't fine any tag matches this is our bail point.
+	if (!bFoundMatches)
+	{
+		return false;
+	}
+
 	if (ModifyingAttribute.IsValid())
 	{
 		bool FailedModifyingAttributeCheck = true;
diff --git a/Engine/Source/Runtime/GameplayAbilities/Public/GameplayEffect.h b/Engine/Source/Runtime/GameplayAbilities/Public/GameplayEffect.h
--- a/Engine/Source/Runtime/GameplayAbilities/Public/GameplayEffect.h
+++ b/Engine/Source/Runtime/GameplayAbilities/Public/GameplayEffect.h
@@ -1072,19 +1072,32 @@ struct FActiveGameplayEffectQuery
 	GENERATED_USTRUCT_BODY()
 
 	FActiveGameplayEffectQuery()
-		: TagContainer(NULL)
+		: OwningTagContainer(NULL)
+		, EffectTagContainer(NULL)
 	{
 	}
 
-	FActiveGameplayEffectQuery(const FGameplayTagContainer* InTagContainer, FGameplayAttribute InModifyingAttribute = FGameplayAttribute())
-		: TagContainer(InTagContainer)
+	FActiveGameplayEffectQuery(const FGameplayTagContainer* InOwningTagContainer, FGameplayAttribute InModifyingAttribute = FGameplayAttribute())
+		: OwningTagContainer(InOwningTagContainer)
+		, EffectTagContainer(NULL)
+		, ModifyingAttribute(InModifyingAttribute)
+	{
+	}
+
+	FActiveGameplayEffectQuery(const FGameplayTagContainer* InOwningTagContainer, const FGameplayTagContainer* InEffectTagContainer, FGameplayAttribute InModifyingAttribute = FGameplayAttribute())
+		: OwningTagContainer(InOwningTagContainer)
+		, EffectTagContainer(InEffectTagContainer)
 		, ModifyingAttribute(InModifyingAttribute)
 	{
 	}
 
 	bool Matches(const FActiveGameplayEffect& Effect) const;
 
-	const FGameplayTagContainer* TagContainer;
+	/** used to match with InheritableOwnedTagsContainer */
+	const FGameplayTagContainer* OwningTagContainer;
+
+	/** used to match with InheritableGameplayEffectTags */
+	const FGameplayTagContainer* EffectTagContainer;
 
 	/** Matches on GameplayEffects which modify given attribute */
 	FGameplayAttribute ModifyingAttribute;
