fixing logic bugs that I introdiced

diff --git a/Engine/Source/Runtime/GameplayAbilities/Private/GameplayEffect.cpp b/Engine/Source/Runtime/GameplayAbilities/Private/GameplayEffect.cpp
--- a/Engine/Source/Runtime/GameplayAbilities/Private/GameplayEffect.cpp
+++ b/Engine/Source/Runtime/GameplayAbilities/Private/GameplayEffect.cpp
@@ -2147,34 +2147,28 @@ void FActiveGameplayEffectHandle::RemoveFromGlobalMap()
 
 bool FActiveGameplayEffectQuery::Matches(const FActiveGameplayEffect& Effect) const
 {
-	bool bFoundMatches = false;
-
-	// if we are looking for owner tags check InheritableOwnedTagsContainer with OwningTagContainer
+	// if we are looking for owning tags check them on the Granted Tags and Owned Tags Container
 	if (OwningTagContainer)
 	{
-		if (Effect.Spec.Def->InheritableOwnedTagsContainer.CombinedTags.MatchesAny(*OwningTagContainer, false) ||
-			Effect.Spec.DynamicGrantedTags.MatchesAny(*OwningTagContainer, false))
+		if (!Effect.Spec.Def->InheritableOwnedTagsContainer.CombinedTags.MatchesAny(*OwningTagContainer, false) &&
+			!Effect.Spec.DynamicGrantedTags.MatchesAny(*OwningTagContainer, false))
 		{
-			bFoundMatches = true;
+			// this doesnt match our Tags so bail
+			return false;
 		}
 	}	
-
-	// if we are looking for effect tags check InheritableGameplayEffectTags with EffectTagContainer
+	
+	// if we are just looking for Tags on the Effect then look at the Gameplay Effect Tags
 	if (EffectTagContainer)
 	{
-		if (Effect.Spec.Def->InheritableGameplayEffectTags.CombinedTags.MatchesAny(*EffectTagContainer, false) ||
-			Effect.Spec.DynamicGrantedTags.MatchesAny(*EffectTagContainer, false))
+		if (!Effect.Spec.Def->InheritableGameplayEffectTags.CombinedTags.MatchesAny(*EffectTagContainer, false))
 		{
-			bFoundMatches = true;
+			// this doesnt match our Tags so bail
+			return false;
 		}
 	}
 
-	// if we didn't fine any tag matches this is our bail point.
-	if (!bFoundMatches)
-	{
-		return false;
-	}
-
+	// if we are looking for ModifyingAttribute go over each of the Spec Modifiers and check the Attributes
 	if (ModifyingAttribute.IsValid())
 	{
 		bool FailedModifyingAttributeCheck = true;
