Target actor tracing;

* Switched all target actors that do traces to use a collision profile instead of a collision channel.

* Optimized traces done in GameplayAbilityTargetActor_Trace to only trace once with a multi hit instead of adding ignore actors and re-tracing in a loop.

* Fixed SGraphPinCollisionProfile not getting the correct default from the CDO.

diff --git a/Engine/Source/Editor/BlueprintGraph/Private/EdGraphSchema_K2.cpp b/Engine/Source/Editor/BlueprintGraph/Private/EdGraphSchema_K2.cpp
--- a/Engine/Source/Editor/BlueprintGraph/Private/EdGraphSchema_K2.cpp
+++ b/Engine/Source/Editor/BlueprintGraph/Private/EdGraphSchema_K2.cpp
@@ -1166,7 +1166,8 @@ bool UEdGraphSchema_K2::PinDefaultValueIsEditable(const UEdGraphPin& InGraphPin)
 			|| &SubCategoryObject == RotatorStruct
 			|| &SubCategoryObject == TransformStruct
 			|| &SubCategoryObject == LinearColorStruct
-			|| &SubCategoryObject == ColorStruct;
+			|| &SubCategoryObject == ColorStruct
+			|| &SubCategoryObject == FCollisionProfileName::StaticStruct();
 	}
 
 	return true;
diff --git a/Engine/Source/Runtime/GameplayAbilities/Private/Abilities/GameplayAbilityTargetActor_GroundTrace.cpp b/Engine/Source/Runtime/GameplayAbilities/Private/Abilities/GameplayAbilityTargetActor_GroundTrace.cpp
--- a/Engine/Source/Runtime/GameplayAbilities/Private/Abilities/GameplayAbilityTargetActor_GroundTrace.cpp
+++ b/Engine/Source/Runtime/GameplayAbilities/Private/Abilities/GameplayAbilityTargetActor_GroundTrace.cpp
@@ -81,7 +81,7 @@ bool AGameplayAbilityTargetActor_GroundTrace::AdjustCollisionResultForShape(cons
 	{
 		TraceEnd = TraceStart = OriginalEndPoint - (LerpValue * Movement);
 		TraceEnd.Z -= 99999.0f;
-		SweepWithFilter(LocalResult, ThisWorld, Filter, TraceStart, TraceEnd, FQuat::Identity, TraceChannel, CollisionShape, Params);
+		SweepWithFilter(LocalResult, ThisWorld, Filter, TraceStart, TraceEnd, FQuat::Identity, CollisionShape, TraceProfile.Name, Params);
 		if (!LocalResult.bStartPenetrating)
 		{
 			if (!LocalResult.bBlockingHit || (LocalResult.Actor.IsValid() && Cast<APawn>(LocalResult.Actor.Get())))
@@ -150,7 +150,7 @@ FHitResult AGameplayAbilityTargetActor_GroundTrace::PerformTrace(AActor* InSourc
 
 	FHitResult ReturnHitResult;
 	//Use a line trace initially to see where the player is actually pointing
-	LineTraceWithFilter(ReturnHitResult, InSourceActor->GetWorld(), Filter, TraceStart, TraceEnd, TraceChannel, Params);
+	LineTraceWithFilter(ReturnHitResult, InSourceActor->GetWorld(), Filter, TraceStart, TraceEnd, TraceProfile.Name, Params);
 	//Default to end of trace line if we don't hit anything.
 	if (!ReturnHitResult.bBlockingHit)
 	{
@@ -162,7 +162,7 @@ FHitResult AGameplayAbilityTargetActor_GroundTrace::PerformTrace(AActor* InSourc
 	TraceEnd = TraceStart;
 	TraceStart.Z += CollisionHeightOffset;
 	TraceEnd.Z -= 99999.0f;
-	LineTraceWithFilter(ReturnHitResult, InSourceActor->GetWorld(), Filter, TraceStart, TraceEnd, TraceChannel, Params);
+	LineTraceWithFilter(ReturnHitResult, InSourceActor->GetWorld(), Filter, TraceStart, TraceEnd, TraceProfile.Name, Params);
 	//if (!ReturnHitResult.bBlockingHit) then our endpoint may be off the map. Hopefully this is only possible in debug maps.
 
 	bLastTraceWasGood = true;		//So far, we're good. If we need a ground spot and can't find one, we'll come back.
diff --git a/Engine/Source/Runtime/GameplayAbilities/Private/Abilities/GameplayAbilityTargetActor_SingleLineTrace.cpp b/Engine/Source/Runtime/GameplayAbilities/Private/Abilities/GameplayAbilityTargetActor_SingleLineTrace.cpp
--- a/Engine/Source/Runtime/GameplayAbilities/Private/Abilities/GameplayAbilityTargetActor_SingleLineTrace.cpp
+++ b/Engine/Source/Runtime/GameplayAbilities/Private/Abilities/GameplayAbilityTargetActor_SingleLineTrace.cpp
@@ -37,7 +37,7 @@ FHitResult AGameplayAbilityTargetActor_SingleLineTrace::PerformTrace(AActor* InS
 	// ------------------------------------------------------
 
 	FHitResult ReturnHitResult;
-	LineTraceWithFilter(ReturnHitResult, InSourceActor->GetWorld(), Filter, TraceStart, TraceEnd, TraceChannel, Params);
+	LineTraceWithFilter(ReturnHitResult, InSourceActor->GetWorld(), Filter, TraceStart, TraceEnd, TraceProfile.Name, Params);
 	//Default to end of trace line if we don't hit anything.
 	if (!ReturnHitResult.bBlockingHit)
 	{
diff --git a/Engine/Source/Runtime/GameplayAbilities/Private/Abilities/GameplayAbilityTargetActor_Trace.cpp b/Engine/Source/Runtime/GameplayAbilities/Private/Abilities/GameplayAbilityTargetActor_Trace.cpp
--- a/Engine/Source/Runtime/GameplayAbilities/Private/Abilities/GameplayAbilityTargetActor_Trace.cpp
+++ b/Engine/Source/Runtime/GameplayAbilities/Private/Abilities/GameplayAbilityTargetActor_Trace.cpp
@@ -20,7 +20,6 @@ AGameplayAbilityTargetActor_Trace::AGameplayAbilityTargetActor_Trace(const FObje
 	StaticTargetFunction = false;
 
 	MaxRange = 999999.0f;
-	TraceChannel = ECC_WorldStatic;
 }
 
 void AGameplayAbilityTargetActor_Trace::EndPlay(const EEndPlayReason::Type EndPlayReason)
@@ -33,42 +32,50 @@ void AGameplayAbilityTargetActor_Trace::EndPlay(const EEndPlayReason::Type EndPl
 	Super::EndPlay(EndPlayReason);
 }
 
-void AGameplayAbilityTargetActor_Trace::LineTraceWithFilter(FHitResult& ReturnHitResult, const UWorld* InWorld, const FGameplayTargetDataFilterHandle InFilterHandle, const FVector& InTraceStart, const FVector& InTraceEnd, ECollisionChannel Channel, const FCollisionQueryParams Params) const
+void AGameplayAbilityTargetActor_Trace::LineTraceWithFilter(FHitResult& ReturnHitResult, const UWorld* InWorld, const FGameplayTargetDataFilterHandle InFilterHandle, const FVector& InTraceStart, const FVector& InTraceEnd, FName TraceProfileName, const FCollisionQueryParams Params) const
 {
 	check(InWorld);
-	FCollisionQueryParams LocalParams = Params;
-	while (true)
+
+	TArray<FHitResult> HitResults;
+	InWorld->LineTraceMultiByProfile(HitResults, InTraceStart, InTraceEnd, TraceProfileName, Params);
+
+	ReturnHitResult.TraceStart = InTraceStart;
+	ReturnHitResult.TraceEnd = InTraceEnd;
+
+	for (int32 HitIdx = 0; HitIdx < HitResults.Num(); ++HitIdx)
 	{
-		FHitResult TempHitResult;
-		InWorld->LineTraceSingle(TempHitResult, InTraceStart, InTraceEnd, Channel, LocalParams);
-		if (TempHitResult.bBlockingHit && TempHitResult.Actor.IsValid() && !InFilterHandle.FilterPassesForActor(TempHitResult.Actor))
+		const FHitResult& Hit = HitResults[HitIdx];
+
+		if (!Hit.Actor.IsValid() || InFilterHandle.FilterPassesForActor(Hit.Actor))
 		{
-			LocalParams.AddIgnoredActor(TempHitResult.Actor.Get());
-			continue;
+			ReturnHitResult = Hit;
+			ReturnHitResult.bBlockingHit = true; // treat it as a blocking hit
+			return;
 		}
-		//Either hit something we're not ignoring, or didn't hit anything.
-		ReturnHitResult = TempHitResult;
-		break;
-	};
+	}
 }
 
-void AGameplayAbilityTargetActor_Trace::SweepWithFilter(FHitResult& ReturnHitResult, const UWorld* InWorld, const FGameplayTargetDataFilterHandle InFilterHandle, const FVector& InTraceStart, const FVector& InTraceEnd, const FQuat& InRotation, ECollisionChannel Channel, const FCollisionShape CollisionShape, const FCollisionQueryParams Params) const
+void AGameplayAbilityTargetActor_Trace::SweepWithFilter(FHitResult& ReturnHitResult, const UWorld* InWorld, const FGameplayTargetDataFilterHandle InFilterHandle, const FVector& InTraceStart, const FVector& InTraceEnd, const FQuat& InRotation, const FCollisionShape CollisionShape, FName TraceProfileName, const FCollisionQueryParams Params) const
 {
 	check(InWorld);
-	FCollisionQueryParams LocalParams = Params;
-	while (true)
+
+	TArray<FHitResult> HitResults;
+	InWorld->SweepMultiByProfile(HitResults, InTraceStart, InTraceEnd, InRotation, TraceProfileName, CollisionShape, Params);
+
+	ReturnHitResult.TraceStart = InTraceStart;
+	ReturnHitResult.TraceEnd = InTraceEnd;
+
+	for (int32 HitIdx = 0; HitIdx < HitResults.Num(); ++HitIdx)
 	{
-		FHitResult TempHitResult;
-		InWorld->SweepSingle(TempHitResult, InTraceStart, InTraceEnd, InRotation, Channel, CollisionShape, LocalParams);
-		if (TempHitResult.bBlockingHit && TempHitResult.Actor.IsValid() && !InFilterHandle.FilterPassesForActor(TempHitResult.Actor))
+		const FHitResult& Hit = HitResults[HitIdx];
+
+		if (!Hit.Actor.IsValid() || InFilterHandle.FilterPassesForActor(Hit.Actor))
 		{
-			LocalParams.AddIgnoredActor(ReturnHitResult.Actor.Get());
-			continue;
+			ReturnHitResult = Hit;
+			ReturnHitResult.bBlockingHit = true; // treat it as a blocking hit
+			return;
 		}
-		//Either hit something we're not ignoring, or didn't hit anything.
-		ReturnHitResult = TempHitResult;
-		break;
-	};
+	}
 }
 
 FGameplayAbilityTargetDataHandle AGameplayAbilityTargetActor_Trace::StaticGetTargetData(UWorld * World, const FGameplayAbilityActorInfo* ActorInfo, FGameplayAbilityActivationInfo ActivationInfo) const
@@ -94,7 +101,7 @@ void AGameplayAbilityTargetActor_Trace::AimWithPlayerController(AActor* InSource
 		ClipCameraRayToAbilityRange(CamLoc, CamDir, TraceStart, MaxRange, CamTarget);
 
 		FHitResult TempHitResult;
-		LineTraceWithFilter(TempHitResult, InSourceActor->GetWorld(), Filter, CamLoc, CamTarget, TraceChannel, Params);
+		LineTraceWithFilter(TempHitResult, InSourceActor->GetWorld(), Filter, CamLoc, CamTarget, TraceProfile.Name, Params);
 		if (TempHitResult.bBlockingHit && (FVector::DistSquared(TraceStart, TempHitResult.Location) <= (MaxRange * MaxRange)))
 		{
 			//We actually made a hit? Pull back.
diff --git a/Engine/Source/Runtime/GameplayAbilities/Public/Abilities/GameplayAbilityTargetActor_GroundTrace.h b/Engine/Source/Runtime/GameplayAbilities/Public/Abilities/GameplayAbilityTargetActor_GroundTrace.h
--- a/Engine/Source/Runtime/GameplayAbilities/Public/Abilities/GameplayAbilityTargetActor_GroundTrace.h
+++ b/Engine/Source/Runtime/GameplayAbilities/Public/Abilities/GameplayAbilityTargetActor_GroundTrace.h
@@ -13,10 +13,6 @@ class GAMEPLAYABILITIES_API AGameplayAbilityTargetActor_GroundTrace : public AGa
 
 	virtual void StartTargeting(UGameplayAbility* InAbility) override;
 
-	/** This trace channel is used when finding the ground, and may be different from the direct line trace channel */
-	//UPROPERTY(BlueprintReadWrite, EditAnywhere, meta = (ExposeOnSpawn = true), Category = Trace)
-	//TEnumAsByte<ECollisionChannel> TraceChannelToGround;
-
 	/** Radius for a sphere or capsule. */
 	UPROPERTY(BlueprintReadWrite, EditAnywhere, meta = (ExposeOnSpawn = true), Category = Targeting)
 	float CollisionRadius;
diff --git a/Engine/Source/Runtime/GameplayAbilities/Public/Abilities/GameplayAbilityTargetActor_Trace.h b/Engine/Source/Runtime/GameplayAbilities/Public/Abilities/GameplayAbilityTargetActor_Trace.h
--- a/Engine/Source/Runtime/GameplayAbilities/Public/Abilities/GameplayAbilityTargetActor_Trace.h
+++ b/Engine/Source/Runtime/GameplayAbilities/Public/Abilities/GameplayAbilityTargetActor_Trace.h
@@ -7,7 +7,7 @@
 #include "GameplayAbilityTargetActor_Trace.generated.h"
 
 /** Intermediate base class for all line-trace type targeting actors. */
-UCLASS(Abstract, Blueprintable, notplaceable)
+UCLASS(Abstract, Blueprintable, notplaceable, config=Game)
 class GAMEPLAYABILITIES_API AGameplayAbilityTargetActor_Trace : public AGameplayAbilityTargetActor
 {
 	GENERATED_UCLASS_BODY()
@@ -16,10 +16,10 @@ public:
 	virtual void EndPlay(const EEndPlayReason::Type EndPlayReason) override;
 
 	/** Traces as normal, but will manually filter all hit actors */
-	void LineTraceWithFilter(FHitResult& ReturnHitResult, const UWorld* InWorld, const FGameplayTargetDataFilterHandle InFilterHandle, const FVector& InTraceStart, const FVector& InTraceEnd, ECollisionChannel Channel, const FCollisionQueryParams Params) const;
+	void LineTraceWithFilter(FHitResult& ReturnHitResult, const UWorld* InWorld, const FGameplayTargetDataFilterHandle InFilterHandle, const FVector& InTraceStart, const FVector& InTraceEnd, FName TraceProfileName, const FCollisionQueryParams Params) const;
 
 	/** Sweeps as normal, but will manually filter all hit actors */
-	void SweepWithFilter(FHitResult& ReturnHitResult, const UWorld* InWorld, const FGameplayTargetDataFilterHandle InFilterHandle, const FVector& InTraceStart, const FVector& InTraceEnd, const FQuat& InRotation, ECollisionChannel Channel, const FCollisionShape CollisionShape, const FCollisionQueryParams Params) const;
+	void SweepWithFilter(FHitResult& ReturnHitResult, const UWorld* InWorld, const FGameplayTargetDataFilterHandle InFilterHandle, const FVector& InTraceStart, const FVector& InTraceEnd, const FQuat& InRotation, const FCollisionShape CollisionShape, FName TraceProfileName, const FCollisionQueryParams Params) const;
 
 	virtual FGameplayAbilityTargetDataHandle StaticGetTargetData(UWorld * World, const FGameplayAbilityActorInfo* ActorInfo, FGameplayAbilityActivationInfo ActivationInfo) const override;
 
@@ -36,8 +36,8 @@ public:
 	UPROPERTY(BlueprintReadWrite, EditAnywhere, meta = (ExposeOnSpawn = true), Category = Projectile)
 	float MaxRange;
 
-	UPROPERTY(BlueprintReadWrite, EditAnywhere, meta = (ExposeOnSpawn = true), Category = Trace)
-	TEnumAsByte<ECollisionChannel> TraceChannel;
+	UPROPERTY(BlueprintReadWrite, EditAnywhere, config, meta = (ExposeOnSpawn = true), Category = Trace)
+	FCollisionProfileName TraceProfile;
 
 protected:
 	virtual FHitResult PerformTrace(AActor* InSourceActor) PURE_VIRTUAL(AGameplayAbilityTargetActor_Trace, return FHitResult(););
