#UE4 Back out changelist 2378164. Main branch does not yet have other TriggerEventData changes. This will be re-introduced to the main branch during the dev->main merge

diff --git a/Engine/Source/Runtime/GameplayAbilities/Private/AbilitySystemComponent_Abilities.cpp b/Engine/Source/Runtime/GameplayAbilities/Private/AbilitySystemComponent_Abilities.cpp
--- a/Engine/Source/Runtime/GameplayAbilities/Private/AbilitySystemComponent_Abilities.cpp
+++ b/Engine/Source/Runtime/GameplayAbilities/Private/AbilitySystemComponent_Abilities.cpp
@@ -583,14 +583,7 @@ bool UAbilitySystemComponent::TryActivateAbility(FGameplayAbilitySpecHandle Hand
 	}
 	else if (Ability->GetNetExecutionPolicy() == EGameplayAbilityNetExecutionPolicy::Server)
 	{
-		if (TriggerEventData)
-		{
-			ServerTryActivateAbilityWithEventData(Handle, Spec->InputPressed, FPredictionKey(), *TriggerEventData);
-		}
-		else
-		{
-			ServerTryActivateAbility(Handle, Spec->InputPressed, FPredictionKey());
-		}
+		ServerTryActivateAbility(Handle, Spec->InputPressed, FPredictionKey());
 	}
 	else if (Ability->GetNetExecutionPolicy() == EGameplayAbilityNetExecutionPolicy::Predictive)
 	{
@@ -600,14 +593,7 @@ bool UAbilitySystemComponent::TryActivateAbility(FGameplayAbilitySpecHandle Hand
 		ActivationInfo.SetPredicting(ScopedPredictionKey);
 		
 		// This must be called immediately after GeneratePredictionKey to prevent problems with recursively activating abilities
-		if (TriggerEventData)
-		{
-			ServerTryActivateAbilityWithEventData(Handle, Spec->InputPressed, ScopedPredictionKey, *TriggerEventData);
-		}
-		else
-		{
-			ServerTryActivateAbility(Handle, Spec->InputPressed, ScopedPredictionKey);
-		}
+		ServerTryActivateAbility(Handle, Spec->InputPressed, ScopedPredictionKey);
 
 		// If this PredictionKey is rejected, we will call OnClientActivateAbilityFailed.
 		ScopedPredictionKey.NewRejectedDelegate().BindUObject(this, &UAbilitySystemComponent::OnClientActivateAbilityFailed, Handle, ScopedPredictionKey.Current);
@@ -644,26 +630,6 @@ bool UAbilitySystemComponent::TryActivateAbility(FGameplayAbilitySpecHandle Hand
 }
 
 void UAbilitySystemComponent::ServerTryActivateAbility_Implementation(FGameplayAbilitySpecHandle Handle, bool InputPressed, FPredictionKey PredictionKey)
-{
-	InternalServerTryActiveAbility(Handle, InputPressed, PredictionKey, nullptr);
-}
-
-bool UAbilitySystemComponent::ServerTryActivateAbility_Validate(FGameplayAbilitySpecHandle Handle, bool InputPressed, FPredictionKey PredictionKey)
-{
-	return true;
-}
-
-void UAbilitySystemComponent::ServerTryActivateAbilityWithEventData_Implementation(FGameplayAbilitySpecHandle Handle, bool InputPressed, FPredictionKey PredictionKey, FGameplayEventData TriggerEventData)
-{
-	InternalServerTryActiveAbility(Handle, InputPressed, PredictionKey, &TriggerEventData);
-}
-
-bool UAbilitySystemComponent::ServerTryActivateAbilityWithEventData_Validate(FGameplayAbilitySpecHandle Handle, bool InputPressed, FPredictionKey PredictionKey, FGameplayEventData TriggerEventData)
-{
-	return true;
-}
-
-void UAbilitySystemComponent::InternalServerTryActiveAbility(FGameplayAbilitySpecHandle Handle, bool InputPressed, const FPredictionKey& PredictionKey, const FGameplayEventData* TriggerEventData)
 {
 #if !(UE_BUILD_SHIPPING || UE_BUILD_TEST)
 	if (DenyClientActivation > 0)
@@ -731,7 +697,7 @@ void UAbilitySystemComponent::InternalServerTryActiveAbility(FGameplayAbilitySpe
 	Spec->InputPressed = InputPressed; // Tasks that check input may need this to be right immediately on ability startup.
 
 	// Attempt to activate the ability (server side) and tell the client if it succeeded or failed.
-	if (TryActivateAbility(Handle, PredictionKey, &InstancedAbility, nullptr, TriggerEventData))
+	if (TryActivateAbility(Handle, PredictionKey, &InstancedAbility))
 	{
 		ClientActivateAbilitySucceed(Handle, PredictionKey.Current);
 		//Client commands to end the ability that come in after this point are considered for this instance
@@ -747,6 +713,11 @@ void UAbilitySystemComponent::InternalServerTryActiveAbility(FGameplayAbilitySpe
 	}
 }
 
+bool UAbilitySystemComponent::ServerTryActivateAbility_Validate(FGameplayAbilitySpecHandle Handle, bool InputPressed, FPredictionKey PredictionKey)
+{
+	return true;
+}
+
 //This is only called when ending an ability in response to a remote instruction.
 void UAbilitySystemComponent::EndAbility(FGameplayAbilitySpecHandle AbilityToEnd)
 {
diff --git a/Engine/Source/Runtime/GameplayAbilities/Public/AbilitySystemComponent.h b/Engine/Source/Runtime/GameplayAbilities/Public/AbilitySystemComponent.h
--- a/Engine/Source/Runtime/GameplayAbilities/Public/AbilitySystemComponent.h
+++ b/Engine/Source/Runtime/GameplayAbilities/Public/AbilitySystemComponent.h
@@ -498,9 +498,6 @@ class GAMEPLAYABILITIES_API UAbilitySystemComponent : public UActorComponent, pu
 	UFUNCTION(Server, reliable, WithValidation)
 	void	ServerTryActivateAbility(FGameplayAbilitySpecHandle AbilityToActivate, bool InputPressed, FPredictionKey PredictionKey);
 
-	UFUNCTION(Server, reliable, WithValidation)
-	void	ServerTryActivateAbilityWithEventData(FGameplayAbilitySpecHandle AbilityToActivate, bool InputPressed, FPredictionKey PredictionKey, FGameplayEventData TriggerEventData);
-
 	/** Called by ServerEndAbility and ClientEndAbility; avoids code duplication. */
 	void EndAbility(FGameplayAbilitySpecHandle AbilityToEnd);
 
@@ -650,9 +647,6 @@ class GAMEPLAYABILITIES_API UAbilitySystemComponent : public UActorComponent, pu
 
 protected:
 
-	/** Implementation of ServerTryActivateAbility */
-	virtual void InternalServerTryActiveAbility(FGameplayAbilitySpecHandle AbilityToActivate, bool InputPressed, const FPredictionKey& PredictionKey, const FGameplayEventData* TriggerEventData);
-
 	/** Called when a prediction key that played a montage is rejected */
 	void OnPredictiveMontageRejected(UAnimMontage* PredictiveMontage);
 
