Target actor trace aiming;

* Targeting actors that trace now have an optional bool that makes the trace not affect the aiming pitch.

diff --git a/Engine/Source/Runtime/GameplayAbilities/Private/Abilities/GameplayAbilityTargetActor_Trace.cpp b/Engine/Source/Runtime/GameplayAbilities/Private/Abilities/GameplayAbilityTargetActor_Trace.cpp
--- a/Engine/Source/Runtime/GameplayAbilities/Private/Abilities/GameplayAbilityTargetActor_Trace.cpp
+++ b/Engine/Source/Runtime/GameplayAbilities/Private/Abilities/GameplayAbilityTargetActor_Trace.cpp
@@ -20,6 +20,7 @@ AGameplayAbilityTargetActor_Trace::AGameplayAbilityTargetActor_Trace(const FObje
 	StaticTargetFunction = false;
 
 	MaxRange = 999999.0f;
+	bTraceAffectsAimPitch = true;
 }
 
 void AGameplayAbilityTargetActor_Trace::EndPlay(const EEndPlayReason::Type EndPlayReason)
@@ -32,47 +33,47 @@ void AGameplayAbilityTargetActor_Trace::EndPlay(const EEndPlayReason::Type EndPl
 	Super::EndPlay(EndPlayReason);
 }
 
-void AGameplayAbilityTargetActor_Trace::LineTraceWithFilter(FHitResult& ReturnHitResult, const UWorld* InWorld, const FGameplayTargetDataFilterHandle InFilterHandle, const FVector& InTraceStart, const FVector& InTraceEnd, FName TraceProfileName, const FCollisionQueryParams Params) const
+void AGameplayAbilityTargetActor_Trace::LineTraceWithFilter(FHitResult& OutHitResult, const UWorld* World, const FGameplayTargetDataFilterHandle FilterHandle, const FVector& Start, const FVector& End, FName ProfileName, const FCollisionQueryParams Params)
 {
-	check(InWorld);
+	check(World);
 
 	TArray<FHitResult> HitResults;
-	InWorld->LineTraceMultiByProfile(HitResults, InTraceStart, InTraceEnd, TraceProfileName, Params);
+	World->LineTraceMultiByProfile(HitResults, Start, End, ProfileName, Params);
 
-	ReturnHitResult.TraceStart = InTraceStart;
-	ReturnHitResult.TraceEnd = InTraceEnd;
+	OutHitResult.TraceStart = Start;
+	OutHitResult.TraceEnd = End;
 
 	for (int32 HitIdx = 0; HitIdx < HitResults.Num(); ++HitIdx)
 	{
 		const FHitResult& Hit = HitResults[HitIdx];
 
-		if (!Hit.Actor.IsValid() || InFilterHandle.FilterPassesForActor(Hit.Actor))
+		if (!Hit.Actor.IsValid() || FilterHandle.FilterPassesForActor(Hit.Actor))
 		{
-			ReturnHitResult = Hit;
-			ReturnHitResult.bBlockingHit = true; // treat it as a blocking hit
+			OutHitResult = Hit;
+			OutHitResult.bBlockingHit = true; // treat it as a blocking hit
 			return;
 		}
 	}
 }
 
-void AGameplayAbilityTargetActor_Trace::SweepWithFilter(FHitResult& ReturnHitResult, const UWorld* InWorld, const FGameplayTargetDataFilterHandle InFilterHandle, const FVector& InTraceStart, const FVector& InTraceEnd, const FQuat& InRotation, const FCollisionShape CollisionShape, FName TraceProfileName, const FCollisionQueryParams Params) const
+void AGameplayAbilityTargetActor_Trace::SweepWithFilter(FHitResult& OutHitResult, const UWorld* World, const FGameplayTargetDataFilterHandle FilterHandle, const FVector& Start, const FVector& End, const FQuat& Rotation, const FCollisionShape CollisionShape, FName ProfileName, const FCollisionQueryParams Params)
 {
-	check(InWorld);
+	check(World);
 
 	TArray<FHitResult> HitResults;
-	InWorld->SweepMultiByProfile(HitResults, InTraceStart, InTraceEnd, InRotation, TraceProfileName, CollisionShape, Params);
+	World->SweepMultiByProfile(HitResults, Start, End, Rotation, ProfileName, CollisionShape, Params);
 
-	ReturnHitResult.TraceStart = InTraceStart;
-	ReturnHitResult.TraceEnd = InTraceEnd;
+	OutHitResult.TraceStart = Start;
+	OutHitResult.TraceEnd = End;
 
 	for (int32 HitIdx = 0; HitIdx < HitResults.Num(); ++HitIdx)
 	{
 		const FHitResult& Hit = HitResults[HitIdx];
 
-		if (!Hit.Actor.IsValid() || InFilterHandle.FilterPassesForActor(Hit.Actor))
+		if (!Hit.Actor.IsValid() || FilterHandle.FilterPassesForActor(Hit.Actor))
 		{
-			ReturnHitResult = Hit;
-			ReturnHitResult.bBlockingHit = true; // treat it as a blocking hit
+			OutHitResult = Hit;
+			OutHitResult.bBlockingHit = true; // treat it as a blocking hit
 			return;
 		}
 	}
@@ -86,44 +87,55 @@ FGameplayAbilityTargetDataHandle AGameplayAbilityTargetActor_Trace::StaticGetTar
 	return MakeTargetData(FHitResult());
 }
 
-void AGameplayAbilityTargetActor_Trace::AimWithPlayerController(AActor* InSourceActor, FCollisionQueryParams Params, FVector TraceStart, FVector& TraceEnd) const
+void AGameplayAbilityTargetActor_Trace::AimWithPlayerController(const AActor* InSourceActor, FCollisionQueryParams Params, const FVector& TraceStart, FVector& OutTraceEnd, bool bIgnorePitch) const
 {
-	if (OwningAbility)		//Server and launching client only
+	if (!OwningAbility) // Server and launching client only
 	{
-		APlayerController* AimingPC = OwningAbility->GetCurrentActorInfo()->PlayerController.Get();
-		check(AimingPC);
-		FVector CamLoc;
-		FRotator CamRot;
-		AimingPC->GetPlayerViewPoint(CamLoc, CamRot);
-		FVector CamDir = CamRot.Vector();
-		FVector CamTarget = CamLoc + (CamDir * MaxRange);		//Straight, dumb aiming to a point that's reasonable though not exactly correct
-
-		ClipCameraRayToAbilityRange(CamLoc, CamDir, TraceStart, MaxRange, CamTarget);
-
-		FHitResult TempHitResult;
-		LineTraceWithFilter(TempHitResult, InSourceActor->GetWorld(), Filter, CamLoc, CamTarget, TraceProfile.Name, Params);
-		if (TempHitResult.bBlockingHit && (FVector::DistSquared(TraceStart, TempHitResult.Location) <= (MaxRange * MaxRange)))
-		{
-			//We actually made a hit? Pull back.
-			TraceEnd = TempHitResult.Location;
-		}
-		else
-		{
-			//If we didn't make a hit, use the clipped location.
-			TraceEnd = CamTarget;
-		}
+		return;
+	}
 
-		//Readjust so we have a full-length line going through the predicted target point.
-		FVector AimDirection = (TraceEnd - TraceStart).GetSafeNormal();
-		if (AimDirection.SizeSquared() > 0.0f)
-		{
-			TraceEnd = TraceStart + (AimDirection * MaxRange);
-		}
-		else
+	APlayerController* PC = OwningAbility->GetCurrentActorInfo()->PlayerController.Get();
+	check(PC);
+
+	FVector ViewStart;
+	FRotator ViewRot;
+	PC->GetPlayerViewPoint(ViewStart, ViewRot);
+
+	const FVector ViewDir = ViewRot.Vector();
+	FVector ViewEnd = ViewStart + (ViewDir * MaxRange);
+
+	ClipCameraRayToAbilityRange(ViewStart, ViewDir, TraceStart, MaxRange, ViewEnd);
+
+	FHitResult HitResult;
+	LineTraceWithFilter(HitResult, InSourceActor->GetWorld(), Filter, ViewStart, ViewEnd, TraceProfile.Name, Params);
+
+	const bool bUseTraceResult = HitResult.bBlockingHit && (FVector::DistSquared(TraceStart, HitResult.Location) <= (MaxRange * MaxRange));
+
+	const FVector AdjustedEnd = (bUseTraceResult) ? HitResult.Location : ViewEnd;
+
+	FVector AdjustedAimDir = (AdjustedEnd - TraceStart).GetSafeNormal();
+	if (AdjustedAimDir.IsZero())
+	{
+		AdjustedAimDir = ViewDir;
+	}
+
+	if (!bTraceAffectsAimPitch && bUseTraceResult)
+	{
+		FVector OriginalAimDir = (ViewEnd - TraceStart).GetSafeNormal();
+
+		if (!OriginalAimDir.IsZero())
 		{
-			FVector TraceEnd = TraceStart + (InSourceActor->GetActorForwardVector() * MaxRange);		//Default
+			// Convert to angles and use original pitch
+			const FRotator OriginalAimRot = OriginalAimDir.Rotation();
+
+			FRotator AdjustedAimRot = AdjustedAimDir.Rotation();
+			AdjustedAimRot.Pitch = OriginalAimRot.Pitch;
+
+			AdjustedAimDir = AdjustedAimRot.Vector();
 		}
 	}
+
+	OutTraceEnd = TraceStart + (AdjustedAimDir * MaxRange);
 }
 
 bool AGameplayAbilityTargetActor_Trace::ClipCameraRayToAbilityRange(FVector CameraLocation, FVector CameraDirection, FVector AbilityCenter, float AbilityRange, FVector& ClippedPosition)
diff --git a/Engine/Source/Runtime/GameplayAbilities/Public/Abilities/GameplayAbilityTargetActor_Trace.h b/Engine/Source/Runtime/GameplayAbilities/Public/Abilities/GameplayAbilityTargetActor_Trace.h
--- a/Engine/Source/Runtime/GameplayAbilities/Public/Abilities/GameplayAbilityTargetActor_Trace.h
+++ b/Engine/Source/Runtime/GameplayAbilities/Public/Abilities/GameplayAbilityTargetActor_Trace.h
@@ -16,14 +16,14 @@ public:
 	virtual void EndPlay(const EEndPlayReason::Type EndPlayReason) override;
 
 	/** Traces as normal, but will manually filter all hit actors */
-	void LineTraceWithFilter(FHitResult& ReturnHitResult, const UWorld* InWorld, const FGameplayTargetDataFilterHandle InFilterHandle, const FVector& InTraceStart, const FVector& InTraceEnd, FName TraceProfileName, const FCollisionQueryParams Params) const;
+	static void LineTraceWithFilter(FHitResult& OutHitResult, const UWorld* World, const FGameplayTargetDataFilterHandle FilterHandle, const FVector& Start, const FVector& End, FName ProfileName, const FCollisionQueryParams Params);
 
 	/** Sweeps as normal, but will manually filter all hit actors */
-	void SweepWithFilter(FHitResult& ReturnHitResult, const UWorld* InWorld, const FGameplayTargetDataFilterHandle InFilterHandle, const FVector& InTraceStart, const FVector& InTraceEnd, const FQuat& InRotation, const FCollisionShape CollisionShape, FName TraceProfileName, const FCollisionQueryParams Params) const;
+	static void SweepWithFilter(FHitResult& OutHitResult, const UWorld* World, const FGameplayTargetDataFilterHandle FilterHandle, const FVector& Start, const FVector& End, const FQuat& Rotation, const FCollisionShape CollisionShape, FName ProfileName, const FCollisionQueryParams Params);
 
 	virtual FGameplayAbilityTargetDataHandle StaticGetTargetData(UWorld * World, const FGameplayAbilityActorInfo* ActorInfo, FGameplayAbilityActivationInfo ActivationInfo) const override;
 
-	void AimWithPlayerController(AActor* InSourceActor, FCollisionQueryParams Params, FVector TraceStart, FVector& TraceEnd) const;
+	void AimWithPlayerController(const AActor* InSourceActor, FCollisionQueryParams Params, const FVector& TraceStart, FVector& OutTraceEnd, bool bIgnorePitch = false) const;
 
 	static bool ClipCameraRayToAbilityRange(FVector CameraLocation, FVector CameraDirection, FVector AbilityCenter, float AbilityRange, FVector& ClippedPosition);
 
@@ -32,13 +32,17 @@ public:
 	virtual void ConfirmTargetingAndContinue() override;
 
 	virtual void Tick(float DeltaSeconds) override;
-	
-	UPROPERTY(BlueprintReadWrite, EditAnywhere, meta = (ExposeOnSpawn = true), Category = Projectile)
+
+	UPROPERTY(BlueprintReadWrite, EditAnywhere, meta = (ExposeOnSpawn = true), Category = Trace)
 	float MaxRange;
 
 	UPROPERTY(BlueprintReadWrite, EditAnywhere, config, meta = (ExposeOnSpawn = true), Category = Trace)
 	FCollisionProfileName TraceProfile;
 
+	// Does the trace affect the aiming pitch
+	UPROPERTY(BlueprintReadWrite, EditAnywhere, meta = (ExposeOnSpawn = true), Category = Trace)
+	bool bTraceAffectsAimPitch;
+
 protected:
 	virtual FHitResult PerformTrace(AActor* InSourceActor) PURE_VIRTUAL(AGameplayAbilityTargetActor_Trace, return FHitResult(););
 
