- Change apply GP to target functions to return arrays of handles instead of a single handle

diff --git a/Engine/Source/Runtime/GameplayAbilities/Private/Abilities/GameplayAbility.cpp b/Engine/Source/Runtime/GameplayAbilities/Private/Abilities/GameplayAbility.cpp
--- a/Engine/Source/Runtime/GameplayAbilities/Private/Abilities/GameplayAbility.cpp
+++ b/Engine/Source/Runtime/GameplayAbilities/Private/Abilities/GameplayAbility.cpp
@@ -887,72 +887,50 @@ FActiveGameplayEffectHandle UGameplayAbility::ApplyGameplayEffectSpecToOwner(con
 
 // -------------------------------
 
-FActiveGameplayEffectHandle UGameplayAbility::BP_ApplyGameplayEffectToTarget(FGameplayAbilityTargetDataHandle Target, TSubclassOf<UGameplayEffect> GameplayEffectClass, int32 GameplayEffectLevel)
+TArray<FActiveGameplayEffectHandle> UGameplayAbility::BP_ApplyGameplayEffectToTarget(FGameplayAbilityTargetDataHandle Target, TSubclassOf<UGameplayEffect> GameplayEffectClass, int32 GameplayEffectLevel)
 {
-	if (GameplayEffectClass == nullptr)
-	{
-		ABILITY_LOG(Error, TEXT("BP_ApplyGameplayEffectToTarget called on ability %s with no GameplayEffectClass."), *GetName());
-		return FActiveGameplayEffectHandle();
-	}
-	
 	return ApplyGameplayEffectToTarget(CurrentSpecHandle, CurrentActorInfo, CurrentActivationInfo, Target, GameplayEffectClass, GameplayEffectLevel);
 }
 
-FActiveGameplayEffectHandle UGameplayAbility::K2_ApplyGameplayEffectToTarget(FGameplayAbilityTargetDataHandle Target, const UGameplayEffect* GameplayEffect, int32 GameplayEffectLevel)
+TArray<FActiveGameplayEffectHandle> UGameplayAbility::K2_ApplyGameplayEffectToTarget(FGameplayAbilityTargetDataHandle Target, const UGameplayEffect* GameplayEffect, int32 GameplayEffectLevel)
 {
-	if (GameplayEffect==nullptr)
-	{
-		ABILITY_LOG(Error, TEXT("K2_ApplyGameplayEffectToTarget called on ability %s with no GameplayEffect."), *GetName());
-		return FActiveGameplayEffectHandle();
-	}
-
 	return BP_ApplyGameplayEffectToTarget(Target, GameplayEffect->GetClass(), GameplayEffectLevel);
 }
 
-FActiveGameplayEffectHandle UGameplayAbility::ApplyGameplayEffectToTarget(const FGameplayAbilitySpecHandle Handle, const FGameplayAbilityActorInfo* ActorInfo, const FGameplayAbilityActivationInfo ActivationInfo, FGameplayAbilityTargetDataHandle Target, TSubclassOf<UGameplayEffect> GameplayEffectClass, float GameplayEffectLevel)
+TArray<FActiveGameplayEffectHandle> UGameplayAbility::ApplyGameplayEffectToTarget(const FGameplayAbilitySpecHandle Handle, const FGameplayAbilityActorInfo* ActorInfo, const FGameplayAbilityActivationInfo ActivationInfo, FGameplayAbilityTargetDataHandle Target, TSubclassOf<UGameplayEffect> GameplayEffectClass, float GameplayEffectLevel)
 {
+	TArray<FActiveGameplayEffectHandle> EffectHandles;
+
 	if (GameplayEffectClass == nullptr)
 	{
-		ABILITY_LOG(Error, TEXT("K2_ApplyGameplayEffectToTarget called on ability %s with no GameplayEffect."), *GetName());
-		return FActiveGameplayEffectHandle();
+		ABILITY_LOG(Error, TEXT("ApplyGameplayEffectToTarget called on ability %s with no GameplayEffect."), *GetName());
 	}
-
-	FActiveGameplayEffectHandle EffectHandle;
-
-	if (HasAuthorityOrPredictionKey(ActorInfo, &ActivationInfo))
+	else if (HasAuthorityOrPredictionKey(ActorInfo, &ActivationInfo))
 	{
 		FGameplayEffectSpecHandle SpecHandle = MakeOutgoingGameplayEffectSpec(Handle, ActorInfo, ActivationInfo, GameplayEffectClass, GameplayEffectLevel);
-		EffectHandle = ApplyGameplayEffectSpecToTarget(Handle, ActorInfo, ActivationInfo, SpecHandle, Target);
+		EffectHandles.Append(ApplyGameplayEffectSpecToTarget(Handle, ActorInfo, ActivationInfo, SpecHandle, Target));
 	}
 
-	return EffectHandle;
+	return EffectHandles;
 }
 
-FActiveGameplayEffectHandle UGameplayAbility::K2_ApplyGameplayEffectSpecToTarget(const FGameplayEffectSpecHandle SpecHandle, FGameplayAbilityTargetDataHandle TargetData)
+TArray<FActiveGameplayEffectHandle> UGameplayAbility::K2_ApplyGameplayEffectSpecToTarget(const FGameplayEffectSpecHandle SpecHandle, FGameplayAbilityTargetDataHandle TargetData)
 {
 	return ApplyGameplayEffectSpecToTarget(CurrentSpecHandle, CurrentActorInfo, CurrentActivationInfo, SpecHandle, TargetData);
 }
 
-FActiveGameplayEffectHandle UGameplayAbility::ApplyGameplayEffectSpecToTarget(const FGameplayAbilitySpecHandle AbilityHandle, const FGameplayAbilityActorInfo* ActorInfo, const FGameplayAbilityActivationInfo ActivationInfo, const FGameplayEffectSpecHandle SpecHandle, FGameplayAbilityTargetDataHandle TargetData)
+TArray<FActiveGameplayEffectHandle> UGameplayAbility::ApplyGameplayEffectSpecToTarget(const FGameplayAbilitySpecHandle AbilityHandle, const FGameplayAbilityActorInfo* ActorInfo, const FGameplayAbilityActivationInfo ActivationInfo, const FGameplayEffectSpecHandle SpecHandle, FGameplayAbilityTargetDataHandle TargetData)
 {
-	FActiveGameplayEffectHandle EffectHandle;
+	TArray<FActiveGameplayEffectHandle> EffectHandles;
 	
 	if (SpecHandle.IsValid() && HasAuthorityOrPredictionKey(ActorInfo, &ActivationInfo))
 	{
 		for (auto Data : TargetData.Data)
 		{
-			TArray<FActiveGameplayEffectHandle> EffectHandles = Data->ApplyGameplayEffectSpec(*SpecHandle.Data.Get(), ActorInfo->AbilitySystemComponent->GetPredictionKeyForNewAction());
-			if (EffectHandles.Num() > 0)
-			{
-				EffectHandle = EffectHandles[0];
-				if (EffectHandles.Num() > 1)
-				{
-					ABILITY_LOG(Warning, TEXT("ApplyGameplayEffectSpecToTarget called on TargetData with multiple actor targets. Only returning 1 handle though we applied many. %s"), *GetName());
-				}
-			}
+			EffectHandles.Append(Data->ApplyGameplayEffectSpec(*SpecHandle.Data.Get(), ActorInfo->AbilitySystemComponent->GetPredictionKeyForNewAction()));
 		}
 	}
-	return EffectHandle;
+	return EffectHandles;
 }
 
 void UGameplayAbility::ConvertDeprecatedGameplayEffectReferencesToBlueprintReferences(UGameplayEffect* OldGE, TSubclassOf<UGameplayEffect> NewGEClass)
diff --git a/Engine/Source/Runtime/GameplayAbilities/Public/Abilities/GameplayAbility.h b/Engine/Source/Runtime/GameplayAbilities/Public/Abilities/GameplayAbility.h
--- a/Engine/Source/Runtime/GameplayAbilities/Public/Abilities/GameplayAbility.h
+++ b/Engine/Source/Runtime/GameplayAbilities/Public/Abilities/GameplayAbility.h
@@ -364,18 +364,18 @@ protected:
 	// ---------
 
 	UFUNCTION(BlueprintCallable, Category = Ability, FriendlyName = "ApplyGameplayEffectToTarget")
-	FActiveGameplayEffectHandle BP_ApplyGameplayEffectToTarget(FGameplayAbilityTargetDataHandle TargetData, TSubclassOf<UGameplayEffect> GameplayEffectClass, int32 GameplayEffectLevel = 1);
+	TArray<FActiveGameplayEffectHandle> BP_ApplyGameplayEffectToTarget(FGameplayAbilityTargetDataHandle TargetData, TSubclassOf<UGameplayEffect> GameplayEffectClass, int32 GameplayEffectLevel = 1);
 
 	UFUNCTION(BlueprintCallable, Category = Ability, FriendlyName = "ApplyGameplayEffectToTarget", meta=(DeprecatedFunction, DeprecationMessage = "Use new ApplyGameplayEffectToTarget"))
-	FActiveGameplayEffectHandle K2_ApplyGameplayEffectToTarget(FGameplayAbilityTargetDataHandle TargetData, const UGameplayEffect* GameplayEffect, int32 GameplayEffectLevel = 1);
+	TArray<FActiveGameplayEffectHandle> K2_ApplyGameplayEffectToTarget(FGameplayAbilityTargetDataHandle TargetData, const UGameplayEffect* GameplayEffect, int32 GameplayEffectLevel = 1);
 
 	/** Non blueprintcallable, safe to call on CDO/NonInstance abilities */
-	FActiveGameplayEffectHandle ApplyGameplayEffectToTarget(const FGameplayAbilitySpecHandle Handle, const FGameplayAbilityActorInfo* ActorInfo, const FGameplayAbilityActivationInfo ActivationInfo, FGameplayAbilityTargetDataHandle Target, TSubclassOf<UGameplayEffect> GameplayEffectClass, float GameplayEffectLevel);
+	TArray<FActiveGameplayEffectHandle> ApplyGameplayEffectToTarget(const FGameplayAbilitySpecHandle Handle, const FGameplayAbilityActorInfo* ActorInfo, const FGameplayAbilityActivationInfo ActivationInfo, FGameplayAbilityTargetDataHandle Target, TSubclassOf<UGameplayEffect> GameplayEffectClass, float GameplayEffectLevel);
 
 	UFUNCTION(BlueprintCallable, Category = Ability, FriendlyName = "ApplyGameplayEffectSpecToTarget")
-	FActiveGameplayEffectHandle K2_ApplyGameplayEffectSpecToTarget(const FGameplayEffectSpecHandle EffectSpecHandle, FGameplayAbilityTargetDataHandle TargetData);
+	TArray<FActiveGameplayEffectHandle> K2_ApplyGameplayEffectSpecToTarget(const FGameplayEffectSpecHandle EffectSpecHandle, FGameplayAbilityTargetDataHandle TargetData);
 
-	FActiveGameplayEffectHandle ApplyGameplayEffectSpecToTarget(const FGameplayAbilitySpecHandle AbilityHandle, const FGameplayAbilityActorInfo* ActorInfo, const FGameplayAbilityActivationInfo ActivationInfo, const FGameplayEffectSpecHandle SpecHandle, FGameplayAbilityTargetDataHandle TargetData);
+	TArray<FActiveGameplayEffectHandle> ApplyGameplayEffectSpecToTarget(const FGameplayAbilitySpecHandle AbilityHandle, const FGameplayAbilityActorInfo* ActorInfo, const FGameplayAbilityActivationInfo ActivationInfo, const FGameplayEffectSpecHandle SpecHandle, FGameplayAbilityTargetDataHandle TargetData);
 
 	// -------------------------------------
 	//	GameplayCue
