Fixed client remotely ending server abilities.

diff --git a/Engine/Source/Runtime/GameplayAbilities/Private/AbilitySystemComponent_Abilities.cpp b/Engine/Source/Runtime/GameplayAbilities/Private/AbilitySystemComponent_Abilities.cpp
--- a/Engine/Source/Runtime/GameplayAbilities/Private/AbilitySystemComponent_Abilities.cpp
+++ b/Engine/Source/Runtime/GameplayAbilities/Private/AbilitySystemComponent_Abilities.cpp
@@ -570,6 +570,7 @@ bool UAbilitySystemComponent::TryActivateAbility(FGameplayAbilitySpecHandle Hand
 		if (Ability->GetInstancingPolicy() == EGameplayAbilityInstancingPolicy::InstancedPerExecution)
 		{
 			InstancedAbility = CreateNewInstanceOfAbility(*Spec, Ability);
+			ActivationInfo.bCanBeEndedByOtherInstance = true;
 			InstancedAbility->CallActivateAbility(Handle, ActorInfo, ActivationInfo, OnGameplayAbilityEndedDelegate);
 			if (OutInstancedAbility)
 			{
