making it so that we can place gameplay volumes in the world and they will work correctly and not just work when spawned from an ability.

diff --git a/Engine/Source/Runtime/GameplayAbilities/Private/AbilitySystemBlueprintLibrary.cpp b/Engine/Source/Runtime/GameplayAbilities/Private/AbilitySystemBlueprintLibrary.cpp
--- a/Engine/Source/Runtime/GameplayAbilities/Private/AbilitySystemBlueprintLibrary.cpp
+++ b/Engine/Source/Runtime/GameplayAbilities/Private/AbilitySystemBlueprintLibrary.cpp
@@ -116,6 +116,12 @@ FGameplayTargetDataFilterHandle UAbilitySystemBlueprintLibrary::MakeFilterHandle
 	return FilterHandle;
 }
 
+FGameplayEffectSpecHandle UAbilitySystemBlueprintLibrary::MakeSpecHandle(UGameplayEffect* InGameplayEffect, AActor* InInstigator, AActor* InEffectCauser, float InLevel)
+{
+	FGameplayEffectContext* EffectContext = new FGameplayEffectContext(InInstigator, InEffectCauser);
+	return FGameplayEffectSpecHandle(new FGameplayEffectSpec(InGameplayEffect, FGameplayEffectContextHandle(EffectContext), InLevel));
+}
+
 FGameplayAbilityTargetDataHandle UAbilitySystemBlueprintLibrary::AbilityTargetDataFromHitResult(FHitResult HitResult)
 {
 	// Construct TargetData
diff --git a/Engine/Source/Runtime/GameplayAbilities/Public/AbilitySystemBlueprintLibrary.h b/Engine/Source/Runtime/GameplayAbilities/Public/AbilitySystemBlueprintLibrary.h
--- a/Engine/Source/Runtime/GameplayAbilities/Public/AbilitySystemBlueprintLibrary.h
+++ b/Engine/Source/Runtime/GameplayAbilities/Public/AbilitySystemBlueprintLibrary.h
@@ -53,6 +53,10 @@ class GAMEPLAYABILITIES_API UAbilitySystemBlueprintLibrary : public UBlueprintFu
 	UFUNCTION(BlueprintPure, Category = "Filter")
 	static FGameplayTargetDataFilterHandle MakeFilterHandle(FGameplayTargetDataFilter Filter, AActor* FilterActor);
 
+	/** Create a spec handle, filling out all fields */
+	UFUNCTION(BlueprintPure, Category = "Spec")
+	static FGameplayEffectSpecHandle MakeSpecHandle(UGameplayEffect* InGameplayEffect, AActor* InInstigator, AActor* InEffectCauser, float InLevel = 1.0f);
+
 	UFUNCTION(BlueprintPure, Category = "Ability|TargetData")
 	static TArray<AActor*> GetActorsFromTargetData(FGameplayAbilityTargetDataHandle TargetData, int32 Index);
 
