Reticles face toward their owners by default. A bool is available to set this behavior to 2D/flat, which is on by default.

diff --git a/Engine/Source/Runtime/GameplayAbilities/Private/Abilities/GameplayAbilityWorldReticle.cpp b/Engine/Source/Runtime/GameplayAbilities/Private/Abilities/GameplayAbilityWorldReticle.cpp
--- a/Engine/Source/Runtime/GameplayAbilities/Private/Abilities/GameplayAbilityWorldReticle.cpp
+++ b/Engine/Source/Runtime/GameplayAbilities/Private/Abilities/GameplayAbilityWorldReticle.cpp
@@ -17,6 +17,15 @@ AGameplayAbilityWorldReticle::AGameplayAbilityWorldReticle(const FObjectInitiali
 	PrimaryActorTick.TickGroup = TG_PrePhysics;
 	bIsTargetValid = true;
 	bIsTargetAnActor = false;
+	bFaceOwnerFlat = true;
+}
+
+
+void AGameplayAbilityWorldReticle::Tick(float DeltaTime)
+{
+	Super::Tick(DeltaTime);
+
+	FaceTowardSource(bFaceOwnerFlat);
 }
 
 void AGameplayAbilityWorldReticle::InitializeReticle(AGameplayAbilityTargetActor* InTargetingActor, FWorldReticleParameters InParameters)
diff --git a/Engine/Source/Runtime/GameplayAbilities/Private/Abilities/GameplayAbilityWorldReticle_ActorVisualization.cpp b/Engine/Source/Runtime/GameplayAbilities/Private/Abilities/GameplayAbilityWorldReticle_ActorVisualization.cpp
--- a/Engine/Source/Runtime/GameplayAbilities/Private/Abilities/GameplayAbilityWorldReticle_ActorVisualization.cpp
+++ b/Engine/Source/Runtime/GameplayAbilities/Private/Abilities/GameplayAbilityWorldReticle_ActorVisualization.cpp
@@ -29,13 +29,6 @@ AGameplayAbilityWorldReticle_ActorVisualization::AGameplayAbilityWorldReticle_Ac
 }
 
 
-void AGameplayAbilityWorldReticle_ActorVisualization::Tick(float DeltaTime)
-{
-	Super::Tick(DeltaTime);
-
-	FaceTowardSource(true);
-}
-
 void AGameplayAbilityWorldReticle_ActorVisualization::InitializeReticleVisualizationInformation(AGameplayAbilityTargetActor* InTargetingActor, AActor* VisualizationActor, UMaterialInterface *VisualizationMaterial)
 {
 	if (VisualizationActor)
diff --git a/Engine/Source/Runtime/GameplayAbilities/Public/Abilities/GameplayAbilityWorldReticle.h b/Engine/Source/Runtime/GameplayAbilities/Public/Abilities/GameplayAbilityWorldReticle.h
--- a/Engine/Source/Runtime/GameplayAbilities/Public/Abilities/GameplayAbilityWorldReticle.h
+++ b/Engine/Source/Runtime/GameplayAbilities/Public/Abilities/GameplayAbilityWorldReticle.h
@@ -28,6 +28,8 @@ class GAMEPLAYABILITIES_API AGameplayAbilityWorldReticle : public AActor
 
 public:
 
+	virtual void Tick(float DeltaSeconds) override;
+
 	/** Accessor for checking, before instantiating, if this WorldReticle will replicate. */
 	bool GetReplicates() const;
 
@@ -60,9 +62,13 @@ public:
 	UFUNCTION(BlueprintCallable, Category = Reticle)
 	void FaceTowardSource(bool bFaceIn2D);
 
-	UPROPERTY(EditAnywhere, BlueprintReadWrite, meta = (ExposeOnSpawn = "true"), Category = "Reticle")
+	UPROPERTY(EditInstanceOnly, BlueprintReadWrite, meta = (ExposeOnSpawn = "true"), Category = "Reticle")
 	FWorldReticleParameters Parameters;
 
+	/** Makes the reticle's default owner-facing behavior operate in 2D (flat) instead of 3D (pitched). Defaults to true. */
+	UPROPERTY(EditAnywhere, BlueprintReadWrite, meta = (ExposeOnSpawn = "true"), Category = "Reticle")
+	bool bFaceOwnerFlat;
+
 protected:
 	/** This indicates whether or not the targeting actor considers the current target to be valid. Defaults to true. */
 	UPROPERTY(BlueprintReadOnly, Category = "Network")
diff --git a/Engine/Source/Runtime/GameplayAbilities/Public/Abilities/GameplayAbilityWorldReticle_ActorVisualization.h b/Engine/Source/Runtime/GameplayAbilities/Public/Abilities/GameplayAbilityWorldReticle_ActorVisualization.h
--- a/Engine/Source/Runtime/GameplayAbilities/Public/Abilities/GameplayAbilityWorldReticle_ActorVisualization.h
+++ b/Engine/Source/Runtime/GameplayAbilities/Public/Abilities/GameplayAbilityWorldReticle_ActorVisualization.h
@@ -14,8 +14,6 @@ class AGameplayAbilityWorldReticle_ActorVisualization : public AGameplayAbilityW
 	GENERATED_UCLASS_BODY()
 
 public:
-	virtual void Tick(float DeltaSeconds) override;
-
 	void InitializeReticleVisualizationInformation(AGameplayAbilityTargetActor* InTargetingActor, AActor* VisualizationActor, UMaterialInterface *VisualizationMaterial);
 
 private:
