Broke out cost and cooldown checks and commits so designers can get around the commit functionality. These mechanisms do not currently broadcast events.

diff --git a/Engine/Source/Runtime/GameplayAbilities/Private/Abilities/GameplayAbility.cpp b/Engine/Source/Runtime/GameplayAbilities/Private/Abilities/GameplayAbility.cpp
--- a/Engine/Source/Runtime/GameplayAbilities/Private/Abilities/GameplayAbility.cpp
+++ b/Engine/Source/Runtime/GameplayAbilities/Private/Abilities/GameplayAbility.cpp
@@ -215,29 +215,43 @@ bool UGameplayAbility::CommitAbility(const FGameplayAbilitySpecHandle Handle, co
 	return true;
 }
 
-bool UGameplayAbility::CommitCheck(const FGameplayAbilitySpecHandle Handle, const FGameplayAbilityActorInfo* ActorInfo, const FGameplayAbilityActivationInfo ActivationInfo)
+bool UGameplayAbility::CommitAbilityCooldown(const FGameplayAbilitySpecHandle Handle, const FGameplayAbilityActorInfo* ActorInfo, const FGameplayAbilityActivationInfo ActivationInfo)
 {
-	/**
-	 *	Checks if we can (still) commit this ability. There are some subtleties here.
-	 *		-An ability can start activating, play an animation, wait for a user confirmation/target data, and then actually commit
-	 *		-Commit = spend resources/cooldowns. Its possible the source has changed state since he started activation, so a commit may fail.
-	 *		-We don't want to just call CanActivateAbility() since right now that also checks things like input inhibition.
-	 *			-E.g., its possible the act of starting your ability makes it no longer activatable (CanaCtivateAbility() may be false if called here).
-	 */
-
+	// Last chance to fail (maybe we no longer have resources to commit since we after we started this ability activation)
 	if (!CheckCooldown(Handle, ActorInfo))
 	{
 		return false;
 	}
 
+	ApplyCooldown(Handle, ActorInfo, ActivationInfo);
+	return true;
+}
+
+bool UGameplayAbility::CommitAbilityCost(const FGameplayAbilitySpecHandle Handle, const FGameplayAbilityActorInfo* ActorInfo, const FGameplayAbilityActivationInfo ActivationInfo)
+{
+	// Last chance to fail (maybe we no longer have resources to commit since we after we started this ability activation)
 	if (!CheckCost(Handle, ActorInfo))
 	{
 		return false;
 	}
 
+	ApplyCost(Handle, ActorInfo, ActivationInfo);
 	return true;
 }
 
+bool UGameplayAbility::CommitCheck(const FGameplayAbilitySpecHandle Handle, const FGameplayAbilityActorInfo* ActorInfo, const FGameplayAbilityActivationInfo ActivationInfo)
+{
+	/**
+	 *	Checks if we can (still) commit this ability. There are some subtleties here.
+	 *		-An ability can start activating, play an animation, wait for a user confirmation/target data, and then actually commit
+	 *		-Commit = spend resources/cooldowns. Its possible the source has changed state since he started activation, so a commit may fail.
+	 *		-We don't want to just call CanActivateAbility() since right now that also checks things like input inhibition.
+	 *			-E.g., its possible the act of starting your ability makes it no longer activatable (CanaCtivateAbility() may be false if called here).
+	 */
+
+	return (CheckCooldown(Handle, ActorInfo) && CheckCost(Handle, ActorInfo));
+}
+
 void UGameplayAbility::CommitExecute(const FGameplayAbilitySpecHandle Handle, const FGameplayAbilityActorInfo* ActorInfo, const FGameplayAbilityActivationInfo ActivationInfo)
 {
 	ApplyCooldown(Handle, ActorInfo, ActivationInfo);
@@ -247,7 +261,7 @@ void UGameplayAbility::CommitExecute(const FGameplayAbilitySpecHandle Handle, co
 
 void UGameplayAbility::CancelAbility(const FGameplayAbilitySpecHandle Handle, const FGameplayAbilityActorInfo* ActorInfo, const FGameplayAbilityActivationInfo ActivationInfo)
 {
-	EndAbility(Handle, ActorInfo, ActivationInfo);
+	EndAbility();
 }
 
 void UGameplayAbility::EndAbility(const FGameplayAbilitySpecHandle Handle, const FGameplayAbilityActorInfo* ActorInfo, const FGameplayAbilityActivationInfo ActivationInfo)
@@ -566,6 +580,30 @@ bool UGameplayAbility::K2_CommitAbility()
 	return CommitAbility(CurrentSpecHandle, CurrentActorInfo, CurrentActivationInfo);
 }
 
+bool UGameplayAbility::K2_CommitAbilityCooldown()
+{
+	check(CurrentActorInfo);
+	return CommitAbilityCooldown(CurrentSpecHandle, CurrentActorInfo, CurrentActivationInfo);
+}
+
+bool UGameplayAbility::K2_CommitAbilityCost()
+{
+	check(CurrentActorInfo);
+	return CommitAbilityCost(CurrentSpecHandle, CurrentActorInfo, CurrentActivationInfo);
+}
+
+bool UGameplayAbility::K2_CheckAbilityCooldown()
+{
+	check(CurrentActorInfo);
+	return CheckCooldown(CurrentSpecHandle, CurrentActorInfo);
+}
+
+bool UGameplayAbility::K2_CheckAbilityCost()
+{
+	check(CurrentActorInfo);
+	return CheckCost(CurrentSpecHandle, CurrentActorInfo);
+}
+
 void UGameplayAbility::K2_EndAbility()
 {
 	check(CurrentActorInfo);
diff --git a/Engine/Source/Runtime/GameplayAbilities/Public/Abilities/GameplayAbility.h b/Engine/Source/Runtime/GameplayAbilities/Public/Abilities/GameplayAbility.h
--- a/Engine/Source/Runtime/GameplayAbilities/Public/Abilities/GameplayAbility.h
+++ b/Engine/Source/Runtime/GameplayAbilities/Public/Abilities/GameplayAbility.h
@@ -281,7 +281,25 @@ protected:
 	UFUNCTION(BlueprintCallable, Category = Ability, FriendlyName = "CommitAbility")
 	virtual bool K2_CommitAbility();
 
+	/** Attempts to commit the ability's cooldown only. */
+	UFUNCTION(BlueprintCallable, Category = Ability, FriendlyName = "CommitAbilityCooldown")
+	virtual bool K2_CommitAbilityCooldown();
+
+	/** Attempts to commit the ability's cost only. */
+	UFUNCTION(BlueprintCallable, Category = Ability, FriendlyName = "CommitAbilityCost")
+	virtual bool K2_CommitAbilityCost();
+
+	/** Checks the ability's cooldown, but does not apply it. */
+	UFUNCTION(BlueprintCallable, Category = Ability, FriendlyName = "CheckAbilityCooldown")
+	virtual bool K2_CheckAbilityCooldown();
+
+	/** Checks the ability's cost, but does not apply it. */
+	UFUNCTION(BlueprintCallable, Category = Ability, FriendlyName = "CheckAbilityCost")
+	virtual bool K2_CheckAbilityCost();
+
 	virtual bool CommitAbility(const FGameplayAbilitySpecHandle Handle, const FGameplayAbilityActorInfo* ActorInfo, const FGameplayAbilityActivationInfo ActivationInfo);
+	virtual bool CommitAbilityCooldown(const FGameplayAbilitySpecHandle Handle, const FGameplayAbilityActorInfo* ActorInfo, const FGameplayAbilityActivationInfo ActivationInfo);
+	virtual bool CommitAbilityCost(const FGameplayAbilitySpecHandle Handle, const FGameplayAbilityActorInfo* ActorInfo, const FGameplayAbilityActivationInfo ActivationInfo);
 
 	/**
 	 * The last chance to fail before commiting
@@ -483,8 +501,10 @@ protected:
 	/** Applies CooldownGameplayEffect to the target */
 	virtual void ApplyCooldown(const FGameplayAbilitySpecHandle Handle, const FGameplayAbilityActorInfo* ActorInfo, const FGameplayAbilityActivationInfo ActivationInfo);
 
+	/** Checks cost. returns true if we can pay for the abilty. False if not */
 	virtual bool CheckCost(const FGameplayAbilitySpecHandle Handle, const FGameplayAbilityActorInfo* ActorInfo) const;
 
+	/** Applies the ability's cost to the target */
 	virtual void ApplyCost(const FGameplayAbilitySpecHandle Handle, const FGameplayAbilityActorInfo* ActorInfo, const FGameplayAbilityActivationInfo ActivationInfo);
 
 	// -----------------------------------------------	
