OnGiveAbility and OnAvatarSet events for GameplayAbilities

diff --git a/Engine/Source/Runtime/GameplayAbilities/Private/Abilities/GameplayAbility.cpp b/Engine/Source/Runtime/GameplayAbilities/Private/Abilities/GameplayAbility.cpp
--- a/Engine/Source/Runtime/GameplayAbilities/Private/Abilities/GameplayAbility.cpp
+++ b/Engine/Source/Runtime/GameplayAbilities/Private/Abilities/GameplayAbility.cpp
@@ -261,7 +261,7 @@ void UGameplayAbility::CommitExecute(const FGameplayAbilitySpecHandle Handle, co
 
 void UGameplayAbility::CancelAbility(const FGameplayAbilitySpecHandle Handle, const FGameplayAbilityActorInfo* ActorInfo, const FGameplayAbilityActivationInfo ActivationInfo)
 {
-	EndAbility();
+	EndAbility(Handle, ActorInfo, ActivationInfo);
 }
 
 void UGameplayAbility::EndAbility(const FGameplayAbilitySpecHandle Handle, const FGameplayAbilityActorInfo* ActorInfo, const FGameplayAbilityActivationInfo ActivationInfo)
@@ -865,6 +865,22 @@ bool UGameplayAbility::HasAuthorityOrPredictionKey(const FGameplayAbilityActorIn
 	return ActorInfo->AbilitySystemComponent->HasAuthorityOrPredictionKey(ActivationInfo);
 }
 
+void UGameplayAbility::OnGiveAbility(const FGameplayAbilityActorInfo* ActorInfo, const FGameplayAbilitySpec& Spec)
+{
+	SetCurrentActorInfo(Spec.Handle, ActorInfo);
+
+	// If we already have an avatar set, call the OnAvatarSet event as well
+	if (ActorInfo && ActorInfo->AvatarActor.IsValid())
+	{
+		OnAvatarSet(ActorInfo, Spec);
+	}
+}
+
+void UGameplayAbility::OnAvatarSet(const FGameplayAbilityActorInfo* ActorInfo, const FGameplayAbilitySpec& Spec)
+{
+	// Projects may want to initiate passives or do other "BeginPlay" type of logic here.
+}
+
 // -------------------------------------------------------
 
 FActiveGameplayEffectHandle UGameplayAbility::BP_ApplyGameplayEffectToOwner(TSubclassOf<UGameplayEffect> GameplayEffectClass, int32 GameplayEffectLevel)
diff --git a/Engine/Source/Runtime/GameplayAbilities/Private/AbilitySystemComponent_Abilities.cpp b/Engine/Source/Runtime/GameplayAbilities/Private/AbilitySystemComponent_Abilities.cpp
--- a/Engine/Source/Runtime/GameplayAbilities/Private/AbilitySystemComponent_Abilities.cpp
+++ b/Engine/Source/Runtime/GameplayAbilities/Private/AbilitySystemComponent_Abilities.cpp
@@ -142,9 +142,20 @@ void UAbilitySystemComponent::TickComponent(float DeltaTime, enum ELevelTick Tic
 void UAbilitySystemComponent::InitAbilityActorInfo(AActor* InOwnerActor, AActor* InAvatarActor)
 {
 	check(AbilityActorInfo.IsValid());
+	bool AvatarChanged = (InAvatarActor != AbilityActorInfo->AvatarActor);
+
 	AbilityActorInfo->InitFromActor(InOwnerActor, InAvatarActor, this);
+
 	OwnerActor = InOwnerActor;
 	AvatarActor = InAvatarActor;
+
+	if (AvatarChanged)
+	{
+		for (FGameplayAbilitySpec& Spec : ActivatableAbilities)
+		{
+			Spec.Ability->OnAvatarSet(AbilityActorInfo.Get(), Spec);
+		}
+	}
 }
 
 void UAbilitySystemComponent::UpdateShouldTick()
@@ -274,6 +285,8 @@ void UAbilitySystemComponent::OnGiveAbility(const FGameplayAbilitySpec Spec)
 			GameplayEventTriggeredAbilities.Add(EventTag, Triggers);
 		}
 	}
+
+	Spec.Ability->OnGiveAbility(AbilityActorInfo.Get(), Spec);
 }
 
 void UAbilitySystemComponent::CheckForClearedAbilities()
@@ -822,7 +835,7 @@ void UAbilitySystemComponent::ClientActivateAbilitySucceed_Implementation(FGamep
 
 			if (!found)
 			{
-				ABILITY_LOG(Warning, TEXT("Ability %s was confirmed by server but no longer exists on client (replication key: %d"), *AbilityToActivate->GetName(), PredictionKey);
+				ABILITY_LOG(Verbose, TEXT("Ability %s was confirmed by server but no longer exists on client (replication key: %d"), *AbilityToActivate->GetName(), PredictionKey);
 			}
 		}
 	}
@@ -952,7 +965,6 @@ void UAbilitySystemComponent::HandleGameplayEvent(FGameplayTag EventTag, FGamepl
 		for (auto AbilityHandle : TriggeredAbilityHandles)
 		{
 			TriggerAbilityFromGameplayEvent(AbilityHandle, AbilityActorInfo.Get(), EventTag, Payload, *this);
-						
 		}
 	}
 }
diff --git a/Engine/Source/Runtime/GameplayAbilities/Public/Abilities/GameplayAbility.h b/Engine/Source/Runtime/GameplayAbilities/Public/Abilities/GameplayAbility.h
--- a/Engine/Source/Runtime/GameplayAbilities/Public/Abilities/GameplayAbility.h
+++ b/Engine/Source/Runtime/GameplayAbilities/Public/Abilities/GameplayAbility.h
@@ -218,6 +218,12 @@ public:
 
 	bool HasAuthorityOrPredictionKey(const FGameplayAbilityActorInfo* ActorInfo, const FGameplayAbilityActivationInfo* ActivationInfo) const;
 
+	/** Called when the ability is given to an AbilitySystemComponent */
+	virtual void OnGiveAbility(const FGameplayAbilityActorInfo* ActorInfo, const FGameplayAbilitySpec& Spec);
+
+	/** Called when the avatar actor is set/changes */
+	virtual void OnAvatarSet(const FGameplayAbilityActorInfo* ActorInfo, const FGameplayAbilitySpec& Spec);
+
 protected:
 
 	// --------------------------------------
