Use WeakObjPtr in context structs to avoid serilaization crashes when actors are destroyed (in cases where a GE spec is saved outside of the root set. E.g, in a SpecHandle)

diff --git a/Engine/Source/Runtime/GameplayAbilities/Private/GameplayEffectTypes.cpp b/Engine/Source/Runtime/GameplayAbilities/Private/GameplayEffectTypes.cpp
--- a/Engine/Source/Runtime/GameplayAbilities/Private/GameplayEffectTypes.cpp
+++ b/Engine/Source/Runtime/GameplayAbilities/Private/GameplayEffectTypes.cpp
@@ -34,7 +34,7 @@ void FGameplayEffectContext::AddInstigator(class AActor *InInstigator, class AAc
 	InstigatorAbilitySystemComponent = NULL;
 
 	// Cache off his AbilitySystemComponent.
-	IAbilitySystemInterface* AbilitySystemInterface = Cast<IAbilitySystemInterface>(Instigator);
+	IAbilitySystemInterface* AbilitySystemInterface = Cast<IAbilitySystemInterface>(Instigator.Get());
 	if (AbilitySystemInterface)
 	{
 		InstigatorAbilitySystemComponent = AbilitySystemInterface->GetAbilitySystemComponent();
@@ -84,7 +84,7 @@ bool FGameplayEffectContext::NetSerialize(FArchive& Ar, class UPackageMap* Map,
 				HitResult = TSharedPtr<FHitResult>(new FHitResult());
 			}
 		}
-		AddInstigator(Instigator, EffectCauser); // Just to initialize InstigatorAbilitySystemComponent
+		AddInstigator(Instigator.Get(), EffectCauser.Get()); // Just to initialize InstigatorAbilitySystemComponent
 	}
 
 	if (HasHitResults == 1)
@@ -101,10 +101,10 @@ bool FGameplayEffectContext::NetSerialize(FArchive& Ar, class UPackageMap* Map,
 
 bool FGameplayEffectContext::IsLocallyControlled() const
 {
-	APawn* Pawn = Cast<APawn>(Instigator);
+	APawn* Pawn = Cast<APawn>(Instigator.Get());
 	if (!Pawn)
 	{
-		Pawn = Cast<APawn>(EffectCauser);
+		Pawn = Cast<APawn>(EffectCauser.Get());
 	}
 	if (Pawn)
 	{
@@ -121,7 +121,7 @@ void FGameplayEffectContext::AddOrigin(FVector InOrigin)
 
 void FGameplayEffectContext::GetOwnedGameplayTags(OUT FGameplayTagContainer &TagContainer) const
 {
-	IGameplayTagAssetInterface* TagInterface = Cast<IGameplayTagAssetInterface>(Instigator);
+	IGameplayTagAssetInterface* TagInterface = Cast<IGameplayTagAssetInterface>(Instigator.Get());
 	if (TagInterface)
 	{
 		TagInterface->GetOwnedGameplayTags(TagContainer);
diff --git a/Engine/Source/Runtime/GameplayAbilities/Public/Abilities/GameplayAbilityTypes.h b/Engine/Source/Runtime/GameplayAbilities/Public/Abilities/GameplayAbilityTypes.h
--- a/Engine/Source/Runtime/GameplayAbilities/Public/Abilities/GameplayAbilityTypes.h
+++ b/Engine/Source/Runtime/GameplayAbilities/Public/Abilities/GameplayAbilityTypes.h
@@ -405,6 +405,9 @@ struct GAMEPLAYABILITIES_API FGameplayEventData
 	FGameplayEventData()
 	: Instigator(NULL)
 	, Target(NULL)
+	, Var1(0.f)
+	, Var2(0.f)
+	, Var3(0.f)
 	{
 	}
 	
diff --git a/Engine/Source/Runtime/GameplayAbilities/Public/GameplayEffectTypes.h b/Engine/Source/Runtime/GameplayAbilities/Public/GameplayEffectTypes.h
--- a/Engine/Source/Runtime/GameplayAbilities/Public/GameplayEffectTypes.h
+++ b/Engine/Source/Runtime/GameplayAbilities/Public/GameplayEffectTypes.h
@@ -280,7 +280,7 @@ struct GAMEPLAYABILITIES_API FGameplayEffectContext
 	/** Returns the immediate instigator that applied this effect */
 	virtual AActor* GetInstigator() const
 	{
-		return Instigator;
+		return Instigator.Get();
 	}
 
 	/** Returns the ability system component of the instigator of this effect */
@@ -292,13 +292,13 @@ struct GAMEPLAYABILITIES_API FGameplayEffectContext
 	/** Returns the physical actor tied to the application of this effect */
 	virtual AActor* GetEffectCauser() const
 	{
-		return EffectCauser;
+		return EffectCauser.Get();
 	}
 
 	/** Should always return the original instigator that started the whole chain. Subclasses can override what this does */
 	virtual AActor* GetOriginalInstigator() const
 	{
-		return Instigator;
+		return Instigator.Get();
 	}
 
 	/** Returns the ability system component of the instigator that started the whole chain */
@@ -339,7 +339,7 @@ struct GAMEPLAYABILITIES_API FGameplayEffectContext
 
 	virtual FString ToString() const
 	{
-		return Instigator ? Instigator->GetName() : FString(TEXT("NONE"));
+		return Instigator.IsValid() ? Instigator->GetName() : FString(TEXT("NONE"));
 	}
 
 	virtual UScriptStruct* GetScriptStruct() const
@@ -369,11 +369,11 @@ protected:
 
 	/** Instigator actor, the actor that owns the ability system component */
 	UPROPERTY()
-	AActor* Instigator;
+	TWeakObjectPtr<AActor> Instigator;
 
 	/** The physical actor that actually did the damage, can be a weapon or projectile */
 	UPROPERTY()
-	AActor* EffectCauser;
+	TWeakObjectPtr<AActor> EffectCauser;
 
 	/** The ability system component that's bound to instigator */
 	UPROPERTY(NotReplicated)
