#UE4

- Adding optional tag parameter to CanActivateAbility

diff --git a/Engine/Source/Runtime/GameplayAbilities/Private/Abilities/GameplayAbility.cpp b/Engine/Source/Runtime/GameplayAbilities/Private/Abilities/GameplayAbility.cpp
--- a/Engine/Source/Runtime/GameplayAbilities/Private/Abilities/GameplayAbility.cpp
+++ b/Engine/Source/Runtime/GameplayAbilities/Private/Abilities/GameplayAbility.cpp
@@ -133,10 +133,14 @@ bool UGameplayAbility::IsSupportedForNetworking() const
 	return Supported;
 }
 
-bool UGameplayAbility::CanActivateAbility(const FGameplayAbilitySpecHandle Handle, const FGameplayAbilityActorInfo* ActorInfo) const
+bool UGameplayAbility::CanActivateAbility(const FGameplayAbilitySpecHandle Handle, const FGameplayAbilityActorInfo* ActorInfo, OUT FGameplayTagContainer* OptionalRelevantTags) const
 {
 	SetCurrentActorInfo(Handle, ActorInfo);
 
+	//make into a reference for simplicity
+	FGameplayTagContainer DummyContainer;
+	FGameplayTagContainer& OutTags = OptionalRelevantTags ? *OptionalRelevantTags : DummyContainer;
+
 	if (ActorInfo->AbilitySystemComponent->GetUserAbilityActivationInhibited())
 	{
 		/**
@@ -172,7 +176,8 @@ bool UGameplayAbility::CanActivateAbility(const FGameplayAbilitySpecHandle Handl
 
 	if (HasBlueprintCanUse)
 	{
-		if (K2_CanActivateAbility(*ActorInfo) == false)
+
+		if (K2_CanActivateAbility(*ActorInfo, OutTags) == false)
 		{
 			ABILITY_LOG(Log, TEXT("CanActivateAbility %s failed, blueprint refused"), *GetName());
 			return false;
diff --git a/Engine/Source/Runtime/GameplayAbilities/Private/Abilities/GameplayAbility_CharacterJump.cpp b/Engine/Source/Runtime/GameplayAbilities/Private/Abilities/GameplayAbility_CharacterJump.cpp
--- a/Engine/Source/Runtime/GameplayAbilities/Private/Abilities/GameplayAbility_CharacterJump.cpp
+++ b/Engine/Source/Runtime/GameplayAbilities/Private/Abilities/GameplayAbility_CharacterJump.cpp
@@ -39,9 +39,9 @@ void UGameplayAbility_CharacterJump::InputReleased(const FGameplayAbilitySpecHan
 	}
 }
 
-bool UGameplayAbility_CharacterJump::CanActivateAbility(const FGameplayAbilitySpecHandle Handle, const FGameplayAbilityActorInfo* ActorInfo) const
+bool UGameplayAbility_CharacterJump::CanActivateAbility(const FGameplayAbilitySpecHandle Handle, const FGameplayAbilityActorInfo* ActorInfo, OUT FGameplayTagContainer* OptionalRelevantTags) const
 {
-	if (!Super::CanActivateAbility(Handle, ActorInfo))
+	if (!Super::CanActivateAbility(Handle, ActorInfo, OptionalRelevantTags))
 	{
 		return false;
 	}
diff --git a/Engine/Source/Runtime/GameplayAbilities/Public/Abilities/GameplayAbility.h b/Engine/Source/Runtime/GameplayAbilities/Public/Abilities/GameplayAbility.h
--- a/Engine/Source/Runtime/GameplayAbilities/Public/Abilities/GameplayAbility.h
+++ b/Engine/Source/Runtime/GameplayAbilities/Public/Abilities/GameplayAbility.h
@@ -119,7 +119,7 @@ public:
 	virtual void InputReleased(const FGameplayAbilitySpecHandle Handle, const FGameplayAbilityActorInfo* ActorInfo, const FGameplayAbilityActivationInfo ActivationInfo) {};
 
 	/** Returns true if this ability can be activated right now. Has no side effects */
-	virtual bool CanActivateAbility(const FGameplayAbilitySpecHandle Handle, const FGameplayAbilityActorInfo* ActorInfo) const;
+	virtual bool CanActivateAbility(const FGameplayAbilitySpecHandle Handle, const FGameplayAbilityActorInfo* ActorInfo, OUT FGameplayTagContainer* OptionalRelevantTags = nullptr) const;
 
 	/** Returns true if this ability can be triggered right now. Has no side effects */
 	virtual bool ShouldAbilityRespondToEvent(FGameplayTag EventTag, const FGameplayEventData* Payload) const;
@@ -242,7 +242,7 @@ protected:
 	
 	/** Returns true if this ability can be activated right now. Has no side effects */
 	UFUNCTION(BlueprintImplementableEvent, Category = Ability, FriendlyName="CanActivateAbility")
-	virtual bool K2_CanActivateAbility(FGameplayAbilityActorInfo ActorInfo) const;
+	virtual bool K2_CanActivateAbility(FGameplayAbilityActorInfo ActorInfo, FGameplayTagContainer& RelevantTags) const;
 
 	bool HasBlueprintCanUse;
 
diff --git a/Engine/Source/Runtime/GameplayAbilities/Public/Abilities/GameplayAbility_CharacterJump.h b/Engine/Source/Runtime/GameplayAbilities/Public/Abilities/GameplayAbility_CharacterJump.h
--- a/Engine/Source/Runtime/GameplayAbilities/Public/Abilities/GameplayAbility_CharacterJump.h
+++ b/Engine/Source/Runtime/GameplayAbilities/Public/Abilities/GameplayAbility_CharacterJump.h
@@ -15,7 +15,7 @@ class GAMEPLAYABILITIES_API UGameplayAbility_CharacterJump : public UGameplayAbi
 
 public:
 
-	virtual bool CanActivateAbility(const FGameplayAbilitySpecHandle Handle, const FGameplayAbilityActorInfo* ActorInfo) const override;
+	virtual bool CanActivateAbility(const FGameplayAbilitySpecHandle Handle, const FGameplayAbilityActorInfo* ActorInfo, OUT FGameplayTagContainer* OptionalRelevantTags = nullptr) const override;
 	
 	virtual void ActivateAbility(const FGameplayAbilitySpecHandle Handle, const FGameplayAbilityActorInfo* OwnerInfo, const FGameplayAbilityActivationInfo ActivationInfo) override;
 
