chainging the check to be IncludeParentTags from Explicit since the secondary checks all use IncludeParentTags

diff --git a/Engine/Source/Runtime/GameplayAbilities/Private/GameplayEffect.cpp b/Engine/Source/Runtime/GameplayAbilities/Private/GameplayEffect.cpp
--- a/Engine/Source/Runtime/GameplayAbilities/Private/GameplayEffect.cpp
+++ b/Engine/Source/Runtime/GameplayAbilities/Private/GameplayEffect.cpp
@@ -1858,7 +1858,7 @@ bool FActiveGameplayEffectsContainer::CheckApplicationImmunity(const FGameplayEf
 	
 
 	// Quick map test
-	if (!ApplicationImmunityGameplayTagCountContainer.HasAnyMatchingGameplayTags(*Tags, EGameplayTagMatchType::Explicit, false))
+	if (!ApplicationImmunityGameplayTagCountContainer.HasAnyMatchingGameplayTags(*Tags, EGameplayTagMatchType::IncludeParentTags, false))
 	{
 		return false;
 	}
