- Change FGameplayEffectContextHandle's FGameplayEffectContext constructor to be explicit to prevent accidental conversions causing problems
- Fix the periodic damage unit test (UE-6389)
- Fix issue with periodic effects potentially dropping their last tick

diff --git a/Engine/Source/Runtime/GameplayAbilities/Private/GameplayEffect.cpp b/Engine/Source/Runtime/GameplayAbilities/Private/GameplayEffect.cpp
--- a/Engine/Source/Runtime/GameplayAbilities/Private/GameplayEffect.cpp
+++ b/Engine/Source/Runtime/GameplayAbilities/Private/GameplayEffect.cpp
@@ -1919,13 +1919,16 @@ void FActiveGameplayEffectsContainer::CheckDuration(FActiveGameplayEffectHandle
 			if (Duration > 0.f && (Effect.StartWorldTime + Duration) < CurrentTime)
 			{
 				// This gameplay effect has hit its duration. Check if it needs to execute one last time before removing it.
-				if (Effect.PeriodHandle.IsValid())
+				if (Effect.PeriodHandle.IsValid() && TimerManager.TimerExists(Effect.PeriodHandle))
 				{
 					float PeriodTimeRemaining = TimerManager.GetTimerRemaining(Effect.PeriodHandle);
-					if (FMath::IsNearlyEqual(PeriodTimeRemaining, 0.f))
+					if (PeriodTimeRemaining <= KINDA_SMALL_NUMBER)
 					{
 						ExecuteActiveEffectsFrom(Effect.Spec);
 					}
+
+					// Forcibly clear the periodic ticks because this effect is going to be removed
+					TimerManager.ClearTimer(Effect.PeriodHandle);
 				}
 
 				InternalRemoveActiveGameplayEffect(Idx);
diff --git a/Engine/Source/Runtime/GameplayAbilities/Private/GameplayEffectTests.cpp b/Engine/Source/Runtime/GameplayAbilities/Private/GameplayEffectTests.cpp
--- a/Engine/Source/Runtime/GameplayAbilities/Private/GameplayEffectTests.cpp
+++ b/Engine/Source/Runtime/GameplayAbilities/Private/GameplayEffectTests.cpp
@@ -191,29 +191,40 @@ public: // the tests
 		{
 			CONSTRUCT_CLASS(UGameplayEffect, BaseDmgEffect);
 			AddModifier(BaseDmgEffect, GET_FIELD_CHECKED(UAbilitySystemTestAttributeSet, Health), EGameplayModOp::Additive, FScalableFloat(-DamagePerPeriod));
-			BaseDmgEffect->Duration.Value = NumPeriods * PeriodSecs;
+			BaseDmgEffect->DurationPolicy = EGameplayEffectDurationType::HasDuration;
+			BaseDmgEffect->DurationMagnitude = FGameplayEffectModifierMagnitude(FScalableFloat(NumPeriods * PeriodSecs));
 			BaseDmgEffect->Period.Value = PeriodSecs;
 
 			SourceComponent->ApplyGameplayEffectToTarget(BaseDmgEffect, DestComponent, 1.f);
 		}
 
-		// make sure health was not reduced on the first tick
-		TestEqual(SKILL_TEST_TEXT("Health Reduced"), DestComponent->GetSet<UAbilitySystemTestAttributeSet>()->Health, StartingHealth);
+		int32 NumApplications = 0;
 
-		for (int32 i=1;i<=NumPeriods;++i)
+		// Tick a small number to verify the application tick
+		TickWorld(SMALL_NUMBER);
+		++NumApplications;
+
+		TestEqual(SKILL_TEST_TEXT("Health Reduced"), DestComponent->GetSet<UAbilitySystemTestAttributeSet>()->Health, StartingHealth - (DamagePerPeriod * NumApplications));
+
+		// Tick a bit more to address possible floating point issues
+		TickWorld(PeriodSecs * .1f);
+
+		for (int32 i = 0; i < NumPeriods; ++i)
 		{
 			// advance time by one period
 			TickWorld(PeriodSecs);
 
+			++NumApplications;
+
 			// check that health has been reduced
-			TestEqual(SKILL_TEST_TEXT("Health Reduced"), DestComponent->GetSet<UAbilitySystemTestAttributeSet>()->Health, StartingHealth - i*DamagePerPeriod);
+			TestEqual(SKILL_TEST_TEXT("Health Reduced"), DestComponent->GetSet<UAbilitySystemTestAttributeSet>()->Health, StartingHealth - (DamagePerPeriod * NumApplications));
 		}
 
 		// advance time by one extra period
 		TickWorld(PeriodSecs);
 
 		// should not have reduced further
-		TestEqual(SKILL_TEST_TEXT("Health Reduced"), DestComponent->GetSet<UAbilitySystemTestAttributeSet>()->Health, StartingHealth - NumPeriods*DamagePerPeriod);
+		TestEqual(SKILL_TEST_TEXT("Health Reduced"), DestComponent->GetSet<UAbilitySystemTestAttributeSet>()->Health, StartingHealth - (DamagePerPeriod * NumApplications));
 
 		// TODO: test that the effect is no longer applied
 	}
@@ -356,5 +367,3 @@ namespace
 }
 
 #endif //WITH_EDITOR
-
-
diff --git a/Engine/Source/Runtime/GameplayAbilities/Public/GameplayEffectTypes.h b/Engine/Source/Runtime/GameplayAbilities/Public/GameplayEffectTypes.h
--- a/Engine/Source/Runtime/GameplayAbilities/Public/GameplayEffectTypes.h
+++ b/Engine/Source/Runtime/GameplayAbilities/Public/GameplayEffectTypes.h
@@ -415,7 +415,7 @@ struct FGameplayEffectContextHandle
 	}
 
 	/** Constructs from an existing context, should be allocated by new */
-	FGameplayEffectContextHandle(FGameplayEffectContext* DataPtr)
+	explicit FGameplayEffectContextHandle(FGameplayEffectContext* DataPtr)
 	{
 		Data = TSharedPtr<FGameplayEffectContext>(DataPtr);
 	}
