Added code to cancel GEs when an OrionChar dies. Modified queries to include rejection containers.

diff --git a/Engine/Source/Runtime/GameplayAbilities/Private/AbilitySystemComponent.cpp b/Engine/Source/Runtime/GameplayAbilities/Private/AbilitySystemComponent.cpp
--- a/Engine/Source/Runtime/GameplayAbilities/Private/AbilitySystemComponent.cpp
+++ b/Engine/Source/Runtime/GameplayAbilities/Private/AbilitySystemComponent.cpp
@@ -834,6 +834,11 @@ void UAbilitySystemComponent::RemoveActiveEffectsWithTags(const FGameplayTagCont
 	RemoveActiveEffects(FActiveGameplayEffectQuery(&Tags));
 }
 
+void UAbilitySystemComponent::RemoveActiveEffects(FGameplayTagContainer OwningTags, FGameplayTagContainer EffectTags, FGameplayTagContainer OwningTags_Rejection, FGameplayTagContainer EffectTags_Rejection)
+{
+	return ActiveGameplayEffects.RemoveActiveEffects(FActiveGameplayEffectQuery(&OwningTags, &EffectTags, &OwningTags_Rejection, &EffectTags_Rejection));
+}
+
 void UAbilitySystemComponent::RemoveActiveEffects(const FActiveGameplayEffectQuery Query)
 {
 	return ActiveGameplayEffects.RemoveActiveEffects(Query);
diff --git a/Engine/Source/Runtime/GameplayAbilities/Private/GameplayEffect.cpp b/Engine/Source/Runtime/GameplayAbilities/Private/GameplayEffect.cpp
--- a/Engine/Source/Runtime/GameplayAbilities/Private/GameplayEffect.cpp
+++ b/Engine/Source/Runtime/GameplayAbilities/Private/GameplayEffect.cpp
@@ -2244,7 +2244,7 @@ bool FActiveGameplayEffectQuery::Matches(const FActiveGameplayEffect& Effect) co
 		if (!Effect.Spec.Def->InheritableOwnedTagsContainer.CombinedTags.MatchesAny(*OwningTagContainer, false) &&
 			!Effect.Spec.DynamicGrantedTags.MatchesAny(*OwningTagContainer, false))
 		{
-			// this doesnt match our Tags so bail
+			// this doesn't match our Tags so bail
 			return false;
 		}
 	}	
@@ -2254,7 +2254,17 @@ bool FActiveGameplayEffectQuery::Matches(const FActiveGameplayEffect& Effect) co
 	{
 		if (!Effect.Spec.Def->InheritableGameplayEffectTags.CombinedTags.MatchesAny(*EffectTagContainer, false))
 		{
-			// this doesnt match our Tags so bail
+			// this doesn't match our Tags so bail
+			return false;
+		}
+	}
+
+	// if we are just looking for Tags on the Effect then look at the Gameplay Effect Tags
+	if (EffectTagContainer_Rejection)
+	{
+		if (Effect.Spec.Def->InheritableGameplayEffectTags.CombinedTags.MatchesAny(*EffectTagContainer_Rejection, false))
+		{
+			// this matches our Rejection Tags so bail
 			return false;
 		}
 	}
diff --git a/Engine/Source/Runtime/GameplayAbilities/Public/AbilitySystemComponent.h b/Engine/Source/Runtime/GameplayAbilities/Public/AbilitySystemComponent.h
--- a/Engine/Source/Runtime/GameplayAbilities/Public/AbilitySystemComponent.h
+++ b/Engine/Source/Runtime/GameplayAbilities/Public/AbilitySystemComponent.h
@@ -372,6 +372,10 @@ class GAMEPLAYABILITIES_API UAbilitySystemComponent : public UActorComponent, pu
 	UFUNCTION(BlueprintCallable, Category = GameplayEffects)
 	void RemoveActiveEffectsWithTags(FGameplayTagContainer Tags);
 
+	/** Constructs a query and removess active effects as appropriate */
+	UFUNCTION(BlueprintCallable, Category = GameplayEffects)
+	void RemoveActiveEffects(FGameplayTagContainer OwningTags, FGameplayTagContainer EffectTags, FGameplayTagContainer OwningTags_Rejection, FGameplayTagContainer EffectTags_Rejection);
+
 	/** Removes all active effects that match given query */
 	void RemoveActiveEffects(const FActiveGameplayEffectQuery Query);
 
diff --git a/Engine/Source/Runtime/GameplayAbilities/Public/GameplayEffect.h b/Engine/Source/Runtime/GameplayAbilities/Public/GameplayEffect.h
--- a/Engine/Source/Runtime/GameplayAbilities/Public/GameplayEffect.h
+++ b/Engine/Source/Runtime/GameplayAbilities/Public/GameplayEffect.h
@@ -1099,12 +1099,16 @@ struct FActiveGameplayEffectQuery
 	FActiveGameplayEffectQuery()
 		: OwningTagContainer(NULL)
 		, EffectTagContainer(NULL)
+		, OwningTagContainer_Rejection(NULL)
+		, EffectTagContainer_Rejection(NULL)
 	{
 	}
 
 	FActiveGameplayEffectQuery(const FGameplayTagContainer* InOwningTagContainer, FGameplayAttribute InModifyingAttribute = FGameplayAttribute())
 		: OwningTagContainer(InOwningTagContainer)
 		, EffectTagContainer(NULL)
+		, OwningTagContainer_Rejection(NULL)
+		, EffectTagContainer_Rejection(NULL)
 		, ModifyingAttribute(InModifyingAttribute)
 	{
 	}
@@ -1112,6 +1116,26 @@ struct FActiveGameplayEffectQuery
 	FActiveGameplayEffectQuery(const FGameplayTagContainer* InOwningTagContainer, const FGameplayTagContainer* InEffectTagContainer, FGameplayAttribute InModifyingAttribute = FGameplayAttribute())
 		: OwningTagContainer(InOwningTagContainer)
 		, EffectTagContainer(InEffectTagContainer)
+		, OwningTagContainer_Rejection(NULL)
+		, EffectTagContainer_Rejection(NULL)
+		, ModifyingAttribute(InModifyingAttribute)
+	{
+	}
+
+	FActiveGameplayEffectQuery(const FGameplayTagContainer* InOwningTagContainer, const FGameplayTagContainer* InEffectTagContainer, const FGameplayTagContainer* InOwningTagContainer_Rejection, FGameplayAttribute InModifyingAttribute = FGameplayAttribute())
+		: OwningTagContainer(InOwningTagContainer)
+		, EffectTagContainer(InEffectTagContainer)
+		, OwningTagContainer_Rejection(InOwningTagContainer_Rejection)
+		, EffectTagContainer_Rejection(NULL)
+		, ModifyingAttribute(InModifyingAttribute)
+	{
+	}
+
+	FActiveGameplayEffectQuery(const FGameplayTagContainer* InOwningTagContainer, const FGameplayTagContainer* InEffectTagContainer, const FGameplayTagContainer* InOwningTagContainer_Rejection, const FGameplayTagContainer* InEffectTagContainer_Rejection, FGameplayAttribute InModifyingAttribute = FGameplayAttribute())
+		: OwningTagContainer(InOwningTagContainer)
+		, EffectTagContainer(InEffectTagContainer)
+		, OwningTagContainer_Rejection(InOwningTagContainer_Rejection)
+		, EffectTagContainer_Rejection(InEffectTagContainer_Rejection)
 		, ModifyingAttribute(InModifyingAttribute)
 	{
 	}
@@ -1124,6 +1148,12 @@ struct FActiveGameplayEffectQuery
 	/** used to match with InheritableGameplayEffectTags */
 	const FGameplayTagContainer* EffectTagContainer;
 
+	/** used to reject matches with InheritableOwnedTagsContainer */
+	const FGameplayTagContainer* OwningTagContainer_Rejection;
+
+	/** used to reject matches with InheritableGameplayEffectTags */
+	const FGameplayTagContainer* EffectTagContainer_Rejection;
+
 	/** Matches on GameplayEffects which modify given attribute */
 	FGameplayAttribute ModifyingAttribute;
 };
