- Change how FGameplayTagContainer works to address issues with conflation of explicit and parent tags
   - Container no longer has an explicit vs. parent mode, instead always keeps a container of the explicitly added tags while configuring the delegates and count maps to include parent tags
   - Make all of the variables of the container private, provide accessors where appropriate
   - Comment pass on the tag container
   - Rename UpdateTagMap to UpdateTagCount now that it runs on more than just a map
- Rename FActiveGameplayEffectsContainer::CheckApplicationImmunity to FActiveGameplayEffectsContainer::HasApplicationImmunityToSpec to alleviate confusion
- Fix AbilityTask_WaitDelay showing the incorrect value for "TimeLeft" while debugging

diff --git a/Engine/Source/Runtime/GameplayAbilities/Private/Abilities/Tasks/AbilityTask_WaitDelay.cpp b/Engine/Source/Runtime/GameplayAbilities/Private/Abilities/Tasks/AbilityTask_WaitDelay.cpp
--- a/Engine/Source/Runtime/GameplayAbilities/Private/Abilities/Tasks/AbilityTask_WaitDelay.cpp
+++ b/Engine/Source/Runtime/GameplayAbilities/Private/Abilities/Tasks/AbilityTask_WaitDelay.cpp
@@ -40,6 +40,6 @@ void UAbilityTask_WaitDelay::OnTimeFinish()
 
 FString UAbilityTask_WaitDelay::GetDebugString() const
 {
-	float TimeLeft = GetWorld()->GetTimeSeconds() - TimeStarted;
+	float TimeLeft = Time - GetWorld()->TimeSince(TimeStarted);
 	return FString::Printf(TEXT("WaitDelay. Time: %.2f. TimeLeft: %.2f"), Time, TimeLeft);
 }
\ No newline at end of file
diff --git a/Engine/Source/Runtime/GameplayAbilities/Private/AbilitySystemComponent.cpp b/Engine/Source/Runtime/GameplayAbilities/Private/AbilitySystemComponent.cpp
--- a/Engine/Source/Runtime/GameplayAbilities/Private/AbilitySystemComponent.cpp
+++ b/Engine/Source/Runtime/GameplayAbilities/Private/AbilitySystemComponent.cpp
@@ -32,7 +32,7 @@ static FAutoConsoleVariableRef CVarDebugGameplayCues(
 
 UAbilitySystemComponent::UAbilitySystemComponent(const FObjectInitializer& ObjectInitializer)
 	: Super(ObjectInitializer)
-	, GameplayTagCountContainer(EGameplayTagMatchType::IncludeParentTags)
+	, GameplayTagCountContainer()
 {
 	bWantsInitializeComponent = true;
 
@@ -339,12 +339,12 @@ void UAbilitySystemComponent::CaptureAttributeForGameplayEffect(OUT FGameplayEff
 
 FOnGameplayEffectTagCountChanged& UAbilitySystemComponent::RegisterGameplayTagEvent(FGameplayTag Tag)
 {
-	return GameplayTagCountContainer.GameplayTagEventMap.FindOrAdd(Tag);
+	return GameplayTagCountContainer.RegisterGameplayTagEvent(Tag);
 }
 
 FOnGameplayEffectTagCountChanged& UAbilitySystemComponent::RegisterGenericGameplayTagEvent()
 {
-	return GameplayTagCountContainer.OnAnyTagChangeDelegate;
+	return GameplayTagCountContainer.RegisterGenericGameplayEvent();
 }
 
 FOnGameplayAttributeChange& UAbilitySystemComponent::RegisterGameplayAttributeEvent(FGameplayAttribute Attribute)
@@ -383,30 +383,24 @@ const FGameplayEffectAttributeCaptureDefinition& UAbilitySystemComponent::GetInc
 void UAbilitySystemComponent::GetOwnedGameplayTags(FGameplayTagContainer& TagContainer) const
 {
 	SCOPE_CYCLE_COUNTER(STAT_GameplayEffectsGetOwnedTags);
-	for (auto It = GameplayTagCountContainer.GameplayTagCountMap.CreateConstIterator(); It; ++It)
-	{
-		if (It.Value() > 0)
-		{
-			TagContainer.AddTagFast(It.Key());
-		}
-	}
+	TagContainer.AppendTags(GameplayTagCountContainer.GetExplicitGameplayTags());
 }
 
 bool UAbilitySystemComponent::HasMatchingGameplayTag(FGameplayTag TagToCheck) const
 {
-	return GameplayTagCountContainer.HasMatchingGameplayTag(TagToCheck, EGameplayTagMatchType::Explicit);
+	return GameplayTagCountContainer.HasMatchingGameplayTag(TagToCheck);
 }
 
 bool UAbilitySystemComponent::HasAllMatchingGameplayTags(const FGameplayTagContainer& TagContainer, bool bCountEmptyAsMatch) const
 {
 	SCOPE_CYCLE_COUNTER(STAT_GameplayEffectsHasAllTags);
-	return GameplayTagCountContainer.HasAllMatchingGameplayTags(TagContainer, EGameplayTagMatchType::Explicit, bCountEmptyAsMatch);
+	return GameplayTagCountContainer.HasAllMatchingGameplayTags(TagContainer, bCountEmptyAsMatch);
 }
 
 bool UAbilitySystemComponent::HasAnyMatchingGameplayTags(const FGameplayTagContainer& TagContainer, bool bCountEmptyAsMatch) const
 {
 	SCOPE_CYCLE_COUNTER(STAT_GameplayEffectsHasAnyTag);
-	return GameplayTagCountContainer.HasAnyMatchingGameplayTags(TagContainer, EGameplayTagMatchType::Explicit, bCountEmptyAsMatch);
+	return GameplayTagCountContainer.HasAnyMatchingGameplayTags(TagContainer, bCountEmptyAsMatch);
 }
 
 void UAbilitySystemComponent::AddLooseGameplayTag(const FGameplayTag& GameplayTag, int32 Count)
@@ -433,12 +427,12 @@ void UAbilitySystemComponent::RemoveLooseGameplayTags(const FGameplayTagContaine
 
 void UAbilitySystemComponent::UpdateTagMap(const FGameplayTag& BaseTag, int32 CountDelta)
 {
-	GameplayTagCountContainer.UpdateTagMap(BaseTag, CountDelta);
+	GameplayTagCountContainer.UpdateTagCount(BaseTag, CountDelta);
 }
 
 void UAbilitySystemComponent::UpdateTagMap(const FGameplayTagContainer& Container, int32 CountDelta)
 {
-	GameplayTagCountContainer.UpdateTagMap(Container, CountDelta);
+	GameplayTagCountContainer.UpdateTagCount(Container, CountDelta);
 }
 
 // ------------------------------------------------------------------------
@@ -472,7 +466,7 @@ FActiveGameplayEffectHandle UAbilitySystemComponent::ApplyGameplayEffectSpecToSe
 	}
 
 	// Are we currently immune to this? (ApplicationImmunity)
-	if (ActiveGameplayEffects.CheckApplicationImmunity(Spec))
+	if (ActiveGameplayEffects.HasApplicationImmunityToSpec(Spec))
 	{
 		return FActiveGameplayEffectHandle();
 	}
diff --git a/Engine/Source/Runtime/GameplayAbilities/Private/AbilitySystemComponent_Abilities.cpp b/Engine/Source/Runtime/GameplayAbilities/Private/AbilitySystemComponent_Abilities.cpp
--- a/Engine/Source/Runtime/GameplayAbilities/Private/AbilitySystemComponent_Abilities.cpp
+++ b/Engine/Source/Runtime/GameplayAbilities/Private/AbilitySystemComponent_Abilities.cpp
@@ -436,12 +436,12 @@ void UAbilitySystemComponent::CancelAbilities(const FGameplayTagContainer* WithT
 
 void UAbilitySystemComponent::BlockAbilitiesWithTags(const FGameplayTagContainer Tags)
 {
-	BlockedAbilityTags.UpdateTagMap(Tags, 1);
+	BlockedAbilityTags.UpdateTagCount(Tags, 1);
 }
 
 void UAbilitySystemComponent::UnBlockAbilitiesWithTags(const FGameplayTagContainer Tags)
 {
-	BlockedAbilityTags.UpdateTagMap(Tags, -1);
+	BlockedAbilityTags.UpdateTagCount(Tags, -1);
 }
 
 void UAbilitySystemComponent::BlockAbilityByInputID(int32 InputID)
@@ -551,7 +551,7 @@ bool UAbilitySystemComponent::TryActivateAbility(FGameplayAbilitySpecHandle Hand
 	}
 
 	// Check if any of this ability's tags are currently blocked
-	if (BlockedAbilityTags.HasAnyMatchingGameplayTags(Ability->AbilityTags, EGameplayTagMatchType::IncludeParentTags, false))
+	if (Ability->AbilityTags.MatchesAny(BlockedAbilityTags.GetExplicitGameplayTags(), false))
 	{
 		return false;
 	}
diff --git a/Engine/Source/Runtime/GameplayAbilities/Private/AbilitySystemStats.cpp b/Engine/Source/Runtime/GameplayAbilities/Private/AbilitySystemStats.cpp
--- a/Engine/Source/Runtime/GameplayAbilities/Private/AbilitySystemStats.cpp
+++ b/Engine/Source/Runtime/GameplayAbilities/Private/AbilitySystemStats.cpp
@@ -19,6 +19,6 @@ DEFINE_STAT(STAT_InitAttributeSetDefaults);
 DEFINE_STAT(STAT_TickAbilityTasks);
 DEFINE_STAT(STAT_FindAbilitySpecFromHandle);
 DEFINE_STAT(STAT_AggregatorEvaluate);
-DEFINE_STAT(STAT_CheckApplicationImmunity);
+DEFINE_STAT(STAT_HasApplicationImmunityToSpec);
 
 
diff --git a/Engine/Source/Runtime/GameplayAbilities/Private/GameplayEffect.cpp b/Engine/Source/Runtime/GameplayAbilities/Private/GameplayEffect.cpp
--- a/Engine/Source/Runtime/GameplayAbilities/Private/GameplayEffect.cpp
+++ b/Engine/Source/Runtime/GameplayAbilities/Private/GameplayEffect.cpp
@@ -1665,8 +1665,8 @@ void FActiveGameplayEffectsContainer::AddActiveGameplayEffectGrantedTagsAndModif
 
 	Owner->UpdateTagMap(Effect.Spec.DynamicGrantedTags, 1);
 
-	ApplicationImmunityGameplayTagCountContainer.UpdateTagMap(Effect.Spec.Def->GrantedApplicationImmunityTags.RequireTags, 1);
-	ApplicationImmunityGameplayTagCountContainer.UpdateTagMap(Effect.Spec.Def->GrantedApplicationImmunityTags.IgnoreTags, 1);
+	ApplicationImmunityGameplayTagCountContainer.UpdateTagCount(Effect.Spec.Def->GrantedApplicationImmunityTags.RequireTags, 1);
+	ApplicationImmunityGameplayTagCountContainer.UpdateTagCount(Effect.Spec.Def->GrantedApplicationImmunityTags.IgnoreTags, 1);
 
 	for (const FGameplayEffectCue& Cue : Effect.Spec.Def->GameplayCues)
 	{
@@ -1799,8 +1799,8 @@ void FActiveGameplayEffectsContainer::RemoveActiveGameplayEffectGrantedTagsAndMo
 	IGameplayTagsModule& GameplayTagsModule = IGameplayTagsModule::Get();
 	Owner->UpdateTagMap(Effect.Spec.Def->InheritableOwnedTagsContainer.CombinedTags, -1);
 
-	ApplicationImmunityGameplayTagCountContainer.UpdateTagMap(Effect.Spec.Def->GrantedApplicationImmunityTags.RequireTags, -1);
-	ApplicationImmunityGameplayTagCountContainer.UpdateTagMap(Effect.Spec.Def->GrantedApplicationImmunityTags.IgnoreTags, -1);
+	ApplicationImmunityGameplayTagCountContainer.UpdateTagCount(Effect.Spec.Def->GrantedApplicationImmunityTags.RequireTags, -1);
+	ApplicationImmunityGameplayTagCountContainer.UpdateTagCount(Effect.Spec.Def->GrantedApplicationImmunityTags.IgnoreTags, -1);
 
 	Owner->UpdateTagMap(Effect.Spec.DynamicGrantedTags, -1);
 
@@ -1848,24 +1848,25 @@ void FActiveGameplayEffectsContainer::OnOwnerTagChange(FGameplayTag TagChange, i
 	}
 }
 
-bool FActiveGameplayEffectsContainer::CheckApplicationImmunity(const FGameplayEffectSpec& SpecToApply) const
+bool FActiveGameplayEffectsContainer::HasApplicationImmunityToSpec(const FGameplayEffectSpec& SpecToApply) const
 {
-	SCOPE_CYCLE_COUNTER(STAT_CheckApplicationImmunity)
+	SCOPE_CYCLE_COUNTER(STAT_HasApplicationImmunityToSpec)
 
-	const FGameplayTagContainer* Tags = SpecToApply.CapturedSourceTags.GetAggregatedTags();
-	if (!ensure(Tags))
+	const FGameplayTagContainer* AggregatedSourceTags = SpecToApply.CapturedSourceTags.GetAggregatedTags();
+	if (!ensure(AggregatedSourceTags))
+	{
 		return false;
-	
+	}
 
 	// Quick map test
-	if (!ApplicationImmunityGameplayTagCountContainer.HasAnyMatchingGameplayTags(*Tags, EGameplayTagMatchType::IncludeParentTags, false))
+	if (!AggregatedSourceTags->MatchesAny(ApplicationImmunityGameplayTagCountContainer.GetExplicitGameplayTags(), false))
 	{
 		return false;
 	}
 
 	for (const FActiveGameplayEffect& Effect : GameplayEffects)
 	{
-		if (Effect.Spec.Def->GrantedApplicationImmunityTags.IsEmpty() == false && Effect.Spec.Def->GrantedApplicationImmunityTags.RequirementsMet( *Tags ))
+		if (Effect.Spec.Def->GrantedApplicationImmunityTags.IsEmpty() == false && Effect.Spec.Def->GrantedApplicationImmunityTags.RequirementsMet( *AggregatedSourceTags ))
 		{
 			return true;
 		}
diff --git a/Engine/Source/Runtime/GameplayAbilities/Private/GameplayEffectTypes.cpp b/Engine/Source/Runtime/GameplayAbilities/Private/GameplayEffectTypes.cpp
--- a/Engine/Source/Runtime/GameplayAbilities/Private/GameplayEffectTypes.cpp
+++ b/Engine/Source/Runtime/GameplayAbilities/Private/GameplayEffectTypes.cpp
@@ -217,111 +217,111 @@ FString EGameplayEffectStackingPolicyToString(int32 Type)
 	return e->GetEnum(Type).ToString();
 }
 
-bool FGameplayTagCountContainer::HasMatchingGameplayTag(FGameplayTag TagToCheck, EGameplayTagMatchType::Type TagMatchType) const
+bool FGameplayTagCountContainer::HasMatchingGameplayTag(FGameplayTag TagToCheck) const
 {
-	if (TagMatchType == EGameplayTagMatchType::Explicit)
-	{
-		// Search for TagToCheck
-		const int32* Count = GameplayTagCountMap.Find(TagToCheck);
-		if (Count && *Count > 0)
-		{
-			return true;
-		}
-	}
-	else if (TagMatchType == EGameplayTagMatchType::IncludeParentTags)
-	{
-		// Search for TagToCheck or any of its parent tags
-		FGameplayTagContainer TagAndParentsContainer = IGameplayTagsModule::Get().GetGameplayTagsManager().RequestGameplayTagParents(TagToCheck);
-		for (auto TagIt = TagAndParentsContainer.CreateConstIterator(); TagIt; ++TagIt)
-		{
-			const int32* Count = GameplayTagCountMap.Find(*TagIt);
-			if (Count && *Count > 0)
-			{
-				return true;
-			}
-		}
-	}
-
-	return false;
+	return ExplicitTags.HasTag(TagToCheck, EGameplayTagMatchType::IncludeParentTags, EGameplayTagMatchType::Explicit);
 }
 
-bool FGameplayTagCountContainer::HasAllMatchingGameplayTags(const FGameplayTagContainer& TagContainer, EGameplayTagMatchType::Type TagMatchType, bool bCountEmptyAsMatch) const
+bool FGameplayTagCountContainer::HasAllMatchingGameplayTags(const FGameplayTagContainer& TagContainer, bool bCountEmptyAsMatch) const
 {
-	if (TagContainer.Num() == 0)
-		return bCountEmptyAsMatch;
+	return ExplicitTags.MatchesAll(TagContainer, bCountEmptyAsMatch);
+}
 
-	for (auto It = TagContainer.CreateConstIterator(); It; ++It)
-	{
-		if (!HasMatchingGameplayTag(*It, TagMatchType))
-		{
-			return false;
-		}
-	}
 
-	return true;
+bool FGameplayTagCountContainer::HasAnyMatchingGameplayTags(const FGameplayTagContainer& TagContainer, bool bCountEmptyAsMatch) const
+{
+	return ExplicitTags.MatchesAny(TagContainer, bCountEmptyAsMatch);
 }
 
-
-bool FGameplayTagCountContainer::HasAnyMatchingGameplayTags(const FGameplayTagContainer& TagContainer, EGameplayTagMatchType::Type TagMatchType, bool bCountEmptyAsMatch) const
+void FGameplayTagCountContainer::UpdateTagCount(const struct FGameplayTag& Tag, int32 CountDelta)
 {
-	if (TagContainer.Num() == 0)
-		return bCountEmptyAsMatch;
+	if (CountDelta != 0)
+	{
+		UpdateTagMap_Internal(Tag, CountDelta);
+	}
+}
 
-	for (auto It = TagContainer.CreateConstIterator(); It; ++It)
+void FGameplayTagCountContainer::UpdateTagCount(const struct FGameplayTagContainer& Container, int32 CountDelta)
+{	
+	if (CountDelta != 0)
 	{
-		if (HasMatchingGameplayTag(*It, TagMatchType))
+		for (auto TagIt = Container.CreateConstIterator(); TagIt; ++TagIt)
 		{
-			return true;
+			UpdateTagMap_Internal(*TagIt, CountDelta);
 		}
 	}
+}
 
-	return false;
+FOnGameplayEffectTagCountChanged& FGameplayTagCountContainer::RegisterGameplayTagEvent(const FGameplayTag& Tag)
+{
+	return GameplayTagEventMap.FindOrAdd(Tag);
+}
+
+FOnGameplayEffectTagCountChanged& FGameplayTagCountContainer::RegisterGenericGameplayEvent()
+{
+	return OnAnyTagChangeDelegate;
 }
 
-void FGameplayTagCountContainer::UpdateTagMap(const struct FGameplayTag& Tag, int32 CountDelta)
+const FGameplayTagContainer& FGameplayTagCountContainer::GetExplicitGameplayTags() const
 {
-	if (TagContainerType == EGameplayTagMatchType::Explicit)
+	return ExplicitTags;
+}
+
+void FGameplayTagCountContainer::UpdateTagMap_Internal(const FGameplayTag& Tag, int32 CountDelta)
+{
+	const bool bTagAlreadyExplicitlyExists = ExplicitTags.HasTag(Tag, EGameplayTagMatchType::Explicit, EGameplayTagMatchType::Explicit);
+
+	// Need special case handling to maintain the explicit tag list correctly, adding the tag to the list if it didn't previously exist and a
+	// positive delta comes in, and removing it from the list if it did exist and a negative delta comes in.
+	if (!bTagAlreadyExplicitlyExists)
 	{
-		// Update count of BaseTag
-		UpdateTagMap_Internal(Tag, CountDelta);
+		// Brand new tag with a positive delta needs to be explicitly added
+		if (CountDelta > 0)
+		{
+			ExplicitTags.AddTag(Tag);
+		}
+		// Block attempted reduction of non-explicit tags, as they were never truly added to the container directly
+		else
+		{
+			ABILITY_LOG(Warning, TEXT("Attempted to remove tag: %s from tag count container, but it is not explicitly in the container!"), *Tag.ToString());
+			return;
+		}
 	}
-	else if (TagContainerType == EGameplayTagMatchType::IncludeParentTags)
+	else if (CountDelta < 0)
 	{
-		// Update count of BaseTag and all of its parent tags
-		FGameplayTagContainer TagAndParentsContainer = IGameplayTagsModule::Get().GetGameplayTagsManager().RequestGameplayTagParents(Tag);
-		for (auto ParentTagIt = TagAndParentsContainer.CreateConstIterator(); ParentTagIt; ++ParentTagIt)
+		// Existing tag with a negative delta that would cause a complete removal needs to be explicitly removed; Count will be updated correctly below,
+		// so that part is skipped for now
+		int32& ExistingCount = GameplayTagCountMap.FindOrAdd(Tag);
+		if ((ExistingCount + CountDelta) <= 0)
 		{
-			UpdateTagMap_Internal(*ParentTagIt, CountDelta);
+			ExplicitTags.RemoveTag(Tag);
 		}
 	}
-}
 
-void FGameplayTagCountContainer::UpdateTagMap(const struct FGameplayTagContainer& Container, int32 CountDelta)
-{
-	for (auto TagIt = Container.CreateConstIterator(); TagIt; ++TagIt)
+	// Check if change delegates are required to fire for the tag or any of its parents based on the count change
+	FGameplayTagContainer TagAndParentsContainer = IGameplayTagsModule::Get().GetGameplayTagsManager().RequestGameplayTagParents(Tag);
+	for (auto CompleteTagIt = TagAndParentsContainer.CreateConstIterator(); CompleteTagIt; ++CompleteTagIt)
 	{
-		const FGameplayTag& BaseTag = *TagIt;
-		UpdateTagMap(BaseTag, CountDelta);
-	}
-}
+		const FGameplayTag& CurTag = *CompleteTagIt;
 
-void FGameplayTagCountContainer::UpdateTagMap_Internal(const FGameplayTag& Tag, int32 CountDelta)
-{
-	// Update count of Tag
-	int32& Count = GameplayTagCountMap.FindOrAdd(Tag);
+		// Get the current count of the specified tag. NOTE: Stored as a reference, so subsequent changes propogate to the map.
+		int32& TagCount = GameplayTagCountMap.FindOrAdd(CurTag);
 
-	bool WasZero = Count == 0;
-	Count = FMath::Max(Count + CountDelta, 0);
+		const int32 OldCount = TagCount;
 
-	// If we went from 0->1 or 1->0
-	if (WasZero || Count == 0)
-	{
-		OnAnyTagChangeDelegate.Broadcast(Tag, Count);
+		// Apply the delta to the count in the map
+		TagCount = FMath::Max(TagCount + CountDelta, 0);
 
-		FOnGameplayEffectTagCountChanged *Delegate = GameplayTagEventMap.Find(Tag);
-		if (Delegate)
+		// If a significant change (new addition or total removal) occurred, trigger related delegates
+		if (OldCount == 0 || TagCount == 0)
 		{
-			Delegate->Broadcast(Tag, Count);
+			OnAnyTagChangeDelegate.Broadcast(CurTag, TagCount);
+
+			FOnGameplayEffectTagCountChanged* CountChangeDelegate = GameplayTagEventMap.Find(CurTag);
+			if (CountChangeDelegate)
+			{
+				CountChangeDelegate->Broadcast(CurTag, TagCount);
+			}
 		}
 	}
 }
diff --git a/Engine/Source/Runtime/GameplayAbilities/Public/AbilitySystemStats.h b/Engine/Source/Runtime/GameplayAbilities/Public/AbilitySystemStats.h
--- a/Engine/Source/Runtime/GameplayAbilities/Public/AbilitySystemStats.h
+++ b/Engine/Source/Runtime/GameplayAbilities/Public/AbilitySystemStats.h
@@ -23,4 +23,4 @@ DECLARE_CYCLE_STAT_EXTERN(TEXT("InitAttributeSetDefaults"), STAT_InitAttributeSe
 DECLARE_CYCLE_STAT_EXTERN(TEXT("TickAbilityTasks"), STAT_TickAbilityTasks, STATGROUP_AbilitySystem, );
 DECLARE_CYCLE_STAT_EXTERN(TEXT("FindAbilitySpecFromHandle"), STAT_FindAbilitySpecFromHandle, STATGROUP_AbilitySystem, );
 DECLARE_CYCLE_STAT_EXTERN(TEXT("Aggregator Evaluate"), STAT_AggregatorEvaluate, STATGROUP_AbilitySystem, );
-DECLARE_CYCLE_STAT_EXTERN(TEXT("Check Application Immunity"), STAT_CheckApplicationImmunity, STATGROUP_AbilitySystem, );
\ No newline at end of file
+DECLARE_CYCLE_STAT_EXTERN(TEXT("Has Application Immunity To Spec"), STAT_HasApplicationImmunityToSpec, STATGROUP_AbilitySystem, );
\ No newline at end of file
diff --git a/Engine/Source/Runtime/GameplayAbilities/Public/GameplayEffect.h b/Engine/Source/Runtime/GameplayAbilities/Public/GameplayEffect.h
--- a/Engine/Source/Runtime/GameplayAbilities/Public/GameplayEffect.h
+++ b/Engine/Source/Runtime/GameplayAbilities/Public/GameplayEffect.h
@@ -1283,7 +1283,7 @@ struct FActiveGameplayEffectsContainer : public FFastArraySerializer
 
 	void OnOwnerTagChange(FGameplayTag TagChange, int32 NewCount);
 
-	bool CheckApplicationImmunity(const FGameplayEffectSpec& SpecToApply) const;
+	bool HasApplicationImmunityToSpec(const FGameplayEffectSpec& SpecToApply) const;
 
 private:
 
diff --git a/Engine/Source/Runtime/GameplayAbilities/Public/GameplayEffectTypes.h b/Engine/Source/Runtime/GameplayAbilities/Public/GameplayEffectTypes.h
--- a/Engine/Source/Runtime/GameplayAbilities/Public/GameplayEffectTypes.h
+++ b/Engine/Source/Runtime/GameplayAbilities/Public/GameplayEffectTypes.h
@@ -677,11 +677,11 @@ namespace EGameplayCueEvent
 
 DECLARE_DELEGATE_OneParam(FOnGameplayAttributeEffectExecuted, struct FGameplayModifierEvaluatedData&);
 
-DECLARE_MULTICAST_DELEGATE_TwoParams(FOnGameplayEffectTagCountChanged, const FGameplayTag, int32 );
+DECLARE_MULTICAST_DELEGATE_TwoParams(FOnGameplayEffectTagCountChanged, const FGameplayTag, int32);
 
 DECLARE_MULTICAST_DELEGATE(FOnActiveGameplayEffectRemoved);
 
-DECLARE_MULTICAST_DELEGATE_TwoParams(FOnGameplayAttributeChange, float ,const FGameplayEffectModCallbackData*);
+DECLARE_MULTICAST_DELEGATE_TwoParams(FOnGameplayAttributeChange, float, const FGameplayEffectModCallbackData*);
 
 DECLARE_DELEGATE_RetVal(FGameplayTagContainer, FGetGameplayTags);
 
@@ -689,38 +689,97 @@ DECLARE_DELEGATE_RetVal_OneParam(FOnGameplayEffectTagCountChanged&, FRegisterGam
 
 // -----------------------------------------------------------
 
-/** 
- *	Structure that contains a counted set of GameplayTags. Can optionally include parent tags
- *	
+/**
+ * Struct that tracks the number/count of tag applications within it. Explicitly tracks the tags added or removed,
+ * while simultaneously tracking the count of parent tags as well. Events/delegates are fired whenever the tag counts
+ * of any tag (explicit or parent) are modified.
  */
 struct FGameplayTagCountContainer
 {
+	// Constructor
 	FGameplayTagCountContainer()
-	: TagContainerType(EGameplayTagMatchType::Explicit)
-	{ }
-
-	FGameplayTagCountContainer(EGameplayTagMatchType::Type InTagContainerType)
-	: TagContainerType(InTagContainerType)
-	{ }
+	{}
 
-	bool HasMatchingGameplayTag(FGameplayTag TagToCheck, EGameplayTagMatchType::Type TagMatchType) const;
-	bool HasAllMatchingGameplayTags(const FGameplayTagContainer& TagContainer, EGameplayTagMatchType::Type TagMatchType, bool bCountEmptyAsMatch = true) const;
-	bool HasAnyMatchingGameplayTags(const FGameplayTagContainer& TagContainer, EGameplayTagMatchType::Type TagMatchType, bool bCountEmptyAsMatch = true) const;
-	void UpdateTagMap(const struct FGameplayTagContainer& Container, int32 CountDelta);
-	void UpdateTagMap(const struct FGameplayTag& Tag, int32 CountDelta);
+	/**
+	 * Check if the count container has a gameplay tag that matches against the specified tag (expands to include parents of asset tags)
+	 * 
+	 * @param TagToCheck	Tag to check for a match
+	 * 
+	 * @return True if the count container has a gameplay tag that matches, false if not
+	 */
+	bool HasMatchingGameplayTag(FGameplayTag TagToCheck) const;
 
-	TMap<struct FGameplayTag, FOnGameplayEffectTagCountChanged> GameplayTagEventMap;
-	TMap<struct FGameplayTag, int32> GameplayTagCountMap;
+	/**
+	 * Check if the count container has gameplay tags that matches against all of the specified tags (expands to include parents of asset tags)
+	 * 
+	 * @param TagContainer			Tag container to check for a match
+	 * @param bCountEmptyAsMatch	If true, the parameter tag container will count as matching, even if it's empty
+	 * 
+	 * @return True if the count container matches all of the gameplay tags
+	 */
+	bool HasAllMatchingGameplayTags(const FGameplayTagContainer& TagContainer, bool bCountEmptyAsMatch = true) const;
+	
+	/**
+	 * Check if the count container has gameplay tags that matches against any of the specified tags (expands to include parents of asset tags)
+	 * 
+	 * @param TagContainer			Tag container to check for a match
+	 * @param bCountEmptyAsMatch	If true, the parameter tag container will count as matching, even if it's empty
+	 * 
+	 * @return True if the count container matches any of the gameplay tags
+	 */
+	bool HasAnyMatchingGameplayTags(const FGameplayTagContainer& TagContainer, bool bCountEmptyAsMatch = true) const;
+	
+	/**
+	 * Update the specified container of tags by the specified delta, potentially causing an additional or removal from the explicit tag list
+	 * 
+	 * @param Container		Container of tags to update
+	 * @param CountDelta	Delta of the tag count to apply
+	 */
+	void UpdateTagCount(const FGameplayTagContainer& Container, int32 CountDelta);
+	
+	/**
+	 * Update the specified tag by the specified delta, potentially causing an additional or removal from the explicit tag list
+	 * 
+	 * @param Tag			Tag to update
+	 * @param CountDelta	Delta of the tag count to apply
+	 */
+	void UpdateTagCount(const FGameplayTag& Tag, int32 CountDelta);
 
-	/** This is called when any tag is added new or removed completely (going too or from 0 count). Not called for other count increases (e.g, going from 2-3 count) */
-	FOnGameplayEffectTagCountChanged	OnAnyTagChangeDelegate;
+	/**
+	 * Return delegate that can be bound to for when the specific tag's count changes to or off of zero
+	 *
+	 * @param Tag	Tag to get a delegate for
+	 * 
+	 * @return Delegate for when the specified tag's count changes to or off of zero
+	 */
+	FOnGameplayEffectTagCountChanged& RegisterGameplayTagEvent(const FGameplayTag& Tag);
+	
+	/**
+	 * Return delegate that can be bound to for when the any tag's count changes to or off of zero
+	 * 
+	 * @return Delegate for when any tag's count changes to or off of zero
+	 */
+	FOnGameplayEffectTagCountChanged& RegisterGenericGameplayEvent();
 
-	EGameplayTagMatchType::Type TagContainerType;
+	/** Simple accessor to the explicit gameplay tag list */
+	const FGameplayTagContainer& GetExplicitGameplayTags() const;
 
 private:
 
-	// Fixme: This may not be adding tag parents correctly. The TagContainer version of this function properly adds parent tags
-	void UpdateTagMap_Internal(const struct FGameplayTag& Tag, int32 CountDelta);
+	/** Map of tag to delegate that will be fired when the count for the key tag changes to or away from zero */
+	TMap<FGameplayTag, FOnGameplayEffectTagCountChanged> GameplayTagEventMap;
+
+	/** Map of tag to active count of that tag */
+	TMap<FGameplayTag, int32> GameplayTagCountMap;
+
+	/** Delegate fired whenever any tag's count changes to or away from zero */
+	FOnGameplayEffectTagCountChanged OnAnyTagChangeDelegate;
+
+	/** Container of tags that were explicitly added */
+	FGameplayTagContainer ExplicitTags;
+
+	/** Internal helper function to adjust the explicit tag list & corresponding maps/delegates/etc. as necessary */
+	void UpdateTagMap_Internal(const FGameplayTag& Tag, int32 CountDelta);
 };
 
 // -----------------------------------------------------------
