Abilities support a bool to disable client abilities canceling the server ones.

diff --git a/Engine/Source/Runtime/GameplayAbilities/Private/AbilitySystemComponent_Abilities.cpp b/Engine/Source/Runtime/GameplayAbilities/Private/AbilitySystemComponent_Abilities.cpp
--- a/Engine/Source/Runtime/GameplayAbilities/Private/AbilitySystemComponent_Abilities.cpp
+++ b/Engine/Source/Runtime/GameplayAbilities/Private/AbilitySystemComponent_Abilities.cpp
@@ -14,7 +14,7 @@
 #include "UObjectToken.h"
 #include "MapErrors.h"
 #define LOCTEXT_NAMESPACE "AbilitySystemComponent"
-
+#pragma optimize("",off)
 /** Enable to log out all render state create, destroy and updatetransform events */
 #define LOG_RENDER_STATE 0
 
@@ -583,7 +583,7 @@ bool UAbilitySystemComponent::TryActivateAbility(FGameplayAbilitySpecHandle Hand
 		if (Ability->GetInstancingPolicy() == EGameplayAbilityInstancingPolicy::InstancedPerExecution)
 		{
 			InstancedAbility = CreateNewInstanceOfAbility(*Spec, Ability);
-			ActivationInfo.bCanBeEndedByOtherInstance = true;
+			ActivationInfo.bCanBeEndedByOtherInstance = Ability->bServerRespectsRemoteAbilityCancelation;
 			InstancedAbility->CallActivateAbility(Handle, ActorInfo, ActivationInfo, OnGameplayAbilityEndedDelegate);
 			if (OutInstancedAbility)
 			{
diff --git a/Engine/Source/Runtime/GameplayAbilities/Public/Abilities/GameplayAbility.h b/Engine/Source/Runtime/GameplayAbilities/Public/Abilities/GameplayAbility.h
--- a/Engine/Source/Runtime/GameplayAbilities/Public/Abilities/GameplayAbility.h
+++ b/Engine/Source/Runtime/GameplayAbilities/Public/Abilities/GameplayAbility.h
@@ -521,6 +521,10 @@ protected:
 	UPROPERTY(EditDefaultsOnly, Category = Advanced)
 	TEnumAsByte<EGameplayAbilityInstancingPolicy::Type>	InstancingPolicy;						
 
+	/** If this is set, the server-side version of the ability can be canceled by the client-side version. The client-side version can always be canceled by the server. */
+	UPROPERTY(EditDefaultsOnly, Category = Advanced)
+	bool bServerRespectsRemoteAbilityCancelation;
+
 	/** This is information specific to this instance of the ability. E.g, whether it is predicting, authoring, confirmed, etc. */
 	UPROPERTY(BlueprintReadOnly, Category = Ability)
 	FGameplayAbilityActivationInfo	CurrentActivationInfo;
