Stray pragma cleanup, explicit default value added to a variable.

diff --git a/Engine/Source/Runtime/GameplayAbilities/Private/Abilities/GameplayAbility.cpp b/Engine/Source/Runtime/GameplayAbilities/Private/Abilities/GameplayAbility.cpp
--- a/Engine/Source/Runtime/GameplayAbilities/Private/Abilities/GameplayAbility.cpp
+++ b/Engine/Source/Runtime/GameplayAbilities/Private/Abilities/GameplayAbility.cpp
@@ -44,6 +44,8 @@ UGameplayAbility::UGameplayAbility(const FObjectInitializer& ObjectInitializer)
 		}
 	}
 #endif
+
+	bServerRespectsRemoteAbilityCancelation = false;
 }
 
 int32 UGameplayAbility::GetFunctionCallspace(UFunction* Function, void* Parameters, FFrame* Stack)
diff --git a/Engine/Source/Runtime/GameplayAbilities/Private/AbilitySystemComponent_Abilities.cpp b/Engine/Source/Runtime/GameplayAbilities/Private/AbilitySystemComponent_Abilities.cpp
--- a/Engine/Source/Runtime/GameplayAbilities/Private/AbilitySystemComponent_Abilities.cpp
+++ b/Engine/Source/Runtime/GameplayAbilities/Private/AbilitySystemComponent_Abilities.cpp
@@ -14,7 +14,7 @@
 #include "UObjectToken.h"
 #include "MapErrors.h"
 #define LOCTEXT_NAMESPACE "AbilitySystemComponent"
-#pragma optimize("",off)
+
 /** Enable to log out all render state create, destroy and updatetransform events */
 #define LOG_RENDER_STATE 0
 
