WaitGameplayEffectApplied

diff --git a/Engine/Source/Runtime/GameplayAbilities/Private/Abilities/Tasks/AbilityTask_WaitGameplayEffectApplied.cpp b/Engine/Source/Runtime/GameplayAbilities/Private/Abilities/Tasks/AbilityTask_WaitGameplayEffectApplied.cpp
new file mode 100644
--- /dev/null
+++ b/Engine/Source/Runtime/GameplayAbilities/Private/Abilities/Tasks/AbilityTask_WaitGameplayEffectApplied.cpp
@@ -0,0 +1,59 @@
+// Copyright 1998-2015 Epic Games, Inc. All Rights Reserved.
+
+#include "AbilitySystemPrivatePCH.h"
+#include "Abilities/Tasks/AbilityTask_WaitGameplayEffectApplied.h"
+
+#include "AbilitySystemComponent.h"
+
+#include "Abilities/GameplayAbility.h"
+#include "GameplayEffectExtension.h"
+
+UAbilityTask_WaitGameplayEffectApplied::UAbilityTask_WaitGameplayEffectApplied(const FObjectInitializer& ObjectInitializer)
+	: Super(ObjectInitializer)
+{
+}
+
+void UAbilityTask_WaitGameplayEffectApplied::Activate()
+{
+	if (AbilitySystemComponent.IsValid())
+	{
+		RegisterDelegate();
+	}
+}
+
+void UAbilityTask_WaitGameplayEffectApplied::OnApplyGameplayEffectCallback(UAbilitySystemComponent* Target, const FGameplayEffectSpec& SpecApplied, FActiveGameplayEffectHandle ActiveHandle)
+{
+	bool PassedComparison = false;
+
+	AActor* AvatarActor = Target ? Target->AvatarActor : nullptr;
+
+	if (!Filter.FilterPassesForActor(AvatarActor))
+	{
+		return;
+	}
+	if (!SourceTagRequirements.RequirementsMet(*SpecApplied.CapturedSourceTags.GetAggregatedTags()))
+	{
+		return;
+	}
+	if (!TargetTagRequirements.RequirementsMet(*SpecApplied.CapturedTargetTags.GetAggregatedTags()))
+	{
+		return;
+	}
+	
+	FGameplayEffectSpecHandle	SpecHandle(new FGameplayEffectSpec(SpecApplied));
+	BroadcastDelegate(AvatarActor, SpecHandle, ActiveHandle);
+	if (TriggerOnce)
+	{
+		EndTask();
+	}
+}
+
+void UAbilityTask_WaitGameplayEffectApplied::OnDestroy(bool AbilityEnded)
+{
+	if (AbilitySystemComponent.IsValid())
+	{
+		RemoveDelegate();
+	}
+
+	Super::OnDestroy(AbilityEnded);
+}
\ No newline at end of file
diff --git a/Engine/Source/Runtime/GameplayAbilities/Private/Abilities/Tasks/AbilityTask_WaitGameplayEffectApplied_Self.cpp b/Engine/Source/Runtime/GameplayAbilities/Private/Abilities/Tasks/AbilityTask_WaitGameplayEffectApplied_Self.cpp
new file mode 100644
--- /dev/null
+++ b/Engine/Source/Runtime/GameplayAbilities/Private/Abilities/Tasks/AbilityTask_WaitGameplayEffectApplied_Self.cpp
@@ -0,0 +1,40 @@
+// Copyright 1998-2015 Epic Games, Inc. All Rights Reserved.
+
+#include "AbilitySystemPrivatePCH.h"
+#include "Abilities/Tasks/AbilityTask_WaitGameplayEffectApplied_Self.h"
+
+#include "AbilitySystemComponent.h"
+
+#include "Abilities/GameplayAbility.h"
+#include "GameplayEffectExtension.h"
+
+UAbilityTask_WaitGameplayEffectApplied_Self::UAbilityTask_WaitGameplayEffectApplied_Self(const FObjectInitializer& ObjectInitializer)
+	: Super(ObjectInitializer)
+{
+}
+
+UAbilityTask_WaitGameplayEffectApplied_Self* UAbilityTask_WaitGameplayEffectApplied_Self::WaitGameplayEffectAppliedToSelf(UObject* WorldContextObject, const FGameplayTargetDataFilterHandle InFilter, FGameplayTagRequirements InSourceTagRequirements, FGameplayTagRequirements InTargetTagRequirements, bool InTriggerOnce)
+{
+	auto MyObj = NewTask<UAbilityTask_WaitGameplayEffectApplied_Self>(WorldContextObject);
+	MyObj->Filter = InFilter;
+	MyObj->SourceTagRequirements = InSourceTagRequirements;
+	MyObj->TargetTagRequirements = InTargetTagRequirements;
+	MyObj->TriggerOnce = InTriggerOnce;
+	return MyObj;
+}
+
+void UAbilityTask_WaitGameplayEffectApplied_Self::BroadcastDelegate(AActor* Avatar, FGameplayEffectSpecHandle SpecHandle, FActiveGameplayEffectHandle ActiveHandle)
+{
+	OnApplied.Broadcast(Avatar, SpecHandle, ActiveHandle);
+}
+
+void UAbilityTask_WaitGameplayEffectApplied_Self::RegisterDelegate()
+{
+	AbilitySystemComponent->OnGameplayEffectAppliedDelegateToSelf.AddUObject(this, &UAbilityTask_WaitGameplayEffectApplied::OnApplyGameplayEffectCallback);
+}
+
+void UAbilityTask_WaitGameplayEffectApplied_Self::RemoveDelegate()
+{
+	AbilitySystemComponent->OnGameplayEffectAppliedDelegateToSelf.RemoveUObject(this, &UAbilityTask_WaitGameplayEffectApplied::OnApplyGameplayEffectCallback);
+
+}
\ No newline at end of file
diff --git a/Engine/Source/Runtime/GameplayAbilities/Private/Abilities/Tasks/AbilityTask_WaitGameplayEffectApplied_Target.cpp b/Engine/Source/Runtime/GameplayAbilities/Private/Abilities/Tasks/AbilityTask_WaitGameplayEffectApplied_Target.cpp
new file mode 100644
--- /dev/null
+++ b/Engine/Source/Runtime/GameplayAbilities/Private/Abilities/Tasks/AbilityTask_WaitGameplayEffectApplied_Target.cpp
@@ -0,0 +1,40 @@
+// Copyright 1998-2015 Epic Games, Inc. All Rights Reserved.
+
+#include "AbilitySystemPrivatePCH.h"
+#include "Abilities/Tasks/AbilityTask_WaitGameplayEffectApplied_Target.h"
+
+#include "AbilitySystemComponent.h"
+
+#include "Abilities/GameplayAbility.h"
+#include "GameplayEffectExtension.h"
+
+UAbilityTask_WaitGameplayEffectApplied_Target::UAbilityTask_WaitGameplayEffectApplied_Target(const FObjectInitializer& ObjectInitializer)
+	: Super(ObjectInitializer)
+{
+}
+
+UAbilityTask_WaitGameplayEffectApplied_Target* UAbilityTask_WaitGameplayEffectApplied_Target::WaitGameplayEffectAppliedToTarget(UObject* WorldContextObject, const FGameplayTargetDataFilterHandle InFilter, FGameplayTagRequirements InSourceTagRequirements, FGameplayTagRequirements InTargetTagRequirements, bool InTriggerOnce)
+{
+	auto MyObj = NewTask<UAbilityTask_WaitGameplayEffectApplied_Target>(WorldContextObject);
+	MyObj->Filter = InFilter;
+	MyObj->SourceTagRequirements = InSourceTagRequirements;
+	MyObj->TargetTagRequirements = InTargetTagRequirements;
+	MyObj->TriggerOnce = InTriggerOnce;
+	return MyObj;
+}
+
+void UAbilityTask_WaitGameplayEffectApplied_Target::BroadcastDelegate(AActor* Avatar, FGameplayEffectSpecHandle SpecHandle, FActiveGameplayEffectHandle ActiveHandle)
+{
+	OnApplied.Broadcast(Avatar, SpecHandle, ActiveHandle);
+}
+
+void UAbilityTask_WaitGameplayEffectApplied_Target::RegisterDelegate()
+{
+	AbilitySystemComponent->OnGameplayEffectAppliedDelegateToTarget.AddUObject(this, &UAbilityTask_WaitGameplayEffectApplied::OnApplyGameplayEffectCallback);
+}
+
+void UAbilityTask_WaitGameplayEffectApplied_Target::RemoveDelegate()
+{
+	AbilitySystemComponent->OnGameplayEffectAppliedDelegateToTarget.RemoveUObject(this, &UAbilityTask_WaitGameplayEffectApplied::OnApplyGameplayEffectCallback);
+
+}
\ No newline at end of file
diff --git a/Engine/Source/Runtime/GameplayAbilities/Private/AbilitySystemComponent.cpp b/Engine/Source/Runtime/GameplayAbilities/Private/AbilitySystemComponent.cpp
--- a/Engine/Source/Runtime/GameplayAbilities/Private/AbilitySystemComponent.cpp
+++ b/Engine/Source/Runtime/GameplayAbilities/Private/AbilitySystemComponent.cpp
@@ -516,7 +516,7 @@ FActiveGameplayEffectHandle UAbilitySystemComponent::ApplyGameplayEffectSpecToSe
 	FActiveGameplayEffectHandle	MyHandle;
 	bool bInvokeGameplayCueApplied = UGameplayEffect::INSTANT_APPLICATION != Spec.GetDuration(); // Cache this now before possibly modifying predictive instant effect to infinite duration effect.
 
-	FGameplayEffectSpec* OurCopyOfSpec = NULL;
+	FGameplayEffectSpec* OurCopyOfSpec = nullptr;
 	TSharedPtr<FGameplayEffectSpec> StackSpec;
 	float Duration = bTreatAsInfiniteDuration ? UGameplayEffect::INFINITE_DURATION : Spec.GetDuration();
 	{
@@ -599,6 +599,18 @@ FActiveGameplayEffectHandle UAbilitySystemComponent::ApplyGameplayEffectSpecToSe
 		}
 	}
 
+	UAbilitySystemComponent* InstigatorASC = Spec.GetContext().GetInstigatorAbilitySystemComponent();
+
+	// Send ourselves a callback	
+	OnGameplayEffectAppliedToSelf(InstigatorASC, *OurCopyOfSpec, MyHandle);
+
+	// Send the instigator a callback
+	
+	if (InstigatorASC)
+	{
+		InstigatorASC->OnGameplayEffectAppliedToTarget(this, *OurCopyOfSpec, MyHandle);
+	}
+
 	return MyHandle;
 }
 
@@ -998,6 +1010,18 @@ void UAbilitySystemComponent::OnMagnitudeDependancyChange(FActiveGameplayEffectH
 	ActiveGameplayEffects.OnMagnitudeDependancyChange(Handle, ChangedAggregator);
 }
 
+void UAbilitySystemComponent::OnGameplayEffectAppliedToTarget(UAbilitySystemComponent* Target, const FGameplayEffectSpec& SpecApplied, FActiveGameplayEffectHandle ActiveHandle)
+{
+	OnGameplayEffectAppliedDelegateToTarget.Broadcast(Target, SpecApplied, ActiveHandle);
+}
+
+void UAbilitySystemComponent::OnGameplayEffectAppliedToSelf(UAbilitySystemComponent* Source, const FGameplayEffectSpec& SpecApplied, FActiveGameplayEffectHandle ActiveHandle)
+{
+	OnGameplayEffectAppliedDelegateToSelf.Broadcast(Source, SpecApplied, ActiveHandle);
+}
+
+// ------------------------------------------------------------------------
+
 FString ASC_CleanupName(FString Str)
 {
 	Str.RemoveFromStart(TEXT("Default__"));
diff --git a/Engine/Source/Runtime/GameplayAbilities/Public/Abilities/GameplayAbilityTargetDataFilter.h b/Engine/Source/Runtime/GameplayAbilities/Public/Abilities/GameplayAbilityTargetDataFilter.h
--- a/Engine/Source/Runtime/GameplayAbilities/Public/Abilities/GameplayAbilityTargetDataFilter.h
+++ b/Engine/Source/Runtime/GameplayAbilities/Public/Abilities/GameplayAbilityTargetDataFilter.h
@@ -84,7 +84,9 @@ struct GAMEPLAYABILITIES_API FGameplayTargetDataFilterHandle
 	{
 		if (!ActorToBeFiltered)
 		{
-			return false;
+			// If no filter is set, then always return true even if there is no actor.
+			// If there is no actor and there is a filter, then always fail.
+			return (Filter.IsValid() == false);
 		}
 		//Eventually, this might iterate through multiple filters. We'll need to decide how to designate OR versus AND functionality.
 		if (Filter.IsValid())
diff --git a/Engine/Source/Runtime/GameplayAbilities/Public/Abilities/Tasks/AbilityTask_WaitGameplayEffectApplied.h b/Engine/Source/Runtime/GameplayAbilities/Public/Abilities/Tasks/AbilityTask_WaitGameplayEffectApplied.h
new file mode 100644
--- /dev/null
+++ b/Engine/Source/Runtime/GameplayAbilities/Public/Abilities/Tasks/AbilityTask_WaitGameplayEffectApplied.h
@@ -0,0 +1,33 @@
+// Copyright 1998-2015 Epic Games, Inc. All Rights Reserved.
+#pragma once
+#include "AbilityTask.h"
+#include "Abilities/GameplayAbilityTypes.h"
+#include "AbilityTask_WaitGameplayEffectApplied.generated.h"
+
+
+
+UCLASS(MinimalAPI)
+class UAbilityTask_WaitGameplayEffectApplied : public UAbilityTask
+{
+	GENERATED_UCLASS_BODY()
+
+	UFUNCTION()
+	void OnApplyGameplayEffectCallback(UAbilitySystemComponent* Target, const FGameplayEffectSpec& SpecApplied, FActiveGameplayEffectHandle ActiveHandle);
+
+	virtual void Activate() override;
+
+	FGameplayTargetDataFilterHandle Filter;
+	FGameplayTagRequirements SourceTagRequirements;
+	FGameplayTagRequirements TargetTagRequirements;
+	bool TriggerOnce;
+
+protected:
+
+	virtual void BroadcastDelegate(AActor* Avatar, FGameplayEffectSpecHandle SpecHandle, FActiveGameplayEffectHandle ActiveHandle) { }
+	virtual void RegisterDelegate() { }
+	virtual void RemoveDelegate() { }
+
+	virtual void OnDestroy(bool AbilityEnded) override;
+
+	bool RegisteredCallback;
+};
\ No newline at end of file
diff --git a/Engine/Source/Runtime/GameplayAbilities/Public/Abilities/Tasks/AbilityTask_WaitGameplayEffectApplied_Self.h b/Engine/Source/Runtime/GameplayAbilities/Public/Abilities/Tasks/AbilityTask_WaitGameplayEffectApplied_Self.h
new file mode 100644
--- /dev/null
+++ b/Engine/Source/Runtime/GameplayAbilities/Public/Abilities/Tasks/AbilityTask_WaitGameplayEffectApplied_Self.h
@@ -0,0 +1,27 @@
+// Copyright 1998-2015 Epic Games, Inc. All Rights Reserved.
+#pragma once
+#include "AbilityTask.h"
+#include "Abilities/GameplayAbilityTypes.h"
+#include "AbilityTask_WaitGameplayEffectApplied.h"
+#include "AbilityTask_WaitGameplayEffectApplied_Self.generated.h"
+
+DECLARE_DYNAMIC_MULTICAST_DELEGATE_ThreeParams(FGameplayEffectAppliedSelfDelegate, AActor*, Source, FGameplayEffectSpecHandle, SpecHandle,  FActiveGameplayEffectHandle, ActiveHandle );
+
+UCLASS(MinimalAPI)
+class UAbilityTask_WaitGameplayEffectApplied_Self : public UAbilityTask_WaitGameplayEffectApplied
+{
+	GENERATED_UCLASS_BODY()
+
+	UPROPERTY(BlueprintAssignable)
+	FGameplayEffectAppliedSelfDelegate OnApplied;
+
+	/** Wait until the owner *receives* a GameplayEffect from a given source (the source may be the owner too!). If TriggerOnce is true, this task will only return one time. Otherwise it will return everytime a GE is applied that meets the requirements over the life of the ability */
+	UFUNCTION(BlueprintCallable, Category="Ability|Tasks", meta = (HidePin = "WorldContextObject", DefaultToSelf = "WorldContextObject", BlueprintInternalUseOnly = "TRUE"))
+	static UAbilityTask_WaitGameplayEffectApplied_Self* WaitGameplayEffectAppliedToSelf(UObject* WorldContextObject, const FGameplayTargetDataFilterHandle SourceFilter, FGameplayTagRequirements SourceTagRequirements, FGameplayTagRequirements TargetTagRequirements, bool TriggerOnce=false );
+	
+protected:
+
+	virtual void BroadcastDelegate(AActor* Avatar, FGameplayEffectSpecHandle SpecHandle, FActiveGameplayEffectHandle ActiveHandle) override;
+	virtual void RegisterDelegate() override;
+	virtual void RemoveDelegate() override;
+};
\ No newline at end of file
diff --git a/Engine/Source/Runtime/GameplayAbilities/Public/Abilities/Tasks/AbilityTask_WaitGameplayEffectApplied_Target.h b/Engine/Source/Runtime/GameplayAbilities/Public/Abilities/Tasks/AbilityTask_WaitGameplayEffectApplied_Target.h
new file mode 100644
--- /dev/null
+++ b/Engine/Source/Runtime/GameplayAbilities/Public/Abilities/Tasks/AbilityTask_WaitGameplayEffectApplied_Target.h
@@ -0,0 +1,27 @@
+// Copyright 1998-2015 Epic Games, Inc. All Rights Reserved.
+#pragma once
+#include "AbilityTask.h"
+#include "Abilities/GameplayAbilityTypes.h"
+#include "AbilityTask_WaitGameplayEffectApplied.h"
+#include "AbilityTask_WaitGameplayEffectApplied_Target.generated.h"
+
+DECLARE_DYNAMIC_MULTICAST_DELEGATE_ThreeParams(FGameplayEffectAppliedTargetDelegate, AActor*, Target, FGameplayEffectSpecHandle, SpecHandle,  FActiveGameplayEffectHandle, ActiveHandle );
+
+UCLASS(MinimalAPI)
+class UAbilityTask_WaitGameplayEffectApplied_Target : public UAbilityTask_WaitGameplayEffectApplied
+{
+	GENERATED_UCLASS_BODY()
+
+	UPROPERTY(BlueprintAssignable)
+	FGameplayEffectAppliedTargetDelegate OnApplied;
+
+	/** Wait until the owner applies a GameplayEffect to a Target (the target may be the owner too!). If TriggerOnce is true, this task will only return one time. Otherwise it will return everytime a GE is applied that meets the requirements over the life of the ability */
+	UFUNCTION(BlueprintCallable, Category="Ability|Tasks", meta = (HidePin = "WorldContextObject", DefaultToSelf = "WorldContextObject", BlueprintInternalUseOnly = "TRUE"))
+	static UAbilityTask_WaitGameplayEffectApplied_Target* WaitGameplayEffectAppliedToTarget(UObject* WorldContextObject, const FGameplayTargetDataFilterHandle TargetFilter, FGameplayTagRequirements SourceTagRequirements, FGameplayTagRequirements TargetTagRequirements, bool TriggerOnce=false );
+	
+protected:
+
+	virtual void BroadcastDelegate(AActor* Avatar, FGameplayEffectSpecHandle SpecHandle, FActiveGameplayEffectHandle ActiveHandle) override;
+	virtual void RegisterDelegate() override;
+	virtual void RemoveDelegate() override;
+};
\ No newline at end of file
diff --git a/Engine/Source/Runtime/GameplayAbilities/Public/AbilitySystemComponent.h b/Engine/Source/Runtime/GameplayAbilities/Public/AbilitySystemComponent.h
--- a/Engine/Source/Runtime/GameplayAbilities/Public/AbilitySystemComponent.h
+++ b/Engine/Source/Runtime/GameplayAbilities/Public/AbilitySystemComponent.h
@@ -271,13 +271,27 @@ class GAMEPLAYABILITIES_API UAbilitySystemComponent : public UActorComponent, pu
 	 * @param OutCaptureSpec	[OUT] Capture spec to populate with captured data
 	 */
 	void CaptureAttributeForGameplayEffect(OUT FGameplayEffectAttributeCaptureSpec& OutCaptureSpec);
-
-	// Delegates (these need to be at the UObject level so we can safetly bind, rather than binding to raw at the ActiveGameplayEffect/Container level which is unsafe if the AbilitySystemComponent were killed).
+	
+	// --------------------------------------------
+	// Callbacks / Notifies
+	// (these need to be at the UObject level so we can safetly bind, rather than binding to raw at the ActiveGameplayEffect/Container level which is unsafe if the AbilitySystemComponent were killed).
+	// --------------------------------------------
 
 	void OnAttributeAggregatorDirty(FAggregator* Aggregator, FGameplayAttribute Attribute);
 
 	void OnMagnitudeDependancyChange(FActiveGameplayEffectHandle Handle, const FAggregator* ChangedAggregator);
 
+	/** This ASC has successfully applied a GE to something (potentially itself) */
+	void OnGameplayEffectAppliedToTarget(UAbilitySystemComponent* Target, const FGameplayEffectSpec& SpecApplied, FActiveGameplayEffectHandle ActiveHandle);
+
+	void OnGameplayEffectAppliedToSelf(UAbilitySystemComponent* Source, const FGameplayEffectSpec& SpecApplied, FActiveGameplayEffectHandle ActiveHandle);
+
+	DECLARE_MULTICAST_DELEGATE_ThreeParams(FOnGameplayEffectAppliedDelegate, UAbilitySystemComponent*, const FGameplayEffectSpec&, FActiveGameplayEffectHandle);
+
+	FOnGameplayEffectAppliedDelegate OnGameplayEffectAppliedDelegateToSelf;
+
+	FOnGameplayEffectAppliedDelegate OnGameplayEffectAppliedDelegateToTarget;
+
 	// --------------------------------------------
 	// Tags
 	// --------------------------------------------
@@ -331,7 +345,7 @@ class GAMEPLAYABILITIES_API UAbilitySystemComponent : public UActorComponent, pu
 
 
 	// --------------------------------------------
-	// Possibly useful but not primary API functions:
+	// Additional Helper Functions
 	// --------------------------------------------
 	
 	FOnActiveGameplayEffectRemoved* OnGameplayEffectRemovedDelegate(FActiveGameplayEffectHandle Handle);
diff --git a/Engine/Source/Runtime/GameplayAbilities/Public/GameplayEffectTypes.h b/Engine/Source/Runtime/GameplayAbilities/Public/GameplayEffectTypes.h
--- a/Engine/Source/Runtime/GameplayAbilities/Public/GameplayEffectTypes.h
+++ b/Engine/Source/Runtime/GameplayAbilities/Public/GameplayEffectTypes.h
@@ -791,11 +791,11 @@ struct FGameplayTagRequirements
 	GENERATED_USTRUCT_BODY()
 
 	/** All of these tags must be present */
-	UPROPERTY(EditDefaultsOnly, Category = GameplayModifier)
+	UPROPERTY(EditAnywhere, Category = GameplayModifier)
 	FGameplayTagContainer RequireTags;
 
 	/** None of these tags may be present */
-	UPROPERTY(EditDefaultsOnly, Category = GameplayModifier)
+	UPROPERTY(EditAnywhere, Category = GameplayModifier)
 	FGameplayTagContainer IgnoreTags;
 
 	bool	RequirementsMet(const FGameplayTagContainer& Container) const;
diff --git a/Engine/Source/Runtime/GameplayTags/Classes/GameplayTagContainer.h b/Engine/Source/Runtime/GameplayTags/Classes/GameplayTagContainer.h
--- a/Engine/Source/Runtime/GameplayTags/Classes/GameplayTagContainer.h
+++ b/Engine/Source/Runtime/GameplayTags/Classes/GameplayTagContainer.h
@@ -256,6 +256,9 @@ protected:
 	*/
 	bool DoesTagContainerMatch(const FGameplayTagContainer& OtherContainer, TEnumAsByte<EGameplayTagMatchType::Type> TagMatchType, TEnumAsByte<EGameplayTagMatchType::Type> OtherTagMatchType, EGameplayContainerMatchType ContainerMatchType) const;
 
+	/** Array of gameplay tags */
+	UPROPERTY(BlueprintReadWrite, Category=GameplayTags)
+	TArray<FGameplayTag> GameplayTags;
 
 private:
 
@@ -263,10 +266,6 @@ private:
 	UPROPERTY()
 	TArray<FName> Tags_DEPRECATED;
 
-	/** Array of gameplay tags */
-	UPROPERTY(VisibleAnywhere, Category=GameplayTags)
-	TArray<FGameplayTag> GameplayTags;
-
 	/**
 	 * DO NOT USE DIRECTLY
 	 * STL-like iterators to enable range-based for loop support.
