count empty query tag containers as matching

diff --git a/Engine/Source/Runtime/GameplayAbilities/Private/GameplayEffect.cpp b/Engine/Source/Runtime/GameplayAbilities/Private/GameplayEffect.cpp
--- a/Engine/Source/Runtime/GameplayAbilities/Private/GameplayEffect.cpp
+++ b/Engine/Source/Runtime/GameplayAbilities/Private/GameplayEffect.cpp
@@ -2242,7 +2242,7 @@ bool FActiveGameplayEffectQuery::Matches(const FActiveGameplayEffect& Effect) co
 	// if we are looking for owning tags check them on the Granted Tags and Owned Tags Container
 	if (OwningTagContainer)
 	{
-		if (!Effect.Spec.Def->InheritableOwnedTagsContainer.CombinedTags.MatchesAny(*OwningTagContainer, false) &&
+		if (!Effect.Spec.Def->InheritableOwnedTagsContainer.CombinedTags.MatchesAny(*OwningTagContainer, true) &&
 			!Effect.Spec.DynamicGrantedTags.MatchesAny(*OwningTagContainer, false))
 		{
 			// this doesn't match our Tags so bail
@@ -2253,7 +2253,7 @@ bool FActiveGameplayEffectQuery::Matches(const FActiveGameplayEffect& Effect) co
 	// if we are just looking for Tags on the Effect then look at the Gameplay Effect Tags
 	if (EffectTagContainer)
 	{
-		if (!Effect.Spec.Def->InheritableGameplayEffectTags.CombinedTags.MatchesAny(*EffectTagContainer, false))
+		if (!Effect.Spec.Def->InheritableGameplayEffectTags.CombinedTags.MatchesAny(*EffectTagContainer, true))
 		{
 			// this doesn't match our Tags so bail
 			return false;
@@ -2263,7 +2263,7 @@ bool FActiveGameplayEffectQuery::Matches(const FActiveGameplayEffect& Effect) co
 	// if we are just looking for Tags on the Effect then look at the Gameplay Effect Tags
 	if (EffectTagContainer_Rejection)
 	{
-		if (Effect.Spec.Def->InheritableGameplayEffectTags.CombinedTags.MatchesAny(*EffectTagContainer_Rejection, false))
+		if (Effect.Spec.Def->InheritableGameplayEffectTags.CombinedTags.MatchesAny(*EffectTagContainer_Rejection, true))
 		{
 			// this matches our Rejection Tags so bail
 			return false;
@@ -2360,4 +2360,4 @@ void FActiveGameplayEffectAction::PerformAction()
 			ASC->ActiveGameplayEffects.GameplayEffects.Add(TempEffect);
 		}
 	}
-}
\ No newline at end of file
+}
