Removed an unused function.

diff --git a/Engine/Source/Runtime/GameplayAbilities/Private/AbilitySystemComponent.cpp b/Engine/Source/Runtime/GameplayAbilities/Private/AbilitySystemComponent.cpp
--- a/Engine/Source/Runtime/GameplayAbilities/Private/AbilitySystemComponent.cpp
+++ b/Engine/Source/Runtime/GameplayAbilities/Private/AbilitySystemComponent.cpp
@@ -840,11 +840,6 @@ void UAbilitySystemComponent::RemoveActiveEffectsWithTags(const FGameplayTagCont
 	RemoveActiveEffects(FActiveGameplayEffectQuery(&Tags));
 }
 
-void UAbilitySystemComponent::RemoveActiveEffects(FGameplayTagContainer OwningTags, FGameplayTagContainer EffectTags, FGameplayTagContainer OwningTags_Rejection, FGameplayTagContainer EffectTags_Rejection)
-{
-	return ActiveGameplayEffects.RemoveActiveEffects(FActiveGameplayEffectQuery(&OwningTags, &EffectTags, &OwningTags_Rejection, &EffectTags_Rejection));
-}
-
 void UAbilitySystemComponent::RemoveActiveEffects(const FActiveGameplayEffectQuery Query)
 {
 	return ActiveGameplayEffects.RemoveActiveEffects(Query);
diff --git a/Engine/Source/Runtime/GameplayAbilities/Public/AbilitySystemComponent.h b/Engine/Source/Runtime/GameplayAbilities/Public/AbilitySystemComponent.h
--- a/Engine/Source/Runtime/GameplayAbilities/Public/AbilitySystemComponent.h
+++ b/Engine/Source/Runtime/GameplayAbilities/Public/AbilitySystemComponent.h
@@ -386,10 +386,6 @@ class GAMEPLAYABILITIES_API UAbilitySystemComponent : public UActorComponent, pu
 	UFUNCTION(BlueprintCallable, Category = GameplayEffects)
 	void RemoveActiveEffectsWithTags(FGameplayTagContainer Tags);
 
-	/** Constructs a query and removess active effects as appropriate */
-	UFUNCTION(BlueprintCallable, Category = GameplayEffects)
-	void RemoveActiveEffects(FGameplayTagContainer OwningTags, FGameplayTagContainer EffectTags, FGameplayTagContainer OwningTags_Rejection, FGameplayTagContainer EffectTags_Rejection);
-
 	/** Removes all active effects that match given query */
 	void RemoveActiveEffects(const FActiveGameplayEffectQuery Query);
 
