Added a function to AbilitySystemComponent to remove all GameplayCues that were added individually, i.e. not as part of GameplayEffects.

diff --git a/Engine/Source/Runtime/GameplayAbilities/Private/AbilitySystemComponent.cpp b/Engine/Source/Runtime/GameplayAbilities/Private/AbilitySystemComponent.cpp
--- a/Engine/Source/Runtime/GameplayAbilities/Private/AbilitySystemComponent.cpp
+++ b/Engine/Source/Runtime/GameplayAbilities/Private/AbilitySystemComponent.cpp
@@ -776,6 +776,14 @@ void UAbilitySystemComponent::RemoveGameplayCue(const FGameplayTag GameplayCueTa
 	}
 }
 
+void UAbilitySystemComponent::RemoveAllGameplayCues()
+{
+	for (int32 i = (ActiveGameplayCues.GameplayCues.Num() - 1); i >= 0; --i)
+	{
+		RemoveGameplayCue(ActiveGameplayCues.GameplayCues[i].GameplayCueTag);
+	}
+}
+
 void UAbilitySystemComponent::NetMulticast_InvokeGameplayCueExecuted_FromSpec_Implementation(const FGameplayEffectSpec Spec, FPredictionKey PredictionKey)
 {
 	if (IsOwnerActorAuthoritative() || PredictionKey.IsValidKey() == false)
diff --git a/Engine/Source/Runtime/GameplayAbilities/Public/AbilitySystemComponent.h b/Engine/Source/Runtime/GameplayAbilities/Public/AbilitySystemComponent.h
--- a/Engine/Source/Runtime/GameplayAbilities/Public/AbilitySystemComponent.h
+++ b/Engine/Source/Runtime/GameplayAbilities/Public/AbilitySystemComponent.h
@@ -415,6 +415,9 @@ class GAMEPLAYABILITIES_API UAbilitySystemComponent : public UActorComponent, pu
 	
 	void RemoveGameplayCue(const FGameplayTag GameplayCueTag);
 
+	/** Removes any GameplayCue added on its own, i.e. not as part of a GameplayEffect. */
+	void RemoveAllGameplayCues();
+
 	UFUNCTION(NetMulticast, unreliable)
 	void NetMulticast_InvokeGameplayCueExecuted(const FGameplayTag GameplayCueTag, FPredictionKey PredictionKey, FGameplayEffectContextHandle EffectContext);
 
