adding in support for ability key T

diff --git a/Engine/Source/Runtime/GameplayAbilities/Public/Abilities/GameplayAbilityTypes.h b/Engine/Source/Runtime/GameplayAbilities/Public/Abilities/GameplayAbilityTypes.h
--- a/Engine/Source/Runtime/GameplayAbilities/Public/Abilities/GameplayAbilityTypes.h
+++ b/Engine/Source/Runtime/GameplayAbilities/Public/Abilities/GameplayAbilityTypes.h
@@ -441,7 +441,7 @@ struct FGameplayAbiliyInputBinds
 	/** Defines command string that will be bound to Cancel Targeting */
 	FString CancelTargetCommand;
 
-	/** Returns enum to use for ability bings. E.g., "Ability1"-"Ability8" input commands will be bound to ability activations inside the AbiltiySystemComponent */
+	/** Returns enum to use for ability binds. E.g., "Ability1"-"Ability9" input commands will be bound to ability activations inside the AbiltiySystemComponent */
 	FString	EnumName;
 
 	UEnum* GetBindEnum() { return FindObject<UEnum>(ANY_PACKAGE, *EnumName); }
diff --git a/Engine/Source/Runtime/GameplayAbilities/Public/GameplayAbilitySet.h b/Engine/Source/Runtime/GameplayAbilities/Public/GameplayAbilitySet.h
--- a/Engine/Source/Runtime/GameplayAbilities/Public/GameplayAbilitySet.h
+++ b/Engine/Source/Runtime/GameplayAbilities/Public/GameplayAbilitySet.h
@@ -27,7 +27,8 @@ namespace EGameplayAbilityInputBinds
 		Ability5				UMETA(DisplayName = "Ability5 (R)"),
 		Ability6				UMETA(DisplayName = "Ability6 (Shift)"),
 		Ability7				UMETA(DisplayName = "Ability7 (Space)"),
-		Ability8				UMETA(DisplayName = "Ability7 (B)"),
+		Ability8				UMETA(DisplayName = "Ability8 (B)"),
+		Ability9				UMETA(DisplayName = "Ability9 (T)"),
 	};
 }
 
