only call removed on the RemoveGameplayEffectsWithTags.CombinedTags if there are tags in the container.

diff --git a/Engine/Source/Runtime/GameplayAbilities/Private/AbilitySystemComponent.cpp b/Engine/Source/Runtime/GameplayAbilities/Private/AbilitySystemComponent.cpp
--- a/Engine/Source/Runtime/GameplayAbilities/Private/AbilitySystemComponent.cpp
+++ b/Engine/Source/Runtime/GameplayAbilities/Private/AbilitySystemComponent.cpp
@@ -584,8 +584,9 @@ FActiveGameplayEffectHandle UAbilitySystemComponent::ApplyGameplayEffectSpecToSe
 	// ------------------------------------------------------
 	//	Remove gameplay effects with tags
 	//		Remove any active gameplay effects that match the RemoveGameplayEffectsWithTags in the definition for this spec
+	//		Only call this if we are the Authoritative owner and we have some RemoveGameplayEffectsWithTags.CombinedTag to remove
 	// ------------------------------------------------------
-	if (IsOwnerActorAuthoritative())
+	if (IsOwnerActorAuthoritative() && Spec.Def->RemoveGameplayEffectsWithTags.CombinedTags.Num() > 0)
 	{
 		ActiveGameplayEffects.RemoveActiveEffects(FActiveGameplayEffectQuery(&Spec.Def->RemoveGameplayEffectsWithTags.CombinedTags));
 	}
