remove old GE stacking code in preperation for refactor

diff --git a/Engine/Source/Editor/GameplayAbilitiesEditor/Private/AssetTypeActions_GameplayEffect.cpp b/Engine/Source/Editor/GameplayAbilitiesEditor/Private/AssetTypeActions_GameplayEffect.cpp
--- a/Engine/Source/Editor/GameplayAbilitiesEditor/Private/AssetTypeActions_GameplayEffect.cpp
+++ b/Engine/Source/Editor/GameplayAbilitiesEditor/Private/AssetTypeActions_GameplayEffect.cpp
@@ -48,10 +48,6 @@ void CopyInfoFromDataGE(const UGameplayEffect& InGE, UGameplayEffect& OutGE)
 
 	OutGE.Description = InGE.Description;
 
-	OutGE.StackingPolicy = InGE.StackingPolicy;
-	OutGE.StackedAttribName = InGE.StackedAttribName;
-	OutGE.StackingExtension = InGE.StackingExtension;
-
 	OutGE.UpdateInheritedTagProperties();
 }
 
diff --git a/Engine/Source/Runtime/GameplayAbilities/Private/AbilitySystemComponent.cpp b/Engine/Source/Runtime/GameplayAbilities/Private/AbilitySystemComponent.cpp
--- a/Engine/Source/Runtime/GameplayAbilities/Private/AbilitySystemComponent.cpp
+++ b/Engine/Source/Runtime/GameplayAbilities/Private/AbilitySystemComponent.cpp
@@ -312,11 +312,6 @@ int32 UAbilitySystemComponent::GetNumActiveGameplayEffect() const
 	return ActiveGameplayEffects.GetNumGameplayEffects();
 }
 
-bool UAbilitySystemComponent::IsGameplayEffectActive(FActiveGameplayEffectHandle InHandle) const
-{
-	return ActiveGameplayEffects.IsGameplayEffectActive(InHandle);
-}
-
 const FGameplayTagContainer* UAbilitySystemComponent::GetGameplayEffectSourceTagsFromHandle(FActiveGameplayEffectHandle Handle) const
 {
 	return ActiveGameplayEffects.GetGameplayEffectSourceTagsFromHandle(Handle);
@@ -522,12 +517,6 @@ FActiveGameplayEffectHandle UAbilitySystemComponent::ApplyGameplayEffectSpecToSe
 	{
 		if (Duration != UGameplayEffect::INSTANT_APPLICATION)
 		{
-			// recalculating stacking needs to come before creating the new effect
-			if (Spec.GetStackingType() != EGameplayEffectStackingPolicy::Unlimited)
-			{
-				ActiveGameplayEffects.StacksNeedToRecalculate();
-			}
-
 			FActiveGameplayEffect& NewActiveEffect = ActiveGameplayEffects.CreateNewActiveGameplayEffect(Spec, PredictionKey);
 			MyHandle = NewActiveEffect.Handle;
 			OurCopyOfSpec = &NewActiveEffect.Spec;
@@ -854,11 +843,6 @@ void UAbilitySystemComponent::RemoveActiveEffects(const FActiveGameplayEffectQue
 	return ActiveGameplayEffects.RemoveActiveEffects(Query);
 }
 
-void UAbilitySystemComponent::OnRestackGameplayEffects()
-{
-	ActiveGameplayEffects.RecalculateStacking();
-}
-
 // ---------------------------------------------------------------------------------------
 
 void UAbilitySystemComponent::TaskStarted(UAbilityTask* NewTask)
diff --git a/Engine/Source/Runtime/GameplayAbilities/Private/GameplayEffect.cpp b/Engine/Source/Runtime/GameplayAbilities/Private/GameplayEffect.cpp
--- a/Engine/Source/Runtime/GameplayAbilities/Private/GameplayEffect.cpp
+++ b/Engine/Source/Runtime/GameplayAbilities/Private/GameplayEffect.cpp
@@ -4,7 +4,6 @@
 #include "GameplayEffect.h"
 #include "GameplayEffectExtension.h"
 #include "GameplayEffectTypes.h"
-#include "GameplayEffectStackingExtension.h"
 #include "GameplayTagsModule.h"
 #include "AbilitySystemComponent.h"
 #include "AbilitySystemInterface.h"
@@ -29,8 +28,6 @@ UGameplayEffect::UGameplayEffect(const FObjectInitializer& ObjectInitializer)
 	DurationPolicy = EGameplayEffectDurationType::Instant;
 	ChanceToApplyToTarget.SetValue(1.f);
 	ChanceToExecuteOnGameplayEffect.SetValue(1.f);
-	StackingPolicy = EGameplayEffectStackingPolicy::Unlimited;
-	StackedAttribName = NAME_None;
 
 #if WITH_EDITORONLY_DATA
 	ShowAllProperties = true;
@@ -632,11 +629,6 @@ void FGameplayEffectSpec::PruneModifiedAttributes()
 	}
 }
 
-EGameplayEffectStackingPolicy::Type FGameplayEffectSpec::GetStackingType() const
-{
-	return Def->StackingPolicy;
-}
-
 void FGameplayEffectSpec::SetContext(FGameplayEffectContextHandle NewEffectContext)
 {
 	bool bWasAlreadyInit = EffectContext.IsValid();
@@ -919,11 +911,6 @@ void FActiveGameplayEffect::CheckOngoingTagRequirements(const FGameplayTagContai
 	}
 }
 
-bool FActiveGameplayEffect::CanBeStacked(const FActiveGameplayEffect& Other) const
-{
-	return (Handle != Other.Handle) && Spec.GetStackingType() == Other.Spec.GetStackingType() && Spec.Def->Modifiers[0].Attribute == Other.Spec.Def->Modifiers[0].Attribute;
-}
-
 void FActiveGameplayEffect::PreReplicatedRemove(const struct FActiveGameplayEffectsContainer &InArray)
 {
 	if (Spec.Def == nullptr)
@@ -1015,26 +1002,8 @@ void FActiveGameplayEffectsContainer::ExecuteActiveEffectsFrom(FGameplayEffectSp
 {
 	// If there are no modifiers, we always want to apply the GameplayCue. If there are modifiers, we only want to invoke the GameplayCue if one of them went through (could be blocked by immunity or % chance roll)
 	bool InvokeGameplayCueExecute = (Spec.Modifiers.Num() == 0);
-
-	// check if this is a stacking effect and if it is the active stacking effect.
-	if (Spec.GetStackingType() != EGameplayEffectStackingPolicy::Unlimited)
-	{
-		// we're a stacking attribute but not active
-		if (Spec.bTopOfStack == false)
-		{
-			return;
-		}
-	}
-
-	// The process of executing a gameplay spec may cause modifications to the spec, however those modifications
-	// should not persist across multiple executions on the same spec. Therefore, if the spec isn't instantaneous, make a 
-	// copy that can be modified and pass it along to each part of the execution process.
-	TSharedPtr<FGameplayEffectSpec> StackSpec;
-	if (Spec.GetDuration() != UGameplayEffect::INSTANT_APPLICATION)
-	{
-		StackSpec = TSharedPtr<FGameplayEffectSpec>(new FGameplayEffectSpec(Spec));
-	}
-	FGameplayEffectSpec& SpecToUse = StackSpec.IsValid() ? *(StackSpec.Get()) : Spec;
+	
+	FGameplayEffectSpec& SpecToUse = Spec;
 
 	float ChanceToExecute = SpecToUse.GetChanceToExecuteOnGameplayEffect();		// Not implemented? Should we just remove ChanceToExecute?
 
@@ -1336,26 +1305,6 @@ float FActiveGameplayEffectsContainer::GetGameplayEffectMagnitude(FActiveGamepla
 	return -1.f;
 }
 
-bool FActiveGameplayEffectsContainer::IsGameplayEffectActive(FActiveGameplayEffectHandle Handle, bool IncludeEffectsBlockedByStackingRules) const
-{
-	// Could make this a map for quicker lookup
-	for (const FActiveGameplayEffect& Effect : GameplayEffects)
-	{
-		if (Effect.Handle == Handle)
-		{
-			// stacking effects may return false
-			if (Effect.Spec.GetStackingType() != EGameplayEffectStackingPolicy::Unlimited &&
-				Effect.Spec.bTopOfStack == false)
-			{
-				return IncludeEffectsBlockedByStackingRules;
-			}
-			return true;
-		}
-	}
-
-	return false;
-}
-
 const FGameplayTagContainer* FActiveGameplayEffectsContainer::GetGameplayEffectSourceTagsFromHandle(FActiveGameplayEffectHandle Handle) const
 {
 	// @todo: Need to consider this with tag changes
@@ -1487,18 +1436,6 @@ void FActiveGameplayEffectsContainer::ApplyModToAttribute(const FGameplayAttribu
 	}
 }
 
-void FActiveGameplayEffectsContainer::StacksNeedToRecalculate()
-{
-	if (!bNeedToRecalculateStacks)
-	{
-		bNeedToRecalculateStacks = true;
-		FTimerManager& TimerManager = Owner->GetWorld()->GetTimerManager();
-		FTimerDelegate Delegate = FTimerDelegate::CreateUObject(Owner, &UAbilitySystemComponent::OnRestackGameplayEffects);
-
-		TimerManager.SetTimerForNextTick(Delegate);
-	}
-}
-
 /**	This Creates a new ActiveGameplayEffect and is only called on the Server or in cases where a client predictively applies a GameplayEffect. */
 FActiveGameplayEffect& FActiveGameplayEffectsContainer::CreateNewActiveGameplayEffect(const FGameplayEffectSpec &Spec, FPredictionKey InPredictionKey)
 {
@@ -1575,29 +1512,9 @@ FActiveGameplayEffect& FActiveGameplayEffectsContainer::CreateNewActiveGameplayE
 		{
 			FTimerManager& TimerManager = Owner->GetWorld()->GetTimerManager();
 			FTimerDelegate Delegate = FTimerDelegate::CreateUObject(Owner, &UAbilitySystemComponent::ExecutePeriodicEffect, NewHandle);
-
-			// If this is a periodic stacking effect, make sure that it's in sync with the others.
-			float FirstDelay = -1.f;
-			bool bApplyToNextTick = true;
-			if (NewEffect.Spec.GetStackingType() != EGameplayEffectStackingPolicy::Unlimited)
-			{
-				for (FActiveGameplayEffect& Effect : GameplayEffects)
-				{
-					if (Effect.CanBeStacked(NewEffect))
-					{
-						FirstDelay = TimerManager.GetTimerRemaining(Effect.PeriodHandle);
-						bApplyToNextTick = TimerManager.IsTimerPending(Effect.PeriodHandle);
-						break;
-					}
-				}
-			}
 			
-			if (bApplyToNextTick)
-			{
-				TimerManager.SetTimerForNextTick(Delegate); // not part of a stack or the first item on the stack, execute once right away
-			}
-
-			TimerManager.SetTimer(NewEffect.PeriodHandle, Delegate, NewEffect.Spec.GetPeriod(), true, FirstDelay); // this is going to be off by a frame for stacking because of the pending list
+			TimerManager.SetTimerForNextTick(Delegate);
+			TimerManager.SetTimer(NewEffect.PeriodHandle, Delegate, NewEffect.Spec.GetPeriod(), true, -1.f);
 		}
 	}
 	
@@ -1749,8 +1666,7 @@ bool FActiveGameplayEffectsContainer::InternalRemoveActiveGameplayEffect(int32 I
 		// during debugging where breakpoints or pausing can mess up network update times. Open issue
 		// with network team.
 		Owner->GetOwner()->ForceNetUpdate();
-
-		StacksNeedToRecalculate();
+		
 		return true;
 	}
 
@@ -1951,112 +1867,6 @@ void FActiveGameplayEffectsContainer::CheckDuration(FActiveGameplayEffectHandle
 	}
 }
 
-void FActiveGameplayEffectsContainer::RecalculateStacking()
-{
-	bNeedToRecalculateStacks = false;
-
-	// one or more gameplay effects has been removed or added so we need to go through all of them to see what the best elements are in each stack
-	TArray<FActiveGameplayEffect*> StackedEffects;
-	TArray<FActiveGameplayEffect*> CustomStackedEffects;
-
-	for (FActiveGameplayEffect& Effect : GameplayEffects)
-	{
-		// ignore effects that don't stack and effects that replace when they stack
-		// effects that replace when they stack only need to be handled when they are added
-		if (Effect.Spec.GetStackingType() != EGameplayEffectStackingPolicy::Unlimited &&
-			Effect.Spec.GetStackingType() != EGameplayEffectStackingPolicy::Replaces)
-		{
-			if (Effect.Spec.Def->Modifiers.Num() > 0)
-			{
-				Effect.Spec.bTopOfStack = false;
-				bool bFoundStack = false;
-
-				// group all of the custom stacking effects and deal with them after
-				if (Effect.Spec.GetStackingType() == EGameplayEffectStackingPolicy::Callback)
-				{
-					CustomStackedEffects.Add(&Effect);
-					continue;
-				}
-
-				for (FActiveGameplayEffect*& StackedEffect : StackedEffects)
-				{
-					if (StackedEffect->Spec.GetStackingType() == Effect.Spec.GetStackingType() && StackedEffect->Spec.Def->Modifiers[0].Attribute == Effect.Spec.Def->Modifiers[0].Attribute)
-					{
-						switch (Effect.Spec.GetStackingType())
-						{
-						case EGameplayEffectStackingPolicy::Highest:
-						{
-							float BestSpecMagnitude = StackedEffect->Spec.GetMagnitude(StackedEffect->Spec.Def->Modifiers[0].Attribute);
-							float CurrSpecMagnitude = Effect.Spec.GetMagnitude(Effect.Spec.Def->Modifiers[0].Attribute);
-							if (BestSpecMagnitude < CurrSpecMagnitude)
-							{
-								StackedEffect = &Effect;
-							}
-							break;
-						}
-						case EGameplayEffectStackingPolicy::Lowest:
-						{
-							float BestSpecMagnitude = StackedEffect->Spec.GetMagnitude(StackedEffect->Spec.Def->Modifiers[0].Attribute);
-							float CurrSpecMagnitude = Effect.Spec.GetMagnitude(Effect.Spec.Def->Modifiers[0].Attribute);
-							if (BestSpecMagnitude > CurrSpecMagnitude)
-							{
-								StackedEffect = &Effect;
-							}
-							break;
-						}
-						default:
-							// we shouldn't ever be here
-							ABILITY_LOG(Warning, TEXT("%s uses unhandled stacking rule %s."), *Effect.Spec.Def->GetPathName(), *EGameplayEffectStackingPolicyToString(Effect.Spec.GetStackingType()));
-							break;
-						};
-
-						bFoundStack = true;
-						break;
-					}
-				}
-
-				// if we didn't find a stack matching this effect we need to add it to our list
-				if (!bFoundStack)
-				{
-					StackedEffects.Add(&Effect);
-				}
-			}
-		}
-	}
-
-	// now deal with the custom effects
-	while (CustomStackedEffects.Num() > 0)
-	{
-		int32 StackedIdx = StackedEffects.Add(CustomStackedEffects[0]);
-		CustomStackedEffects.RemoveAtSwap(0);
-
-		UGameplayEffectStackingExtension * StackedExt = StackedEffects[StackedIdx]->Spec.Def->StackingExtension->GetDefaultObject<UGameplayEffectStackingExtension>();
-
-		TArray<FActiveGameplayEffect*> Effects;
-		
-		for (int32 Idx = 0; Idx < CustomStackedEffects.Num(); Idx++)
-		{
-			UGameplayEffectStackingExtension * CurrentExt = CustomStackedEffects[Idx]->Spec.Def->StackingExtension->GetDefaultObject<UGameplayEffectStackingExtension>();
-			if (CurrentExt && StackedExt && CurrentExt->Handle == StackedExt->Handle && 
-				StackedEffects[StackedIdx]->Spec.Def->Modifiers[0].Attribute == CustomStackedEffects[Idx]->Spec.Def->Modifiers[0].Attribute)
-			{
-				Effects.Add(CustomStackedEffects[Idx]);
-				CustomStackedEffects.RemoveAtSwap(Idx);
-				--Idx;
-			}
-		}
-
-		UGameplayEffectStackingExtension * Ext = StackedEffects[StackedIdx]->Spec.Def->StackingExtension->GetDefaultObject<UGameplayEffectStackingExtension>();
-		Ext->CalculateStack(Effects, *this, *StackedEffects[StackedIdx]);
-	}
-
-	// we've found all of the best stacking effects, mark them so that they updated properly
-	for (FActiveGameplayEffect* const StackedEffect : StackedEffects)
-	{
-		StackedEffect->Spec.bTopOfStack = true;
-	}
-}
-
 bool FActiveGameplayEffectsContainer::CanApplyAttributeModifiers(const UGameplayEffect* GameplayEffect, float Level, const FGameplayEffectContextHandle& EffectContext)
 {
 	SCOPE_CYCLE_COUNTER(STAT_GameplayEffectsCanApplyAttributeModifiers);
diff --git a/Engine/Source/Runtime/GameplayAbilities/Private/GameplayEffectStackingExtension.cpp b/Engine/Source/Runtime/GameplayAbilities/Private/GameplayEffectStackingExtension.cpp
deleted file mode 100644
diff --git a/Engine/Source/Runtime/GameplayAbilities/Private/GameplayEffectStackingExtension_CappedNumberTest.cpp b/Engine/Source/Runtime/GameplayAbilities/Private/GameplayEffectStackingExtension_CappedNumberTest.cpp
deleted file mode 100644
diff --git a/Engine/Source/Runtime/GameplayAbilities/Private/GameplayEffectStackingExtension_DiminishingReturnsTest.cpp b/Engine/Source/Runtime/GameplayAbilities/Private/GameplayEffectStackingExtension_DiminishingReturnsTest.cpp
deleted file mode 100644
diff --git a/Engine/Source/Runtime/GameplayAbilities/Private/GameplayEffectTests.cpp b/Engine/Source/Runtime/GameplayAbilities/Private/GameplayEffectTests.cpp
--- a/Engine/Source/Runtime/GameplayAbilities/Private/GameplayEffectTests.cpp
+++ b/Engine/Source/Runtime/GameplayAbilities/Private/GameplayEffectTests.cpp
@@ -9,8 +9,6 @@
 #include "AbilitySystemComponent.h"
 #include "GameplayEffectExtension_LifestealTest.h"
 #include "GameplayEffectExtension_ShieldTest.h"
-#include "GameplayEffectStackingExtension_CappedNumberTest.h"
-#include "GameplayEffectStackingExtension_DiminishingReturnsTest.h"
 
 #define SKILL_TEST_TEXT( Format, ... ) FString::Printf(TEXT("%s - %d: %s"), TEXT(__FILE__) , __LINE__ , *FString::Printf(TEXT(Format), ##__VA_ARGS__) )
 
diff --git a/Engine/Source/Runtime/GameplayAbilities/Private/GameplayEffectTypes.cpp b/Engine/Source/Runtime/GameplayAbilities/Private/GameplayEffectTypes.cpp
--- a/Engine/Source/Runtime/GameplayAbilities/Private/GameplayEffectTypes.cpp
+++ b/Engine/Source/Runtime/GameplayAbilities/Private/GameplayEffectTypes.cpp
@@ -205,18 +205,6 @@ FString EGameplayModEffectToString(int32 Type)
 	return e->GetEnum(Type).ToString();
 }
 
-FString EGameplayEffectCopyPolicyToString(int32 Type)
-{
-	static UEnum *e = FindObject<UEnum>(ANY_PACKAGE, TEXT("EGameplayEffectCopyPolicy"));
-	return e->GetEnum(Type).ToString();
-}
-
-FString EGameplayEffectStackingPolicyToString(int32 Type)
-{
-	static UEnum *e = FindObject<UEnum>(ANY_PACKAGE, TEXT("EGameplayEffectStackingPolicy"));
-	return e->GetEnum(Type).ToString();
-}
-
 bool FGameplayTagCountContainer::HasMatchingGameplayTag(FGameplayTag TagToCheck) const
 {
 	return ExplicitTags.HasTag(TagToCheck, EGameplayTagMatchType::IncludeParentTags, EGameplayTagMatchType::Explicit);
diff --git a/Engine/Source/Runtime/GameplayAbilities/Public/AbilitySystemComponent.h b/Engine/Source/Runtime/GameplayAbilities/Public/AbilitySystemComponent.h
--- a/Engine/Source/Runtime/GameplayAbilities/Public/AbilitySystemComponent.h
+++ b/Engine/Source/Runtime/GameplayAbilities/Public/AbilitySystemComponent.h
@@ -244,9 +244,6 @@ class GAMEPLAYABILITIES_API UAbilitySystemComponent : public UActorComponent, pu
 	UFUNCTION(BlueprintCallable, Category = GameplayEffects)
 	float GetGameplayEffectMagnitude(FActiveGameplayEffectHandle Handle, FGameplayAttribute Attribute) const;
 
-	UFUNCTION(BlueprintCallable, Category = GameplayEffects)
-	bool IsGameplayEffectActive(FActiveGameplayEffectHandle InHandle) const;
-
 	/**
 	 * Get the source tags from the gameplay spec represented by the specified handle, if possible
 	 * 
diff --git a/Engine/Source/Runtime/GameplayAbilities/Public/GameplayEffect.h b/Engine/Source/Runtime/GameplayAbilities/Public/GameplayEffect.h
--- a/Engine/Source/Runtime/GameplayAbilities/Public/GameplayEffect.h
+++ b/Engine/Source/Runtime/GameplayAbilities/Public/GameplayEffect.h
@@ -19,15 +19,6 @@ class UAbilitySystemComponent;
 class UGameplayModMagnitudeCalculation;
 class UGameplayEffectExecutionCalculation;
 
-USTRUCT()
-struct FGameplayEffectStackingCallbacks
-{
-	GENERATED_USTRUCT_BODY()
-
-	UPROPERTY(EditDefaultsOnly, Category = GEStack)
-	TArray<TSubclassOf<class UGameplayEffectStackingExtension> >	ExtensionClasses;
-};
-
 /** Enumeration outlining the possible gameplay effect magnitude calculation policies */
 UENUM()
 enum class EGameplayEffectMagnitudeCalculation : uint8
@@ -574,20 +565,6 @@ public:
 	UPROPERTY(EditDefaultsOnly, BlueprintReadWrite, Category = Display)
 	FText Description;
 
-	// ----------------------------------------------
-
-	/** Specifies the rule used to stack this GameplayEffect with other GameplayEffects. */
-	UPROPERTY(EditDefaultsOnly, BlueprintReadOnly, Category = Stacking)
-	TEnumAsByte<EGameplayEffectStackingPolicy::Type>	StackingPolicy;
-
-	/** An identifier for the stack. Both names and stacking policy must match for GameplayEffects to stack with each other. */
-	UPROPERTY(EditDefaultsOnly, BlueprintReadOnly, Category = Stacking)
-	FName StackedAttribName;
-
-	/** Specifies a custom stacking rule if one is needed. */
-	UPROPERTY(EditDefaultsOnly, Category = Stacking)
-	TSubclassOf<class UGameplayEffectStackingExtension> StackingExtension;
-
 	// ----------------------------------------------------------------------
 	//	Tag Containers
 	// ----------------------------------------------------------------------
@@ -856,7 +833,6 @@ struct GAMEPLAYABILITIES_API FGameplayEffectSpec
 		, Period(UGameplayEffect::NO_PERIOD)
 		, ChanceToApplyToTarget(1.f)
 		, ChanceToExecuteOnGameplayEffect(1.f)
-		, bTopOfStack(false)
 		, Level(UGameplayEffect::INVALID_LEVEL)
 	{
 
@@ -902,8 +878,6 @@ struct GAMEPLAYABILITIES_API FGameplayEffectSpec
 	float GetChanceToExecuteOnGameplayEffect() const;
 	float GetMagnitude(const FGameplayAttribute &Attribute) const;
 
-	EGameplayEffectStackingPolicy::Type GetStackingType() const;
-
 	/** other effects that need to be applied to the target if this effect is successful */
 	TArray< FGameplayEffectSpecHandle > TargetEffectSpecs;
 
@@ -941,11 +915,6 @@ struct GAMEPLAYABILITIES_API FGameplayEffectSpec
 	UPROPERTY()
 	float ChanceToExecuteOnGameplayEffect;
 
-	// This should only be true if this is a stacking effect and at the top of its stack
-	// (FIXME: should this be part of the spec or FActiveGameplayEffect?)
-	UPROPERTY()
-	bool bTopOfStack;
-
 	// Captured Source Tags on GameplayEffectSpec creation.	
 	UPROPERTY(NotReplicated)
 	FTagContainerAggregator	CapturedSourceTags;
@@ -1070,8 +1039,6 @@ struct FActiveGameplayEffect : public FFastArraySerializerItem
 
 	void CheckOngoingTagRequirements(const FGameplayTagContainer& OwnerTags, struct FActiveGameplayEffectsContainer& OwningContainer);
 
-	bool CanBeStacked(const FActiveGameplayEffect& Other) const;
-
 	void PrintAll() const;
 
 	void PreReplicatedRemove(const struct FActiveGameplayEffectsContainer &InArray);
@@ -1175,7 +1142,7 @@ struct FActiveGameplayEffectsContainer : public FFastArraySerializer
 	friend struct FActiveGameplayEffect;
 	friend class UAbilitySystemComponent;
 
-	FActiveGameplayEffectsContainer() : bNeedToRecalculateStacks(false) {};
+	FActiveGameplayEffectsContainer() {};
 
 	UAbilitySystemComponent* Owner;
 
@@ -1200,10 +1167,6 @@ struct FActiveGameplayEffectsContainer : public FFastArraySerializer
 
 	float GetGameplayEffectMagnitude(FActiveGameplayEffectHandle Handle, FGameplayAttribute Attribute) const;
 
-	// returns true if the handle points to an effect in this container that is not a stacking effect or an effect in this container that does stack and is applied by the current stacking rules
-	// returns false if the handle points to an effect that is not in this container or is not applied because of the current stacking rules
-	bool IsGameplayEffectActive(FActiveGameplayEffectHandle Handle, bool IncludeEffectsBlockedByStackingRules = false) const;
-
 	void SetAttributeBaseValue(FGameplayAttribute Attribute, float NewBaseValue);
 
 	/** Actually applies given mod to the attribute */
@@ -1243,19 +1206,12 @@ struct FActiveGameplayEffectsContainer : public FFastArraySerializer
 
 	void CheckDuration(FActiveGameplayEffectHandle Handle);
 
-	void StacksNeedToRecalculate();
-
-	// recalculates all of the stacks in the current container
-	void RecalculateStacking();
-
 	bool NetDeltaSerialize(FNetDeltaSerializeInfo & DeltaParms)
 	{
 		return FastArrayDeltaSerialize<FActiveGameplayEffect>(GameplayEffects, DeltaParms, *this);
 	}
 
-	void PreDestroy();
-
-	bool bNeedToRecalculateStacks;
+	void PreDestroy();	
 
 	// ------------------------------------------------
 
@@ -1287,8 +1243,6 @@ struct FActiveGameplayEffectsContainer : public FFastArraySerializer
 
 private:
 
-	FTimerHandle StackHandle;
-
 	void InternalUpdateNumericalAttribute(FGameplayAttribute Attribute, float NewValue, const FGameplayEffectModCallbackData* ModData);
 	
 	/**
diff --git a/Engine/Source/Runtime/GameplayAbilities/Public/GameplayEffectStackingExtension.h b/Engine/Source/Runtime/GameplayAbilities/Public/GameplayEffectStackingExtension.h
deleted file mode 100644
diff --git a/Engine/Source/Runtime/GameplayAbilities/Public/GameplayEffectStackingExtension_CappedNumberTest.h b/Engine/Source/Runtime/GameplayAbilities/Public/GameplayEffectStackingExtension_CappedNumberTest.h
deleted file mode 100644
diff --git a/Engine/Source/Runtime/GameplayAbilities/Public/GameplayEffectStackingExtension_DiminishingReturnsTest.h b/Engine/Source/Runtime/GameplayAbilities/Public/GameplayEffectStackingExtension_DiminishingReturnsTest.h
deleted file mode 100644
diff --git a/Engine/Source/Runtime/GameplayAbilities/Public/GameplayEffectTypes.h b/Engine/Source/Runtime/GameplayAbilities/Public/GameplayEffectTypes.h
--- a/Engine/Source/Runtime/GameplayAbilities/Public/GameplayEffectTypes.h
+++ b/Engine/Source/Runtime/GameplayAbilities/Public/GameplayEffectTypes.h
@@ -28,10 +28,6 @@ FString EGameplayModToString(int32 Type);
 
 FString EGameplayModEffectToString(int32 Type);
 
-FString EGameplayEffectCopyPolicyToString(int32 Type);
-
-FString EGameplayEffectStackingPolicyToString(int32 Type);
-
 UENUM(BlueprintType)
 namespace EGameplayModOp
 {
@@ -50,38 +46,6 @@ namespace EGameplayModOp
 	};
 }
 
-/**
- * Tells us how to handle copying gameplay effect when it is applied.
- *	Default means to use context - e.g, OutgoingGE is are always snapshots, IncomingGE is always Link
- *	AlwaysSnapshot vs AlwaysLink let mods themselves override
- */
-UENUM(BlueprintType)
-namespace EGameplayEffectCopyPolicy
-{
-	enum Type
-	{
-		Default = 0			UMETA(DisplayName="Default"),
-		AlwaysSnapshot		UMETA(DisplayName="AlwaysSnapshot"),
-		AlwaysLink			UMETA(DisplayName="AlwaysLink")
-	};
-}
-
-UENUM(BlueprintType)
-namespace EGameplayEffectStackingPolicy
-{
-	enum Type
-	{
-		Unlimited = 0		UMETA(DisplayName = "NoRule"),
-		Highest				UMETA(DisplayName = "Strongest"),
-		Lowest				UMETA(DisplayName = "Weakest"),
-		Replaces			UMETA(DisplayName = "MostRecent"),
-		Callback			UMETA(DisplayName = "Custom"),
-
-		// This must always be at the end
-		Max					UMETA(DisplayName = "Invalid")
-	};
-}
-
 /** Enumeration for options of where to capture gameplay attributes from for gameplay effects */
 UENUM()
 enum class EGameplayEffectAttributeCaptureSource : uint8
