#ue4 - Fix build (FGameplayEffectSpec.bTopOfStack undefined)

diff --git a/Engine/Source/Runtime/GameplayAbilities/Private/GameplayEffect.cpp b/Engine/Source/Runtime/GameplayAbilities/Private/GameplayEffect.cpp
--- a/Engine/Source/Runtime/GameplayAbilities/Private/GameplayEffect.cpp
+++ b/Engine/Source/Runtime/GameplayAbilities/Private/GameplayEffect.cpp
@@ -383,7 +383,6 @@ FGameplayEffectSpec::FGameplayEffectSpec(const UGameplayEffect* InDef, const FGa
 	, Period(UGameplayEffect::NO_PERIOD)
 	, ChanceToApplyToTarget(1.f)
 	, ChanceToExecuteOnGameplayEffect(1.f)
-	, bTopOfStack(false)
 	, Level(UGameplayEffect::INVALID_LEVEL)
 {
 	check(Def);	
