Fix build error/warning

diff --git a/Engine/Source/Runtime/GameplayAbilities/Private/GameplayEffect.cpp b/Engine/Source/Runtime/GameplayAbilities/Private/GameplayEffect.cpp
--- a/Engine/Source/Runtime/GameplayAbilities/Private/GameplayEffect.cpp
+++ b/Engine/Source/Runtime/GameplayAbilities/Private/GameplayEffect.cpp
@@ -379,11 +379,6 @@ void FGameplayEffectExecutionDefinition::GetAttributeCaptureDefinitions(OUT TArr
 
 FGameplayEffectSpec::FGameplayEffectSpec(const UGameplayEffect* InDef, const FGameplayEffectContextHandle& InEffectContext, float Level)
 	: Def(InDef)
-	, Duration(UGameplayEffect::INSTANT_APPLICATION)
-	, Period(UGameplayEffect::NO_PERIOD)
-	, ChanceToApplyToTarget(1.f)
-	, ChanceToExecuteOnGameplayEffect(1.f)
-	, Level(UGameplayEffect::INVALID_LEVEL)
 {
 	check(Def);	
 	SetLevel(Level);
