Hiding some data from the editor, adding auto-hide feature for Orion GCNAs as a convenience.

diff --git a/Engine/Source/Runtime/GameplayAbilities/Public/GameplayCueNotify_Actor.h b/Engine/Source/Runtime/GameplayAbilities/Public/GameplayCueNotify_Actor.h
--- a/Engine/Source/Runtime/GameplayAbilities/Public/GameplayCueNotify_Actor.h
+++ b/Engine/Source/Runtime/GameplayAbilities/Public/GameplayCueNotify_Actor.h
@@ -19,7 +19,7 @@
  *	
  */
 
-UCLASS(Blueprintable,meta=(ShowWorldContextPin))
+UCLASS(Blueprintable, meta = (ShowWorldContextPin), hidecategories = (Replication))
 class GAMEPLAYABILITIES_API AGameplayCueNotify_Actor : public AActor
 {
 	GENERATED_UCLASS_BODY()
diff --git a/Engine/Source/Runtime/GameplayAbilities/Public/GameplayCueNotify_Static.h b/Engine/Source/Runtime/GameplayAbilities/Public/GameplayCueNotify_Static.h
--- a/Engine/Source/Runtime/GameplayAbilities/Public/GameplayCueNotify_Static.h
+++ b/Engine/Source/Runtime/GameplayAbilities/Public/GameplayCueNotify_Static.h
@@ -19,7 +19,7 @@
  *	
  */
 
-UCLASS(Blueprintable,meta=(ShowWorldContextPin))
+UCLASS(Blueprintable, meta = (ShowWorldContextPin), hidecategories = (Replication))
 class GAMEPLAYABILITIES_API UGameplayCueNotify_Static : public UObject
 {
 	GENERATED_UCLASS_BODY()
