Fix a couple of Doxygen warnings.

diff --git a/Engine/Source/Runtime/GameplayAbilities/Private/GameplayEffectTypes.cpp b/Engine/Source/Runtime/GameplayAbilities/Private/GameplayEffectTypes.cpp
--- a/Engine/Source/Runtime/GameplayAbilities/Private/GameplayEffectTypes.cpp
+++ b/Engine/Source/Runtime/GameplayAbilities/Private/GameplayEffectTypes.cpp
@@ -221,7 +221,7 @@ bool FGameplayTagCountContainer::HasAnyMatchingGameplayTags(const FGameplayTagCo
 	return ExplicitTags.MatchesAny(TagContainer, bCountEmptyAsMatch);
 }
 
-void FGameplayTagCountContainer::UpdateTagCount(const struct FGameplayTag& Tag, int32 CountDelta)
+void FGameplayTagCountContainer::UpdateTagCount(const FGameplayTag& Tag, int32 CountDelta)
 {
 	if (CountDelta != 0)
 	{
@@ -229,7 +229,7 @@ void FGameplayTagCountContainer::UpdateTagCount(const struct FGameplayTag& Tag,
 	}
 }
 
-void FGameplayTagCountContainer::UpdateTagCount(const struct FGameplayTagContainer& Container, int32 CountDelta)
+void FGameplayTagCountContainer::UpdateTagCount(const FGameplayTagContainer& Container, int32 CountDelta)
 {	
 	if (CountDelta != 0)
 	{
