Fix prediction problems due to merge

diff --git a/Engine/Source/Runtime/GameplayAbilities/Private/AbilitySystemComponent.cpp b/Engine/Source/Runtime/GameplayAbilities/Private/AbilitySystemComponent.cpp
--- a/Engine/Source/Runtime/GameplayAbilities/Private/AbilitySystemComponent.cpp
+++ b/Engine/Source/Runtime/GameplayAbilities/Private/AbilitySystemComponent.cpp
@@ -896,7 +896,7 @@ void UAbilitySystemComponent::RemoveAllGameplayCues()
 
 void UAbilitySystemComponent::NetMulticast_InvokeGameplayCueExecuted_FromSpec_Implementation(const FGameplayEffectSpec Spec, FPredictionKey PredictionKey)
 {
-	if (IsOwnerActorAuthoritative() || PredictionKey.Current > ReplicatedPredictionKey.Current)
+	if (IsOwnerActorAuthoritative() || PredictionKey.IsValidKey() == false)
 	{
 		InvokeGameplayCueEvent(Spec, EGameplayCueEvent::Executed);
 	}
@@ -904,7 +904,7 @@ void UAbilitySystemComponent::NetMulticast_InvokeGameplayCueExecuted_FromSpec_Im
 
 void UAbilitySystemComponent::NetMulticast_InvokeGameplayCueExecuted_Implementation(const FGameplayTag GameplayCueTag, FPredictionKey PredictionKey, FGameplayEffectContextHandle EffectContext)
 {
-	if (IsOwnerActorAuthoritative() || PredictionKey.Current > ReplicatedPredictionKey.Current)
+	if (IsOwnerActorAuthoritative() || PredictionKey.IsValidKey() == false)
 	{
 		InvokeGameplayCueEvent(GameplayCueTag, EGameplayCueEvent::Executed, EffectContext);
 	}
@@ -912,7 +912,7 @@ void UAbilitySystemComponent::NetMulticast_InvokeGameplayCueExecuted_Implementat
 
 void UAbilitySystemComponent::NetMulticast_InvokeGameplayCueAdded_Implementation(const FGameplayTag GameplayCueTag, FPredictionKey PredictionKey, FGameplayEffectContextHandle EffectContext)
 {
-	if (IsOwnerActorAuthoritative() || PredictionKey.Current > ReplicatedPredictionKey.Current)
+	if (IsOwnerActorAuthoritative() || PredictionKey.IsValidKey() == false)
 	{
 		InvokeGameplayCueEvent(GameplayCueTag, EGameplayCueEvent::OnActive, EffectContext);
 	}
@@ -920,7 +920,7 @@ void UAbilitySystemComponent::NetMulticast_InvokeGameplayCueAdded_Implementation
 
 void UAbilitySystemComponent::NetMulticast_InvokeGameplayCueRemoved_Implementation(const FGameplayTag GameplayCueTag, FPredictionKey PredictionKey)
 {
-	if (IsOwnerActorAuthoritative() || PredictionKey.Current > ReplicatedPredictionKey.Current)
+	if (IsOwnerActorAuthoritative() || PredictionKey.IsValidKey() == false)
 	{
 		InvokeGameplayCueEvent(GameplayCueTag, EGameplayCueEvent::Removed);
 	}
