bServerRespectsRemoteAbilityCancellation
-fix brokeness after merge - was always being set to true
-Default to true in UGameplayAbility
-fix spelling

diff --git a/Engine/Source/Runtime/GameplayAbilities/Private/Abilities/GameplayAbility.cpp b/Engine/Source/Runtime/GameplayAbilities/Private/Abilities/GameplayAbility.cpp
--- a/Engine/Source/Runtime/GameplayAbilities/Private/Abilities/GameplayAbility.cpp
+++ b/Engine/Source/Runtime/GameplayAbilities/Private/Abilities/GameplayAbility.cpp
@@ -51,7 +51,7 @@ UGameplayAbility::UGameplayAbility(const FObjectInitializer& ObjectInitializer)
 	}
 #endif
 
-	bServerRespectsRemoteAbilityCancelation = false;
+	bServerRespectsRemoteAbilityCancellation = true;
 }
 
 int32 UGameplayAbility::GetFunctionCallspace(UFunction* Function, void* Parameters, FFrame* Stack)
diff --git a/Engine/Source/Runtime/GameplayAbilities/Private/AbilitySystemComponent_Abilities.cpp b/Engine/Source/Runtime/GameplayAbilities/Private/AbilitySystemComponent_Abilities.cpp
--- a/Engine/Source/Runtime/GameplayAbilities/Private/AbilitySystemComponent_Abilities.cpp
+++ b/Engine/Source/Runtime/GameplayAbilities/Private/AbilitySystemComponent_Abilities.cpp
@@ -711,7 +711,6 @@ bool UAbilitySystemComponent::TryActivateAbility(FGameplayAbilitySpecHandle Hand
 		if (Ability->GetInstancingPolicy() == EGameplayAbilityInstancingPolicy::InstancedPerExecution)
 		{
 			InstancedAbility = CreateNewInstanceOfAbility(*Spec, Ability);
-			ActivationInfo.bCanBeEndedByOtherInstance = Ability->bServerRespectsRemoteAbilityCancelation;
 			InstancedAbility->CallActivateAbility(Handle, ActorInfo, ActivationInfo, OnGameplayAbilityEndedDelegate, TriggerEventData);
 		}
 		else if (InstancedAbility)
@@ -892,7 +891,7 @@ void UAbilitySystemComponent::InternalServerTryActiveAbility(FGameplayAbilitySpe
 		if (InstancedAbility)
 		{
 			//Only applies to instanced abilities
-			InstancedAbility->CurrentActivationInfo.bCanBeEndedByOtherInstance = true;
+			InstancedAbility->CurrentActivationInfo.bCanBeEndedByOtherInstance = InstancedAbility->bServerRespectsRemoteAbilityCancellation;
 		}
 	}
 	else
diff --git a/Engine/Source/Runtime/GameplayAbilities/Public/Abilities/GameplayAbility.h b/Engine/Source/Runtime/GameplayAbilities/Public/Abilities/GameplayAbility.h
--- a/Engine/Source/Runtime/GameplayAbilities/Public/Abilities/GameplayAbility.h
+++ b/Engine/Source/Runtime/GameplayAbilities/Public/Abilities/GameplayAbility.h
@@ -542,7 +542,7 @@ protected:
 
 	/** If this is set, the server-side version of the ability can be canceled by the client-side version. The client-side version can always be canceled by the server. */
 	UPROPERTY(EditDefaultsOnly, Category = Advanced)
-	bool bServerRespectsRemoteAbilityCancelation;
+	bool bServerRespectsRemoteAbilityCancellation;
 
 	/** This is information specific to this instance of the ability. E.g, whether it is predicting, authoring, confirmed, etc. */
 	UPROPERTY(BlueprintReadOnly, Category = Ability)
