Merge ability system changes to main; Includes beginnings of gameplay effect stacking support

diff --git a/Engine/Source/Runtime/GameplayAbilities/Private/Abilities/GameplayAbility.cpp b/Engine/Source/Runtime/GameplayAbilities/Private/Abilities/GameplayAbility.cpp
--- a/Engine/Source/Runtime/GameplayAbilities/Private/Abilities/GameplayAbility.cpp
+++ b/Engine/Source/Runtime/GameplayAbilities/Private/Abilities/GameplayAbility.cpp
@@ -959,6 +959,11 @@ bool UGameplayAbility::IsTriggered() const
 	return AbilityTriggers.Num() > 0;
 }
 
+bool UGameplayAbility::HasAuthority(const FGameplayAbilityActivationInfo* ActivationInfo) const
+{
+	return (ActivationInfo->ActivationMode == EGameplayAbilityActivationMode::Authority);
+}
+
 bool UGameplayAbility::HasAuthorityOrPredictionKey(const FGameplayAbilityActorInfo* ActorInfo, const FGameplayAbilityActivationInfo* ActivationInfo) const
 {
 	return ActorInfo->AbilitySystemComponent->HasAuthorityOrPredictionKey(ActivationInfo);
@@ -1086,6 +1091,30 @@ TArray<FActiveGameplayEffectHandle> UGameplayAbility::ApplyGameplayEffectSpecToT
 	return EffectHandles;
 }
 
+void UGameplayAbility::BP_RemoveGameplayEffectFromOwnerWithAssetTags(FGameplayTagContainer WithTags, int32 StacksToRemove)
+{
+	if (HasAuthority(&CurrentActivationInfo) == false)
+	{
+		return;
+	}
+
+	FActiveGameplayEffectQuery Query;
+	Query.EffectTagContainer = &WithTags;
+	CurrentActorInfo->AbilitySystemComponent->RemoveActiveEffects(Query, StacksToRemove);
+}
+
+void UGameplayAbility::BP_RemoveGameplayEffectFromOwnerWithGrantedTags(FGameplayTagContainer WithGrantedTags, int32 StacksToRemove)
+{
+	if (HasAuthority(&CurrentActivationInfo) == false)
+	{
+		return;
+	}
+
+	FActiveGameplayEffectQuery Query;
+	Query.OwningTagContainer = &WithGrantedTags;
+	CurrentActorInfo->AbilitySystemComponent->RemoveActiveEffects(Query, StacksToRemove);
+}
+
 void UGameplayAbility::ConvertDeprecatedGameplayEffectReferencesToBlueprintReferences(UGameplayEffect* OldGE, TSubclassOf<UGameplayEffect> NewGEClass)
 {
 	bool bChangedSomething = false;
diff --git a/Engine/Source/Runtime/GameplayAbilities/Private/Abilities/GameplayAbilityTargetActor.cpp b/Engine/Source/Runtime/GameplayAbilities/Private/Abilities/GameplayAbilityTargetActor.cpp
--- a/Engine/Source/Runtime/GameplayAbilities/Private/Abilities/GameplayAbilityTargetActor.cpp
+++ b/Engine/Source/Runtime/GameplayAbilities/Private/Abilities/GameplayAbilityTargetActor.cpp
@@ -15,7 +15,6 @@
 AGameplayAbilityTargetActor::AGameplayAbilityTargetActor(const FObjectInitializer& ObjectInitializer)
 	: Super(ObjectInitializer)
 {
-	StaticTargetFunction = false;
 	ShouldProduceTargetDataOnServer = false;
 	bDebug = false;
 	bDestroyOnConfirmation = true;
@@ -114,11 +113,6 @@ bool AGameplayAbilityTargetActor::GetReplicates() const
 	return bReplicates;
 }
 
-FGameplayAbilityTargetDataHandle AGameplayAbilityTargetActor::StaticGetTargetData(UWorld * World, const FGameplayAbilityActorInfo* ActorInfo, FGameplayAbilityActivationInfo ActivationInfo) const
-{
-	return FGameplayAbilityTargetDataHandle();
-}
-
 bool AGameplayAbilityTargetActor::OnReplicatedTargetDataReceived(FGameplayAbilityTargetDataHandle& Data) const
 {
 	return true;
diff --git a/Engine/Source/Runtime/GameplayAbilities/Private/Abilities/GameplayAbilityTargetActor_Radius.cpp b/Engine/Source/Runtime/GameplayAbilities/Private/Abilities/GameplayAbilityTargetActor_Radius.cpp
--- a/Engine/Source/Runtime/GameplayAbilities/Private/Abilities/GameplayAbilityTargetActor_Radius.cpp
+++ b/Engine/Source/Runtime/GameplayAbilities/Private/Abilities/GameplayAbilityTargetActor_Radius.cpp
@@ -18,7 +18,6 @@ AGameplayAbilityTargetActor_Radius::AGameplayAbilityTargetActor_Radius(const FOb
 {
 	PrimaryActorTick.bCanEverTick = true;
 	PrimaryActorTick.TickGroup = TG_PrePhysics;
-	StaticTargetFunction = false;
 	ShouldProduceTargetDataOnServer = true;
 }
 
diff --git a/Engine/Source/Runtime/GameplayAbilities/Private/Abilities/GameplayAbilityTargetActor_Trace.cpp b/Engine/Source/Runtime/GameplayAbilities/Private/Abilities/GameplayAbilityTargetActor_Trace.cpp
--- a/Engine/Source/Runtime/GameplayAbilities/Private/Abilities/GameplayAbilityTargetActor_Trace.cpp
+++ b/Engine/Source/Runtime/GameplayAbilities/Private/Abilities/GameplayAbilityTargetActor_Trace.cpp
@@ -17,7 +17,6 @@ AGameplayAbilityTargetActor_Trace::AGameplayAbilityTargetActor_Trace(const FObje
 {
 	PrimaryActorTick.bCanEverTick = true;
 	PrimaryActorTick.TickGroup = TG_PostUpdateWork;
-	StaticTargetFunction = false;
 
 	MaxRange = 999999.0f;
 	bTraceAffectsAimPitch = true;
@@ -79,14 +78,6 @@ void AGameplayAbilityTargetActor_Trace::SweepWithFilter(FHitResult& OutHitResult
 	}
 }
 
-FGameplayAbilityTargetDataHandle AGameplayAbilityTargetActor_Trace::StaticGetTargetData(UWorld * World, const FGameplayAbilityActorInfo* ActorInfo, FGameplayAbilityActivationInfo ActivationInfo) const
-{
-	check(false);		//This should never actually be called, and if it is, it will require a const version of PerformTrace()
-	AActor* StaticSourceActor = ActorInfo->AvatarActor.Get();
-	//return MakeTargetData(PerformTrace(StaticSourceActor));		//Old way, requires PerformTrace to be a const call
-	return MakeTargetData(FHitResult());
-}
-
 void AGameplayAbilityTargetActor_Trace::AimWithPlayerController(const AActor* InSourceActor, FCollisionQueryParams Params, const FVector& TraceStart, FVector& OutTraceEnd, bool bIgnorePitch) const
 {
 	if (!OwningAbility) // Server and launching client only
diff --git a/Engine/Source/Runtime/GameplayAbilities/Private/Abilities/Tasks/AbilityTask_VisualizeTargeting.cpp b/Engine/Source/Runtime/GameplayAbilities/Private/Abilities/Tasks/AbilityTask_VisualizeTargeting.cpp
--- a/Engine/Source/Runtime/GameplayAbilities/Private/Abilities/Tasks/AbilityTask_VisualizeTargeting.cpp
+++ b/Engine/Source/Runtime/GameplayAbilities/Private/Abilities/Tasks/AbilityTask_VisualizeTargeting.cpp
@@ -14,12 +14,10 @@ UAbilityTask_VisualizeTargeting::UAbilityTask_VisualizeTargeting(const FObjectIn
 UAbilityTask_VisualizeTargeting* UAbilityTask_VisualizeTargeting::VisualizeTargeting(UObject* WorldContextObject, TSubclassOf<AGameplayAbilityTargetActor> InTargetClass, FName TaskInstanceName, float Duration)
 {
 	UAbilityTask_VisualizeTargeting* MyObj = NewTask<UAbilityTask_VisualizeTargeting>(WorldContextObject, TaskInstanceName);		//Register for task list here, providing a given FName as a key
-	MyObj->TargetClass = InTargetClass;
 	MyObj->SetDuration(Duration);
 	return MyObj;
 }
 
-// ---------------------------------------------------------------------------------------
 
 void UAbilityTask_VisualizeTargeting::SetDuration(const float Duration)
 {
@@ -37,22 +35,13 @@ bool UAbilityTask_VisualizeTargeting::BeginSpawningActor(UObject* WorldContextOb
 	if (MyAbility)
 	{
 		const AGameplayAbilityTargetActor* CDO = CastChecked<AGameplayAbilityTargetActor>(TargetClass->GetDefaultObject());
-		MyTargetActor = CDO;
 
-		bool Replicates = CDO->GetReplicates();
-		bool StaticFunc = CDO->StaticTargetFunction;
-		bool IsLocallyControlled = MyAbility->GetCurrentActorInfo()->IsLocallyControlled();
-
-		if (Replicates && StaticFunc)
-		{
-			// We can't replicate a staticFunc target actor, since we are just calling a static function and not spawning an actor at all!
-			ABILITY_LOG(Fatal, TEXT("AbilityTargetActor class %s can't be Replicating and Static"), *TargetClass->GetName());
-			Replicates = false;
-		}
+		const bool bReplicates = CDO->GetReplicates();
+		const bool bIsLocallyControlled = MyAbility->GetCurrentActorInfo()->IsLocallyControlled();
 
 		// Spawn the actor if this is a locally controlled ability (always) or if this is a replicating targeting mode.
 		// (E.g., server will spawn this target actor to replicate to all non owning clients)
-		if (Replicates || IsLocallyControlled)
+		if (bReplicates || bIsLocallyControlled)
 		{
 			UClass* Class = *TargetClass;
 			if (Class != NULL)
@@ -87,22 +76,18 @@ void UAbilityTask_VisualizeTargeting::FinishSpawningActor(UObject* WorldContextO
 	}
 }
 
-// ---------------------------------------------------------------------------------------
-
 void UAbilityTask_VisualizeTargeting::OnDestroy(bool AbilityEnded)
 {
-	if (MyTargetActor.IsValid() && !MyTargetActor->HasAnyFlags(RF_ClassDefaultObject))
+	if (MyTargetActor.IsValid())
 	{
 		MyTargetActor->Destroy();
 	}
+
 	GetWorld()->GetTimerManager().ClearTimer(TimerHandle_OnTimeElapsed);
 
 	Super::OnDestroy(AbilityEnded);
 }
 
-
-// --------------------------------------------------------------------------------------
-
 void UAbilityTask_VisualizeTargeting::OnTimeElapsed()
 {
 	TimeElapsed.Broadcast();
diff --git a/Engine/Source/Runtime/GameplayAbilities/Private/Abilities/Tasks/AbilityTask_WaitTargetData.cpp b/Engine/Source/Runtime/GameplayAbilities/Private/Abilities/Tasks/AbilityTask_WaitTargetData.cpp
--- a/Engine/Source/Runtime/GameplayAbilities/Private/Abilities/Tasks/AbilityTask_WaitTargetData.cpp
+++ b/Engine/Source/Runtime/GameplayAbilities/Private/Abilities/Tasks/AbilityTask_WaitTargetData.cpp
@@ -19,8 +19,6 @@ UAbilityTask_WaitTargetData* UAbilityTask_WaitTargetData::WaitTargetData(UObject
 	return MyObj;
 }
 
-// ---------------------------------------------------------------------------------------
-
 bool UAbilityTask_WaitTargetData::BeginSpawningActor(UObject* WorldContextObject, TSubclassOf<AGameplayAbilityTargetActor> TargetClass, AGameplayAbilityTargetActor*& SpawnedActor)
 {
 	SpawnedActor = nullptr;
@@ -29,55 +27,37 @@ bool UAbilityTask_WaitTargetData::BeginSpawningActor(UObject* WorldContextObject
 	if (MyAbility)
 	{
 		const AGameplayAbilityTargetActor* CDO = CastChecked<AGameplayAbilityTargetActor>(TargetClass->GetDefaultObject());
-		MyTargetActor = CDO;
-
-		bool Replicates = CDO->GetReplicates();
-		bool StaticFunc = CDO->StaticTargetFunction;
-		bool IsLocallyControlled = MyAbility->GetCurrentActorInfo()->IsLocallyControlled();
 
-		if (Replicates && StaticFunc)
-		{
-			// We can't replicate a staticFunc target actor, since we are just calling a static function and not spawning an actor at all!
-			ABILITY_LOG(Fatal, TEXT("AbilityTargetActor class %s can't be Replicating and Static"), *TargetClass->GetName());
-			Replicates = false;
-		}
+		const bool bReplicates = CDO->GetReplicates();
+		const bool bIsLocallyControlled = MyAbility->GetCurrentActorInfo()->IsLocallyControlled();
 
 		// Spawn the actor if this is a locally controlled ability (always) or if this is a replicating targeting mode.
 		// (E.g., server will spawn this target actor to replicate to all non owning clients)
-		if (Replicates || IsLocallyControlled || CDO->ShouldProduceTargetDataOnServer)
+		if (bReplicates || bIsLocallyControlled || CDO->ShouldProduceTargetDataOnServer)
 		{
-			if (StaticFunc)
+			UClass* Class = *TargetClass;
+			if (Class != NULL)
 			{
-				// This is just a static function that should instantly give us back target data
-				FGameplayAbilityTargetDataHandle Data = CDO->StaticGetTargetData(MyAbility->GetWorld(), MyAbility->GetCurrentActorInfo(), MyAbility->GetCurrentActivationInfo());
-				OnTargetDataReadyCallback(Data);
+				UWorld* World = GEngine->GetWorldFromContextObject(WorldContextObject);
+				SpawnedActor = World->SpawnActorDeferred<AGameplayAbilityTargetActor>(Class, FVector::ZeroVector, FRotator::ZeroRotator, NULL, NULL, true);
 			}
-			else
-			{
-				UClass* Class = *TargetClass;
-				if (Class != NULL)
-				{
-					UWorld* World = GEngine->GetWorldFromContextObject(WorldContextObject);
-					SpawnedActor = World->SpawnActorDeferred<AGameplayAbilityTargetActor>(Class, FVector::ZeroVector, FRotator::ZeroRotator, NULL, NULL, true);
-				}
 
-				// If we spawned the target actor, always register the callbacks for when the data is ready.
-				SpawnedActor->TargetDataReadyDelegate.AddUObject(this, &UAbilityTask_WaitTargetData::OnTargetDataReadyCallback);
-				SpawnedActor->CanceledDelegate.AddUObject(this, &UAbilityTask_WaitTargetData::OnTargetDataCancelledCallback);
+			// If we spawned the target actor, always register the callbacks for when the data is ready.
+			SpawnedActor->TargetDataReadyDelegate.AddUObject(this, &UAbilityTask_WaitTargetData::OnTargetDataReadyCallback);
+			SpawnedActor->CanceledDelegate.AddUObject(this, &UAbilityTask_WaitTargetData::OnTargetDataCancelledCallback);
 
-				MyTargetActor = SpawnedActor;
+			MyTargetActor = SpawnedActor;
 
-				AGameplayAbilityTargetActor* TargetActor = CastChecked<AGameplayAbilityTargetActor>(SpawnedActor);
-				if (TargetActor)
-				{
-					TargetActor->MasterPC = MyAbility->GetCurrentActorInfo()->PlayerController.Get();
-				}
+			AGameplayAbilityTargetActor* TargetActor = CastChecked<AGameplayAbilityTargetActor>(SpawnedActor);
+			if (TargetActor)
+			{
+				TargetActor->MasterPC = MyAbility->GetCurrentActorInfo()->PlayerController.Get();
 			}
 		}		
 
 		// If not locally controlled (server for remote client), see if TargetData was already sent
 		// else register callback for when it does get here.
-		if (!IsLocallyControlled)
+		if (!bIsLocallyControlled)
 		{
 			// Register with the TargetData callbacks if we are expecting client to send them
 			if (!CDO->ShouldProduceTargetDataOnServer)
@@ -157,8 +137,6 @@ void UAbilityTask_WaitTargetData::FinishSpawningActor(UObject* WorldContextObjec
 	}
 }
 
-// ---------------------------------------------------------------------------------------
-
 /** Valid TargetData was replicated to use (we are server, was sent from client) */
 void UAbilityTask_WaitTargetData::OnTargetDataReplicatedCallback(FGameplayAbilityTargetDataHandle Data)
 {
@@ -261,7 +239,7 @@ void UAbilityTask_WaitTargetData::OnDestroy(bool AbilityEnded)
 	AbilitySystemComponent->ReplicatedTargetDataDelegate.RemoveUObject(this, &UAbilityTask_WaitTargetData::OnTargetDataReplicatedCallback);
 	AbilitySystemComponent->ReplicatedTargetDataCancelledDelegate.RemoveDynamic(this, &UAbilityTask_WaitTargetData::OnTargetDataReplicatedCancelledCallback);
 
-	if (MyTargetActor.IsValid() && !MyTargetActor->HasAnyFlags(RF_ClassDefaultObject))
+	if (MyTargetActor.IsValid())
 	{
 		MyTargetActor->Destroy();
 	}
diff --git a/Engine/Source/Runtime/GameplayAbilities/Private/AbilitySystemBlueprintLibrary.cpp b/Engine/Source/Runtime/GameplayAbilities/Private/AbilitySystemBlueprintLibrary.cpp
--- a/Engine/Source/Runtime/GameplayAbilities/Private/AbilitySystemBlueprintLibrary.cpp
+++ b/Engine/Source/Runtime/GameplayAbilities/Private/AbilitySystemBlueprintLibrary.cpp
@@ -424,7 +424,7 @@ FGameplayEffectSpecHandle UAbilitySystemBlueprintLibrary::SetDuration(FGameplayE
 	FGameplayEffectSpec* Spec = SpecHandle.Data.Get();
 	if (Spec)
 	{
-		Spec->SetDuration(Duration);
+		Spec->SetDuration(Duration, true);
 	}
 	else
 	{
diff --git a/Engine/Source/Runtime/GameplayAbilities/Private/AbilitySystemComponent.cpp b/Engine/Source/Runtime/GameplayAbilities/Private/AbilitySystemComponent.cpp
--- a/Engine/Source/Runtime/GameplayAbilities/Private/AbilitySystemComponent.cpp
+++ b/Engine/Source/Runtime/GameplayAbilities/Private/AbilitySystemComponent.cpp
@@ -172,7 +172,12 @@ float UAbilitySystemComponent::GetNumericAttribute(const FGameplayAttribute &Att
 
 void UAbilitySystemComponent::ApplyModToAttribute(const FGameplayAttribute &Attribute, TEnumAsByte<EGameplayModOp::Type> ModifierOp, float ModifierMagnitude)
 {
-	ActiveGameplayEffects.ApplyModToAttribute(Attribute, ModifierOp, ModifierMagnitude);
+	// We can only apply loose mods on the authority. If we ever need to predict these, they would need to be turned into GEs and be given a prediction key so that
+	// they can be rolled back.
+	if (IsOwnerActorAuthoritative())
+	{
+		ActiveGameplayEffects.ApplyModToAttribute(Attribute, ModifierOp, ModifierMagnitude);
+	}
 }
 
 FGameplayEffectSpecHandle UAbilitySystemComponent::MakeOutgoingSpec(TSubclassOf<UGameplayEffect> GameplayEffectClass, float Level, FGameplayEffectContextHandle Context) const
@@ -525,9 +530,14 @@ FActiveGameplayEffectHandle UAbilitySystemComponent::ApplyGameplayEffectSpecToSe
 	{
 		if (Duration != UGameplayEffect::INSTANT_APPLICATION)
 		{
-			FActiveGameplayEffect& NewActiveEffect = ActiveGameplayEffects.CreateNewActiveGameplayEffect(Spec, PredictionKey);
-			MyHandle = NewActiveEffect.Handle;
-			OurCopyOfSpec = &NewActiveEffect.Spec;
+			FActiveGameplayEffect* AppliedEffect = ActiveGameplayEffects.ApplyGameplayEffectSpec(Spec, PredictionKey);
+			if (!AppliedEffect)
+			{
+				return FActiveGameplayEffectHandle();
+			}
+
+			MyHandle = AppliedEffect->Handle;
+			OurCopyOfSpec = &(AppliedEffect->Spec);
 
 			ABILITY_VLOG(OwnerActor, Log, TEXT("Applied %s"), *OurCopyOfSpec->Def->GetFName().ToString());
 
@@ -544,14 +554,14 @@ FActiveGameplayEffectHandle UAbilitySystemComponent::ApplyGameplayEffectSpecToSe
 			StackSpec = TSharedPtr<FGameplayEffectSpec>(new FGameplayEffectSpec(Spec));
 			OurCopyOfSpec = StackSpec.Get();
 			UAbilitySystemGlobals::Get().GlobalPreGameplayEffectSpecApply(*OurCopyOfSpec, this);
-			OurCopyOfSpec->CapturedRelevantAttributes.CaptureAttributes(this, EGameplayEffectAttributeCaptureSource::Target);
+			OurCopyOfSpec->CaptureAttributeDataFromTarget(this);
 		}
 
 		// if necessary add a modifier to OurCopyOfSpec to force it to have an infinite duration
 		if (bTreatAsInfiniteDuration)
 		{
 			// This should just be a straight set of the duration float now
-			OurCopyOfSpec->SetDuration(UGameplayEffect::INFINITE_DURATION);
+			OurCopyOfSpec->SetDuration(UGameplayEffect::INFINITE_DURATION, true);
 		}
 	}
 	
@@ -601,7 +611,8 @@ FActiveGameplayEffectHandle UAbilitySystemComponent::ApplyGameplayEffectSpecToSe
 	// ------------------------------------------------------
 	if (IsOwnerActorAuthoritative() && Spec.Def->RemoveGameplayEffectsWithTags.CombinedTags.Num() > 0)
 	{
-		ActiveGameplayEffects.RemoveActiveEffects(FActiveGameplayEffectQuery(&Spec.Def->RemoveGameplayEffectsWithTags.CombinedTags));
+		// Clear tags is always removing all stacks.
+		ActiveGameplayEffects.RemoveActiveEffects(FActiveGameplayEffectQuery(&Spec.Def->RemoveGameplayEffectsWithTags.CombinedTags), -1);
 	}
 
 	// todo: this is ignoring the returned handles, should we put them into a TArray and return all of the handles?
@@ -619,7 +630,6 @@ FActiveGameplayEffectHandle UAbilitySystemComponent::ApplyGameplayEffectSpecToSe
 	OnGameplayEffectAppliedToSelf(InstigatorASC, *OurCopyOfSpec, MyHandle);
 
 	// Send the instigator a callback
-	
 	if (InstigatorASC)
 	{
 		InstigatorASC->OnGameplayEffectAppliedToTarget(this, *OurCopyOfSpec, MyHandle);
@@ -669,7 +679,7 @@ void UAbilitySystemComponent::CheckDurationExpired(FActiveGameplayEffectHandle H
 	ActiveGameplayEffects.CheckDuration(Handle);
 }
 
-bool UAbilitySystemComponent::RemoveActiveGameplayEffect(FActiveGameplayEffectHandle Handle)
+bool UAbilitySystemComponent::RemoveActiveGameplayEffect(FActiveGameplayEffectHandle Handle, int32 StacksToRemove)
 {
 	FActiveGameplayEffect* GameplayEffect = ActiveGameplayEffects.GetActiveGameplayEffect(Handle);
 	if (GameplayEffect)
@@ -684,7 +694,7 @@ bool UAbilitySystemComponent::RemoveActiveGameplayEffect(FActiveGameplayEffectHa
 		}
 	}
 
-	return ActiveGameplayEffects.RemoveActiveGameplayEffect(Handle);
+	return ActiveGameplayEffects.RemoveActiveGameplayEffect(Handle, StacksToRemove);
 }
 
 float UAbilitySystemComponent::GetGameplayEffectDuration(FActiveGameplayEffectHandle Handle) const
@@ -897,9 +907,9 @@ void UAbilitySystemComponent::RemoveActiveEffectsWithTags(const FGameplayTagCont
 	RemoveActiveEffects(FActiveGameplayEffectQuery(&Tags));
 }
 
-void UAbilitySystemComponent::RemoveActiveEffects(const FActiveGameplayEffectQuery Query)
+void UAbilitySystemComponent::RemoveActiveEffects(const FActiveGameplayEffectQuery Query, int32 StacksToRemove)
 {
-	ActiveGameplayEffects.RemoveActiveEffects(Query);
+	ActiveGameplayEffects.RemoveActiveEffects(Query, StacksToRemove);
 }
 
 // ---------------------------------------------------------------------------------------
@@ -1052,14 +1062,15 @@ void UAbilitySystemComponent::OnAttributeAggregatorDirty(FAggregator* Aggregator
 	ActiveGameplayEffects.OnAttributeAggregatorDirty(Aggregator, Attribute);
 }
 
-void UAbilitySystemComponent::OnMagnitudeDependancyChange(FActiveGameplayEffectHandle Handle, const FAggregator* ChangedAggregator)
+void UAbilitySystemComponent::OnMagnitudeDependencyChange(FActiveGameplayEffectHandle Handle, const FAggregator* ChangedAggregator)
 {
-	ActiveGameplayEffects.OnMagnitudeDependancyChange(Handle, ChangedAggregator);
+	ActiveGameplayEffects.OnMagnitudeDependencyChange(Handle, ChangedAggregator);
 }
 
 void UAbilitySystemComponent::OnGameplayEffectAppliedToTarget(UAbilitySystemComponent* Target, const FGameplayEffectSpec& SpecApplied, FActiveGameplayEffectHandle ActiveHandle)
 {
 	OnGameplayEffectAppliedDelegateToTarget.Broadcast(Target, SpecApplied, ActiveHandle);
+	ActiveGameplayEffects.ApplyStackingLogicPostApplyAsSource(Target, SpecApplied, ActiveHandle);
 }
 
 void UAbilitySystemComponent::OnGameplayEffectAppliedToSelf(UAbilitySystemComponent* Source, const FGameplayEffectSpec& SpecApplied, FActiveGameplayEffectHandle ActiveHandle)
@@ -1162,9 +1173,23 @@ void UAbilitySystemComponent::DisplayDebug(class UCanvas* Canvas, const class FD
 				DurationStr += FString::Printf(TEXT("Period: %.2f"), ActiveGE.GetPeriod());
 			}
 
+			FString StackString;
+			if (ActiveGE.Spec.StackCount > 1)
+			{
+				
+				if (ActiveGE.Spec.Def->StackingType == EGameplayEffectStackingType::AggregateBySource)
+				{
+					StackString = FString::Printf(TEXT("(Stacks: %d. From: %s) "), ActiveGE.Spec.StackCount, *GetNameSafe(ActiveGE.Spec.GetContext().GetInstigatorAbilitySystemComponent()->AvatarActor));
+				}
+				else
+				{
+					StackString = FString::Printf(TEXT("(Stacks: %d) "), ActiveGE.Spec.StackCount);
+				}
+			}
+
 			Canvas->SetDrawColor(ActiveGE.IsInhibited ? FColor(128, 128, 128): FColor::White );
 
-			YPos += Canvas->DrawText(GEngine->GetTinyFont(), FString::Printf(TEXT("%s %s"), *ASC_CleanupName(GetNameSafe(ActiveGE.Spec.Def)), *DurationStr ), 4.f, YPos);
+			YPos += Canvas->DrawText(GEngine->GetTinyFont(), FString::Printf(TEXT("%s %s %s"), *ASC_CleanupName(GetNameSafe(ActiveGE.Spec.Def)), *DurationStr, *StackString ), 4.f, YPos);
 
 			FGameplayTagContainer GrantedTags;
 			ActiveGE.Spec.GetAllGrantedTags(GrantedTags);
diff --git a/Engine/Source/Runtime/GameplayAbilities/Private/AbilitySystemComponent_Abilities.cpp b/Engine/Source/Runtime/GameplayAbilities/Private/AbilitySystemComponent_Abilities.cpp
--- a/Engine/Source/Runtime/GameplayAbilities/Private/AbilitySystemComponent_Abilities.cpp
+++ b/Engine/Source/Runtime/GameplayAbilities/Private/AbilitySystemComponent_Abilities.cpp
@@ -157,7 +157,10 @@ void UAbilitySystemComponent::InitAbilityActorInfo(AActor* InOwnerActor, AActor*
 	{
 		for (FGameplayAbilitySpec& Spec : ActivatableAbilities.Items)
 		{
-			Spec.Ability->OnAvatarSet(AbilityActorInfo.Get(), Spec);
+			if (Spec.Ability)
+			{
+				Spec.Ability->OnAvatarSet(AbilityActorInfo.Get(), Spec);
+			}
 		}
 	}
 }
diff --git a/Engine/Source/Runtime/GameplayAbilities/Private/AbilitySystemStats.cpp b/Engine/Source/Runtime/GameplayAbilities/Private/AbilitySystemStats.cpp
--- a/Engine/Source/Runtime/GameplayAbilities/Private/AbilitySystemStats.cpp
+++ b/Engine/Source/Runtime/GameplayAbilities/Private/AbilitySystemStats.cpp
@@ -12,7 +12,7 @@ DEFINE_STAT(STAT_GameplayEffectsGetActiveEffectsDuration);
 DEFINE_STAT(STAT_GameplayAbilityGetCooldownTimeRemaining);
 DEFINE_STAT(STAT_GameplayAbilityGetCooldownTimeRemainingAndDuration);
 DEFINE_STAT(STAT_RemoveActiveGameplayEffect);
-DEFINE_STAT(STAT_CreateNewActiveGameplayEffect);
+DEFINE_STAT(STAT_ApplyGameplayEffectSpec);
 DEFINE_STAT(STAT_GetGameplayCueFunction);
 DEFINE_STAT(STAT_GetOutgoingSpec);
 DEFINE_STAT(STAT_InitAttributeSetDefaults);
diff --git a/Engine/Source/Runtime/GameplayAbilities/Private/AttributeSet.cpp b/Engine/Source/Runtime/GameplayAbilities/Private/AttributeSet.cpp
--- a/Engine/Source/Runtime/GameplayAbilities/Private/AttributeSet.cpp
+++ b/Engine/Source/Runtime/GameplayAbilities/Private/AttributeSet.cpp
@@ -351,4 +351,6 @@ void FAttributeSetInitter::InitAttributeSetDefaults(UAbilitySystemComponent* Abi
 			}
 		}		
 	}
+	
+	AbilitySystemComponent->ForceReplication();
 }
diff --git a/Engine/Source/Runtime/GameplayAbilities/Private/GameplayCueNotify_Actor.cpp b/Engine/Source/Runtime/GameplayAbilities/Private/GameplayCueNotify_Actor.cpp
--- a/Engine/Source/Runtime/GameplayAbilities/Private/GameplayCueNotify_Actor.cpp
+++ b/Engine/Source/Runtime/GameplayAbilities/Private/GameplayCueNotify_Actor.cpp
@@ -71,6 +71,10 @@ void AGameplayCueNotify_Actor::HandleGameplayCue(AActor* MyTarget, EGameplayCueE
 			OnActive(MyTarget, Parameters);
 			break;
 
+		case EGameplayCueEvent::WhileActive:
+			WhileActive(MyTarget, Parameters);
+			break;
+
 		case EGameplayCueEvent::Executed:
 			OnExecute(MyTarget, Parameters);
 			break;
@@ -103,6 +107,11 @@ bool AGameplayCueNotify_Actor::OnActive_Implementation(AActor* MyTarget, FGamepl
 	return false;
 }
 
+bool AGameplayCueNotify_Actor::WhileActive_Implementation(AActor* MyTarget, FGameplayCueParameters Parameters)
+{
+	return false;
+}
+
 bool AGameplayCueNotify_Actor::OnRemove_Implementation(AActor* MyTarget, FGameplayCueParameters Parameters)
 {
 	return false;
diff --git a/Engine/Source/Runtime/GameplayAbilities/Private/GameplayCueNotify_Static.cpp b/Engine/Source/Runtime/GameplayAbilities/Private/GameplayCueNotify_Static.cpp
--- a/Engine/Source/Runtime/GameplayAbilities/Private/GameplayCueNotify_Static.cpp
+++ b/Engine/Source/Runtime/GameplayAbilities/Private/GameplayCueNotify_Static.cpp
@@ -70,6 +70,10 @@ void UGameplayCueNotify_Static::HandleGameplayCue(AActor* MyTarget, EGameplayCue
 			OnActive(MyTarget, Parameters);
 			break;
 
+		case EGameplayCueEvent::WhileActive:
+			WhileActive(MyTarget, Parameters);
+			break;
+
 		case EGameplayCueEvent::Executed:
 			OnExecute(MyTarget, Parameters);
 			break;
@@ -99,6 +103,11 @@ bool UGameplayCueNotify_Static::OnActive_Implementation(AActor* MyTarget, FGamep
 	return false;
 }
 
+bool UGameplayCueNotify_Static::WhileActive_Implementation(AActor* MyTarget, FGameplayCueParameters Parameters) const
+{
+	return false;
+}
+
 bool UGameplayCueNotify_Static::OnRemove_Implementation(AActor* MyTarget, FGameplayCueParameters Parameters) const
 {
 	return false;
diff --git a/Engine/Source/Runtime/GameplayAbilities/Private/GameplayEffect.cpp b/Engine/Source/Runtime/GameplayAbilities/Private/GameplayEffect.cpp
--- a/Engine/Source/Runtime/GameplayAbilities/Private/GameplayEffect.cpp
+++ b/Engine/Source/Runtime/GameplayAbilities/Private/GameplayEffect.cpp
@@ -26,8 +26,13 @@ UGameplayEffect::UGameplayEffect(const FObjectInitializer& ObjectInitializer)
 	: Super(ObjectInitializer)
 {
 	DurationPolicy = EGameplayEffectDurationType::Instant;
+	bExecutePeriodicEffectOnApplication = true;
 	ChanceToApplyToTarget.SetValue(1.f);
 	ChanceToExecuteOnGameplayEffect.SetValue(1.f);
+	StackingType = EGameplayEffectStackingType::None;
+	StackLimitCount = 0;
+	StackDurationRefreshPolicy = EGameplayEffectStackingDurationPolicy::RefreshOnSuccessfulApplication;
+	StackPeriodResetPolicy = EGameplayEffectStackingPeriodPolicy::ResetOnSuccessfulApplication;
 
 #if WITH_EDITORONLY_DATA
 	ShowAllProperties = true;
@@ -272,7 +277,7 @@ bool FGameplayEffectModifierMagnitude::AttemptCalculateMagnitude(const FGameplay
 	return bCanCalc;
 }
 
-bool FGameplayEffectModifierMagnitude::AttempteRecalculateMagnitudeFromDependantChange(const FGameplayEffectSpec& InRelevantSpec, OUT float& OutCalculatedMagnitude, const FAggregator* ChangedAggregator) const
+bool FGameplayEffectModifierMagnitude::AttemptRecalculateMagnitudeFromDependentChange(const FGameplayEffectSpec& InRelevantSpec, OUT float& OutCalculatedMagnitude, const FAggregator* ChangedAggregator) const
 {
 	TArray<FGameplayEffectAttributeCaptureDefinition> ReqCaptureDefs;
 	GetAttributeCaptureDefinitions(ReqCaptureDefs);
@@ -379,6 +384,10 @@ void FGameplayEffectExecutionDefinition::GetAttributeCaptureDefinitions(OUT TArr
 
 FGameplayEffectSpec::FGameplayEffectSpec(const UGameplayEffect* InDef, const FGameplayEffectContextHandle& InEffectContext, float Level)
 	: Def(InDef)
+	, StackCount(1)
+	, bCompletedSourceAttributeCapture(false)
+	, bCompletedTargetAttributeCapture(false)
+	, bDurationLocked(false)
 {
 	check(Def);	
 	SetLevel(Level);
@@ -453,6 +462,12 @@ void FGameplayEffectSpec::SetupAttributeCaptureDefinitions()
 	}
 }
 
+void FGameplayEffectSpec::CaptureAttributeDataFromTarget(UAbilitySystemComponent* TargetAbilitySystemComponent)
+{
+	CapturedRelevantAttributes.CaptureAttributes(TargetAbilitySystemComponent, EGameplayEffectAttributeCaptureSource::Target);
+	bCompletedTargetAttributeCapture = true;
+}
+
 void FGameplayEffectSpec::CaptureDataFromSource()
 {
 	// Capture source actor tags
@@ -464,29 +479,36 @@ void FGameplayEffectSpec::CaptureDataFromSource()
 	CapturedRelevantAttributes.CaptureAttributes(EffectContext.GetInstigatorAbilitySystemComponent(), EGameplayEffectAttributeCaptureSource::Source);
 
 	// Now that we have source attributes captured, re-evaluate the duration since it could be based on the captured attributes.
-	RecalculateDuration();
+	float DefCalcDuration = 0.f;
+	if (AttemptCalculateDurationFromDef(DefCalcDuration))
+	{
+		SetDuration(DefCalcDuration, false);
+	}
+
+	bCompletedSourceAttributeCapture = true;
 }
 
-void FGameplayEffectSpec::RecalculateDuration()
+bool FGameplayEffectSpec::AttemptCalculateDurationFromDef(OUT float& OutDefDuration) const
 {
 	check(Def);
 
-	if (Def->DurationPolicy == EGameplayEffectDurationType::Infinite)
+	bool bCalculatedDuration = true;
+
+	const EGameplayEffectDurationType DurType = Def->DurationPolicy;
+	if (DurType == EGameplayEffectDurationType::Infinite)
 	{
-		SetDuration(UGameplayEffect::INFINITE_DURATION);
+		OutDefDuration = UGameplayEffect::INFINITE_DURATION;
 	}
-	else if (Def->DurationPolicy == EGameplayEffectDurationType::Instant)
+	else if (DurType == EGameplayEffectDurationType::Instant)
 	{
-		SetDuration(UGameplayEffect::INSTANT_APPLICATION);
+		OutDefDuration = UGameplayEffect::INSTANT_APPLICATION;
 	}
 	else
 	{
-		float NewDuration = 0.f;
-		if (Def->DurationMagnitude.AttemptCalculateMagnitude(*this, NewDuration))
-		{
-			SetDuration(NewDuration);
-		}
+		bCalculatedDuration = Def->DurationMagnitude.AttemptCalculateMagnitude(*this, OutDefDuration);
 	}
+
+	return bCalculatedDuration;
 }
 
 void FGameplayEffectSpec::SetLevel(float InLevel)
@@ -494,7 +516,11 @@ void FGameplayEffectSpec::SetLevel(float InLevel)
 	Level = InLevel;
 	if (Def)
 	{
-		RecalculateDuration();
+		float DefCalcDuration = 0.f;
+		if (AttemptCalculateDurationFromDef(DefCalcDuration))
+		{
+			SetDuration(DefCalcDuration, false);
+		}
 		
 		Period = Def->Period.GetValueAtLevel(InLevel);
 		ChanceToApplyToTarget = Def->ChanceToApplyToTarget.GetValueAtLevel(InLevel);
@@ -512,14 +538,18 @@ float FGameplayEffectSpec::GetDuration() const
 	return Duration;
 }
 
-void FGameplayEffectSpec::SetDuration(float NewDuration)
+void FGameplayEffectSpec::SetDuration(float NewDuration, bool bLockDuration)
 {
-	Duration = NewDuration;
-	if (Duration > 0.f)
+	if (!bDurationLocked)
 	{
-		// We may  have potential problems one day if a game is applying duration based gameplay effects from instantaneous effects
-		// (E.g., everytime fire damage is applied, a DOT is also applied). We may need to for Duration to always be capturd.
-		CapturedRelevantAttributes.AddCaptureDefinition(UAbilitySystemComponent::GetOutgoingDurationCapture());
+		Duration = NewDuration;
+		bDurationLocked = bLockDuration;
+		if (Duration > 0.f)
+		{
+			// We may have potential problems one day if a game is applying duration based gameplay effects from instantaneous effects
+			// (E.g., every time fire damage is applied, a DOT is also applied). We may need to for Duration to always be captured.
+			CapturedRelevantAttributes.AddCaptureDefinition(UAbilitySystemComponent::GetOutgoingDurationCapture());
+		}
 	}
 }
 
@@ -774,11 +804,20 @@ void FGameplayEffectAttributeCaptureSpec::RegisterLinkedAggregatorCallback(FActi
 		FAggregator* Agg = AttributeAggregator.Get();
 		if (Agg)
 		{
-			Agg->AddDependant(Handle);
+			Agg->AddDependent(Handle);
 		}
 	}
 }
 
+void FGameplayEffectAttributeCaptureSpec::UnregisterLinkedAggregatorCallback(FActiveGameplayEffectHandle Handle) const
+{
+	FAggregator* Agg = AttributeAggregator.Get();
+	if (Agg)
+	{
+		Agg->RemoveDependent(Handle);
+	}
+}
+
 bool FGameplayEffectAttributeCaptureSpec::ShouldRefreshLinkedAggregator(const FAggregator* ChangedAggregator) const
 {
 	return (BackingDefinition.bSnapshot == false && (ChangedAggregator == nullptr || AttributeAggregator.Get() == ChangedAggregator));
@@ -880,13 +919,25 @@ void FGameplayEffectAttributeCaptureSpecContainer::RegisterLinkedAggregatorCallb
 	}
 }
 
+void FGameplayEffectAttributeCaptureSpecContainer::UnregisterLinkedAggregatorCallbacks(FActiveGameplayEffectHandle Handle) const
+{
+	for (const FGameplayEffectAttributeCaptureSpec& CaptureSpec : SourceAttributes)
+	{
+		CaptureSpec.UnregisterLinkedAggregatorCallback(Handle);
+	}
+
+	for (const FGameplayEffectAttributeCaptureSpec& CaptureSpec : TargetAttributes)
+	{
+		CaptureSpec.UnregisterLinkedAggregatorCallback(Handle);
+	}
+}
+
 // --------------------------------------------------------------------------------------------------------------------------------------------------------
 //
 //	FActiveGameplayEffect
 //
 // --------------------------------------------------------------------------------------------------------------------------------------------------------
 
-
 /** This is the core function that turns the ActiveGE 'on' or 'off */
 void FActiveGameplayEffect::CheckOngoingTagRequirements(const FGameplayTagContainer& OwnerTags, FActiveGameplayEffectsContainer& OwningContainer)
 {
@@ -969,6 +1020,7 @@ void FActiveGameplayEffect::PostReplicatedAdd(const struct FActiveGameplayEffect
 
 	// Set our local start time accordingly
 	StartWorldTime = WorldTimeSeconds - static_cast<float>(DeltaGameStateTime);
+	CachedStartGameStateTime = StartGameStateTime;
 }
 
 void FActiveGameplayEffect::PostReplicatedChange(const struct FActiveGameplayEffectsContainer &InArray)
@@ -977,6 +1029,14 @@ void FActiveGameplayEffect::PostReplicatedChange(const struct FActiveGameplayEff
 	{
 		ABILITY_LOG(Error, TEXT("Received ReplicatedGameplayEffect with no UGameplayEffect def."));
 	}
+
+	// Handle potential duration refresh
+	// @todo: Due to precision of gamestate timer, this could be incorrect by just under an entire second; Need more precise replicated timer
+	if (CachedStartGameStateTime != StartGameStateTime)
+	{
+		StartWorldTime = InArray.GetWorldTime() - static_cast<float>(InArray.GetGameStateTime() - StartGameStateTime);
+		CachedStartGameStateTime = StartGameStateTime;
+	}
 }
 
 // --------------------------------------------------------------------------------------------------------------------------------------------------------
@@ -1217,7 +1277,7 @@ void FActiveGameplayEffectsContainer::OnAttributeAggregatorDirty(FAggregator* Ag
 	InternalUpdateNumericalAttribute(Attribute, NewValue, nullptr);
 }
 
-void FActiveGameplayEffectsContainer::OnMagnitudeDependancyChange(FActiveGameplayEffectHandle Handle, const FAggregator* ChangedAgg)
+void FActiveGameplayEffectsContainer::OnMagnitudeDependencyChange(FActiveGameplayEffectHandle Handle, const FAggregator* ChangedAgg)
 {
 	if (Handle.IsValid())
 	{
@@ -1229,7 +1289,7 @@ void FActiveGameplayEffectsContainer::OnMagnitudeDependancyChange(FActiveGamepla
 			FGameplayEffectSpec& Spec = ActiveEffect->Spec;
 
 			// We must update attribute aggregators only if we are actually 'on' right now, and if we are non periodic (periodic effects do their thing on execute callbacks)
-			bool MustUpdateAttributeAggregators = (ActiveEffect->IsInhibited == false && (Spec.GetDuration() <= UGameplayEffect::NO_PERIOD));
+			bool MustUpdateAttributeAggregators = (ActiveEffect->IsInhibited == false && (Spec.GetPeriod() <= UGameplayEffect::NO_PERIOD));
 
 			// As we update our modifier magnitudes, we will update our owner's attribute aggregators. When we do this, we have to clear them first of all of our (Handle's) previous mods.
 			// Since we could potentially have two mods to the same attribute, one that gets updated, and one that doesnt - we need to do this in two passes.
@@ -1241,7 +1301,7 @@ void FActiveGameplayEffectsContainer::OnMagnitudeDependancyChange(FActiveGamepla
 				const FGameplayModifierInfo& ModDef = Spec.Def->Modifiers[ModIdx];
 				FModifierSpec& ModSpec = Spec.Modifiers[ModIdx];
 
-				if (ModDef.ModifierMagnitude.AttempteRecalculateMagnitudeFromDependantChange(Spec, ModSpec.EvaluatedMagnitude, ChangedAgg))
+				if (ModDef.ModifierMagnitude.AttemptRecalculateMagnitudeFromDependentChange(Spec, ModSpec.EvaluatedMagnitude, ChangedAgg))
 				{
 					// We changed, so we need to reapply/update our spot in the attribute aggregator map
 					if (MustUpdateAttributeAggregators)
@@ -1252,28 +1312,140 @@ void FActiveGameplayEffectsContainer::OnMagnitudeDependancyChange(FActiveGamepla
 			}
 
 			// Second pass, update the aggregators that we need to
-			for (const FGameplayAttribute& Attribute : AttributesToUpdate)
+			UpdateAggregatorModMagnitudes(AttributesToUpdate, *ActiveEffect);
+		}
+	}
+}
+
+void FActiveGameplayEffectsContainer::OnStackCountChange(FActiveGameplayEffect& ActiveEffect)
+{
+	const FGameplayEffectSpec& Spec = ActiveEffect.Spec;
+	TSet<FGameplayAttribute>	AttributesToUpdate;
+	
+	for (int32 ModIdx = 0; ModIdx < Spec.Modifiers.Num(); ++ModIdx)
+	{
+		const FGameplayModifierInfo& ModDef = Spec.Def->Modifiers[ModIdx];
+		AttributesToUpdate.Add(ModDef.Attribute);
+	}
+
+	UpdateAggregatorModMagnitudes(AttributesToUpdate, ActiveEffect);
+}
+
+void FActiveGameplayEffectsContainer::UpdateAggregatorModMagnitudes(const TSet<FGameplayAttribute>& AttributesToUpdate, FActiveGameplayEffect& ActiveEffect)
+{
+	const FGameplayEffectSpec& Spec = ActiveEffect.Spec;
+	for (const FGameplayAttribute& Attribute : AttributesToUpdate)
+	{
+		FAggregator* Aggregator = FindOrCreateAttributeAggregator(Attribute).Get();
+		check(Aggregator);
+
+		// Remove ALL of our mods (this is all we can do! We don't keep track of aggregator mods <=> spec mod mappings)
+		// Todo: it would be nice to have a better way of updating this. Some ideas: instead of remove, call an 'invalidate' to NAN out all floats for a given handle, then call an 'update' to set real values.
+		Aggregator->RemoveAggregatorMod(ActiveEffect.Handle);
+
+		// Now readd ALL of our mods
+		for(int32 ModIdx = 0; ModIdx < Spec.Modifiers.Num(); ++ModIdx)
+		{
+			const FGameplayModifierInfo& ModDef = Spec.Def->Modifiers[ModIdx];
+			const FModifierSpec& ModSpec = Spec.Modifiers[ModIdx];
+
+			if (ModDef.Attribute == Attribute)
 			{
-				FAggregator* Aggregator = FindOrCreateAttributeAggregator(Attribute).Get();
-				check(Aggregator);
+				Aggregator->AddAggregatorMod(ModSpec.GetEvaluatedMagnitude() * static_cast<float>(Spec.StackCount), ModDef.ModifierOp, &ModDef.SourceTags, &ModDef.TargetTags, ActiveEffect.PredictionKey.WasLocallyGenerated(), ActiveEffect.Handle);
+			}
+		}
+	}
+}
 
-				// Remove ALL of our mods (this is all we can do! We don't keep track of aggregator mods <=> spec mod mappings)
-				Aggregator->RemoveMod(Handle);
+FActiveGameplayEffect* FActiveGameplayEffectsContainer::FindStackableActiveGameplayEffect(const FGameplayEffectSpec& Spec)
+{
+	FActiveGameplayEffect* StackableGE = nullptr;
+	const UGameplayEffect* GEDef = Spec.Def;
+	EGameplayEffectStackingType StackingType = GEDef->StackingType;
 
-				// Now readd ALL of our mods (even if they weren't recalculated here)
-				for(int32 ModIdx = 0; ModIdx < Spec.Modifiers.Num(); ++ModIdx)
-				{
-					const FGameplayModifierInfo& ModDef = Spec.Def->Modifiers[ModIdx];
-					FModifierSpec& ModSpec = Spec.Modifiers[ModIdx];
+	UAbilitySystemComponent* SourceASC = Spec.GetContext().GetInstigatorAbilitySystemComponent();
+	if (SourceASC && StackingType != EGameplayEffectStackingType::None && Spec.GetDuration() != UGameplayEffect::INSTANT_APPLICATION)
+	{
+		// Iterate through GameplayEffects to see if we find a match. Note that we could cache off a handle in a map but we would still
+		// do a linear search through GameplayEffects to find the actual FActiveGameplayEffect (due to unstable nature of the GameplayEffects array).
+		// If this becomes a slow point in the profiler, the map may still be useful as an early out to avoid an unnecessary sweep.
 
-					if (ModDef.Attribute == Attribute)
-					{
-						Aggregator->AddMod(ModSpec.GetEvaluatedMagnitude(), ModDef.ModifierOp, &ModDef.SourceTags, &ModDef.TargetTags, ActiveEffect->PredictionKey.WasLocallyGenerated(), Handle);
-					}
-				}
+		for (FActiveGameplayEffect& ActiveEffect : GameplayEffects)
+		{
+			// Aggregate by source stacking additionally requires the source ability component to match
+			if (ActiveEffect.Spec.Def == Spec.Def && ((StackingType == EGameplayEffectStackingType::AggregateByTarget) || (SourceASC == ActiveEffect.Spec.GetContext().GetInstigatorAbilitySystemComponent())))
+			{
+				StackableGE = &ActiveEffect;
+				break;
 			}
 		}
 	}
+
+	return StackableGE;
+}
+
+float FActiveGameplayEffectsContainer::ComputeModifiedDurationOfAppliedSpec(const FGameplayEffectSpec& Spec, float BaseValue) const
+{
+	FAggregator DurationAgg;
+
+	const FGameplayEffectAttributeCaptureSpec* OutgoingCaptureSpec = Spec.CapturedRelevantAttributes.FindCaptureSpecByDefinition(UAbilitySystemComponent::GetOutgoingDurationCapture(), true);
+	if (OutgoingCaptureSpec)
+	{
+		OutgoingCaptureSpec->AttemptAddAggregatorModsToAggregator(DurationAgg);
+	}
+
+	const FGameplayEffectAttributeCaptureSpec* IncomingCaptureSpec = Spec.CapturedRelevantAttributes.FindCaptureSpecByDefinition(UAbilitySystemComponent::GetIncomingDurationCapture(), true);
+	if (IncomingCaptureSpec)
+	{
+		IncomingCaptureSpec->AttemptAddAggregatorModsToAggregator(DurationAgg);
+	}
+
+	FAggregatorEvaluateParameters Params;
+	Params.SourceTags = Spec.CapturedSourceTags.GetAggregatedTags();
+	Params.TargetTags = Spec.CapturedTargetTags.GetAggregatedTags();
+
+	return DurationAgg.EvaluateWithBase(BaseValue, Params);
+}
+
+bool FActiveGameplayEffectsContainer::HandleActiveGameplayEffectStackOverflow(const FActiveGameplayEffect& ActiveStackableGE, const FGameplayEffectSpec& OldSpec, const FGameplayEffectSpec& OverflowingSpec)
+{
+	const UGameplayEffect* StackedGE = OldSpec.Def;
+	
+	const bool bAllowOverflowApplication = !(StackedGE->bDenyOverflowApplication);
+
+	FPredictionKey PredictionKey;
+	for (TSubclassOf<UGameplayEffect> OverflowEffect : StackedGE->OverflowEffects)
+	{
+		if (OverflowEffect)
+		{
+			FGameplayEffectSpec NewGESpec(OverflowEffect->GetDefaultObject<UGameplayEffect>(), OverflowingSpec.GetContext(), OverflowingSpec.GetLevel());
+			// @todo: copy over source tags
+			// @todo: scope lock
+			Owner->ApplyGameplayEffectSpecToSelf(NewGESpec, PredictionKey);
+		}
+	}
+	// @todo: Scope lock
+	if (!bAllowOverflowApplication && StackedGE->bClearStackOnOverflow)
+	{
+		Owner->RemoveActiveGameplayEffect(ActiveStackableGE.Handle);
+	}
+
+	return bAllowOverflowApplication;
+}
+
+void FActiveGameplayEffectsContainer::ApplyStackingLogicPostApplyAsSource(UAbilitySystemComponent* Target, const FGameplayEffectSpec& SpecApplied, FActiveGameplayEffectHandle ActiveHandle)
+{
+	if (SpecApplied.Def->StackingType == EGameplayEffectStackingType::AggregateBySource)
+	{
+		TArray<FActiveGameplayEffectHandle>& ActiveHandles = SourceStackingMap.FindOrAdd(SpecApplied.Def);
+
+		// This is probably wrong!
+		if (ActiveHandles.Num() == SpecApplied.Def->StackLimitCount)
+		{
+			// We are at the limit, so replace one based on policy
+			// For now, just always remove the oldest one applied
+		}
+	}
 }
 
 void FActiveGameplayEffectsContainer::SetBaseAttributeValueFromReplication(FGameplayAttribute Attribute, float ServerValue)
@@ -1467,91 +1639,153 @@ void FActiveGameplayEffectsContainer::ApplyModToAttribute(const FGameplayAttribu
 	}
 }
 
-/**	This Creates a new ActiveGameplayEffect and is only called on the Server or in cases where a client predictively applies a GameplayEffect. */
-FActiveGameplayEffect& FActiveGameplayEffectsContainer::CreateNewActiveGameplayEffect(const FGameplayEffectSpec &Spec, FPredictionKey InPredictionKey)
+FActiveGameplayEffect* FActiveGameplayEffectsContainer::ApplyGameplayEffectSpec(const FGameplayEffectSpec& Spec, FPredictionKey InPredictionKey)
 {
-	SCOPE_CYCLE_COUNTER(STAT_CreateNewActiveGameplayEffect);
-	
-	if (Owner && Owner->OwnerActor)
+	SCOPE_CYCLE_COUNTER(STAT_ApplyGameplayEffectSpec);
+
+	if (Owner && Owner->OwnerActor && IsNetAuthority())
 	{
-		if ( IsNetAuthority() )
+		Owner->OwnerActor->FlushNetDormancy();
+	}
+
+	FActiveGameplayEffect* AppliedActiveGE = nullptr;
+	FActiveGameplayEffect* ExistingStackableGE = FindStackableActiveGameplayEffect(Spec);
+
+	bool bSetDuration = true;
+	bool bSetPeriod = true;
+
+	// Check if there's an active GE this application should stack upon
+	if (ExistingStackableGE)
+	{
+		// Don't allow prediction of stacking for now
+		if (!IsNetAuthority())
 		{
-			Owner->OwnerActor->FlushNetDormancy();
+			return nullptr;
 		}
-	}
 
-	FActiveGameplayEffectHandle NewHandle = FActiveGameplayEffectHandle::GenerateNewHandle(Owner);
-	FActiveGameplayEffect& NewEffect = *new (GameplayEffects)FActiveGameplayEffect(NewHandle, Spec, GetWorldTime(), GetGameStateTime(), InPredictionKey);
+		FGameplayEffectSpec& ExistingSpec = ExistingStackableGE->Spec;
 
-	UAbilitySystemGlobals::Get().GlobalPreGameplayEffectSpecApply(NewEffect.Spec, Owner);
+		int32 NewStackCount = ExistingSpec.StackCount;
 
-	// Calc all of our modifier magnitudes now. Some may need to update later based on attributes changing, etc, but those should
-	// be done through delegate callbacks.
-	NewEffect.Spec.CapturedRelevantAttributes.CaptureAttributes(Owner, EGameplayEffectAttributeCaptureSource::Target);
-	NewEffect.Spec.CalculateModifierMagnitudes();
+		// Process overflow if already at the stack limit; Overflow can potentially reject the application outright
+		if (ExistingSpec.StackCount == ExistingSpec.Def->StackLimitCount)
+		{
+			if (!HandleActiveGameplayEffectStackOverflow(*ExistingStackableGE, ExistingSpec, Spec))
+			{
+				return nullptr;
+			}
+		}
+		else
+		{
+			++NewStackCount;
+		}
 
-	// Register Source and Target non snapshot capture delegates here
-	NewEffect.Spec.CapturedRelevantAttributes.RegisterLinkedAggregatorCallbacks(NewHandle);
+		// Need to unregister callbacks because the source aggregators could potentially be different with the new application. They will be
+		// re-registered later below, as necessary.
+		ExistingSpec.CapturedRelevantAttributes.UnregisterLinkedAggregatorCallbacks(ExistingStackableGE->Handle);
 
-	// Calculate Duration mods if we have a real duration
-	float DurationBaseValue = NewEffect.Spec.GetDuration();
-	if (DurationBaseValue > 0.f)
-	{
-		FAggregator DurationAgg;
+		// @todo: If dynamically granted tags differ (which they shouldn't), we'll actually have to diff them
+		// and cause a removal and add of only the ones that have changed. For now, ensure on this happening and come
+		// back to this later.
+		ensureMsgf(ExistingSpec.DynamicGrantedTags == Spec.DynamicGrantedTags, TEXT("While adding a stack of the gameplay effect: %s, the old stack and the new application had different dynamically granted tags, which is currently not resolved properly!"), *Spec.Def->GetName());
+
+		ExistingStackableGE->Spec = Spec;
+		ExistingStackableGE->Spec.StackCount = NewStackCount;
+		
+		AppliedActiveGE = ExistingStackableGE;
 
-		const FGameplayEffectAttributeCaptureSpec* OutgoingCaptureSpec = NewEffect.Spec.CapturedRelevantAttributes.FindCaptureSpecByDefinition(UAbilitySystemComponent::GetOutgoingDurationCapture(), true);
-		if (OutgoingCaptureSpec)
+		const UGameplayEffect* GEDef = ExistingSpec.Def;
+
+		// Make sure the GE actually wants to refresh its duration
+		if (GEDef->StackDurationRefreshPolicy == EGameplayEffectStackingDurationPolicy::NeverRefresh)
+		{
+			bSetDuration = false;
+		}
+		else
 		{
-			OutgoingCaptureSpec->AttemptAddAggregatorModsToAggregator(DurationAgg);
+			ExistingStackableGE->StartGameStateTime = GetGameStateTime();
+			ExistingStackableGE->CachedStartGameStateTime = ExistingStackableGE->StartGameStateTime;
+			ExistingStackableGE->StartWorldTime = GetWorldTime();
 		}
 
-		const FGameplayEffectAttributeCaptureSpec* IncomingCaptureSpec = NewEffect.Spec.CapturedRelevantAttributes.FindCaptureSpecByDefinition(UAbilitySystemComponent::GetIncomingDurationCapture(), true);
-		if (IncomingCaptureSpec)
+		// Make sure the GE actually wants to reset its period
+		if (GEDef->StackPeriodResetPolicy == EGameplayEffectStackingPeriodPolicy::NeverReset)
 		{
-			IncomingCaptureSpec->AttemptAddAggregatorModsToAggregator(DurationAgg);
+			bSetPeriod = false;
 		}
+	}
+	else
+	{
+		FActiveGameplayEffectHandle NewHandle = FActiveGameplayEffectHandle::GenerateNewHandle(Owner);
+		AppliedActiveGE = new(GameplayEffects) FActiveGameplayEffect(NewHandle, Spec, GetWorldTime(), GetGameStateTime(), InPredictionKey);
+	}
+
+	FGameplayEffectSpec& AppliedEffectSpec = AppliedActiveGE->Spec;
+	UAbilitySystemGlobals::Get().GlobalPreGameplayEffectSpecApply(AppliedEffectSpec, Owner);
+
+	// Calc all of our modifier magnitudes now. Some may need to update later based on attributes changing, etc, but those should
+	// be done through delegate callbacks.
+	AppliedEffectSpec.CaptureAttributeDataFromTarget(Owner);
+	AppliedEffectSpec.CalculateModifierMagnitudes();
 
-		FAggregatorEvaluateParameters Params;
-		Params.SourceTags = NewEffect.Spec.CapturedSourceTags.GetAggregatedTags();
-		Params.TargetTags = NewEffect.Spec.CapturedTargetTags.GetAggregatedTags();
+	// Register Source and Target non snapshot capture delegates here
+	AppliedEffectSpec.CapturedRelevantAttributes.RegisterLinkedAggregatorCallbacks(AppliedActiveGE->Handle);
+	
+	// Calculate Duration mods if we have a real duration
+	if (AppliedEffectSpec.Def->DurationPolicy == EGameplayEffectDurationType::HasDuration && bSetDuration)
+	{
+		// Re-calculate the duration, as it could rely on target captured attributes
+		float DefCalcDuration = 0.f;
+		if (AppliedEffectSpec.AttemptCalculateDurationFromDef(DefCalcDuration))
+		{
+			AppliedEffectSpec.SetDuration(DefCalcDuration, false);
+		}
 
-		float FinalDuration = DurationAgg.EvaluateWithBase(DurationBaseValue, Params);
+		const float DurationBaseValue = AppliedEffectSpec.GetDuration();
+
+		float FinalDuration = ComputeModifiedDurationOfAppliedSpec(AppliedEffectSpec, DurationBaseValue);
 
 		// We cannot mod ourselves into an instant or infinite duration effect
 		if (FinalDuration <= 0.f)
 		{
-			ABILITY_LOG(Error, TEXT("GameplayEffect %s Duration was modified to %.2f. Clamping to 0.1s duration."), *NewEffect.Spec.Def->GetName(), FinalDuration);
+			ABILITY_LOG(Error, TEXT("GameplayEffect %s Duration was modified to %.2f. Clamping to 0.1s duration."), *AppliedEffectSpec.Def->GetName(), FinalDuration);
 			FinalDuration = 0.1f;
 		}
 
-		NewEffect.Spec.SetDuration(FinalDuration);
+		AppliedEffectSpec.SetDuration(FinalDuration, true);
 
 		// ABILITY_LOG(Warning, TEXT("SetDuration for %s. Base: %.2f, Final: %.2f"), *NewEffect.Spec.Def->GetName(), DurationBaseValue, FinalDuration);
-	}
-	
-	// register callbacks with the timer manager
-	if (Owner)
-	{
-		if (NewEffect.Spec.GetDuration() > 0.f)
+		
+		// Register duration callbacks with the timer manager
+		if (Owner)
 		{
 			FTimerManager& TimerManager = Owner->GetWorld()->GetTimerManager();
-			FTimerDelegate Delegate = FTimerDelegate::CreateUObject(Owner, &UAbilitySystemComponent::CheckDurationExpired, NewHandle);
-			TimerManager.SetTimer(NewEffect.DurationHandle, Delegate, NewEffect.Spec.GetDuration(), false, NewEffect.Spec.GetDuration());
+			FTimerDelegate Delegate = FTimerDelegate::CreateUObject(Owner, &UAbilitySystemComponent::CheckDurationExpired, AppliedActiveGE->Handle);
+			TimerManager.SetTimer(AppliedActiveGE->DurationHandle, Delegate, FinalDuration, false);
 		}
+	}
+	
+	// Register period callbacks with the timer manager
+	if (Owner && (AppliedEffectSpec.GetPeriod() != UGameplayEffect::NO_PERIOD))
+	{
+		FTimerManager& TimerManager = Owner->GetWorld()->GetTimerManager();
+		FTimerDelegate Delegate = FTimerDelegate::CreateUObject(Owner, &UAbilitySystemComponent::ExecutePeriodicEffect, AppliedActiveGE->Handle);
+			
 		// The timer manager moves things from the pending list to the active list after checking the active list on the first tick so we need to execute here
-		if (NewEffect.Spec.GetPeriod() != UGameplayEffect::NO_PERIOD)
+		if (AppliedEffectSpec.Def->bExecutePeriodicEffectOnApplication)
 		{
-			FTimerManager& TimerManager = Owner->GetWorld()->GetTimerManager();
-			FTimerDelegate Delegate = FTimerDelegate::CreateUObject(Owner, &UAbilitySystemComponent::ExecutePeriodicEffect, NewHandle);
-			
 			TimerManager.SetTimerForNextTick(Delegate);
-			TimerManager.SetTimer(NewEffect.PeriodHandle, Delegate, NewEffect.Spec.GetPeriod(), true, -1.f);
+		}
+
+		if (bSetPeriod)
+		{
+			TimerManager.SetTimer(AppliedActiveGE->PeriodHandle, Delegate, AppliedEffectSpec.GetPeriod(), true);
 		}
 	}
-	
+
 	if (InPredictionKey.IsValidKey() == false || IsNetAuthority())	// Clients predicting a GameplayEffect must not call MarkItemDirty
 	{
-		MarkItemDirty(NewEffect);
+		MarkItemDirty(*AppliedActiveGE);
 	}
 	else
 	{
@@ -1559,13 +1793,22 @@ FActiveGameplayEffect& FActiveGameplayEffectsContainer::CreateNewActiveGameplayE
 		MarkArrayDirty();
 
 		// Once replicated state has caught up to this prediction key, we must remove this gameplay effect.
-		InPredictionKey.NewRejectOrCaughtUpDelegate(FPredictionKeyEvent::CreateUObject(Owner, &UAbilitySystemComponent::RemoveActiveGameplayEffect_NoReturn, NewEffect.Handle));
+		InPredictionKey.NewRejectOrCaughtUpDelegate(FPredictionKeyEvent::CreateUObject(Owner, &UAbilitySystemComponent::RemoveActiveGameplayEffect_NoReturn, AppliedActiveGE->Handle, -1));
 		
 	}
 
-	InternalOnActiveGameplayEffectAdded(NewEffect);
+	// @note @todo: This is currently assuming (potentially incorrectly) that the inhibition state of a stacked GE won't change
+	// as a result of stacking. In reality it could in complicated cases with differing sets of dynamically-granted tags.
+	if (ExistingStackableGE)
+	{
+		OnStackCountChange(*ExistingStackableGE);
+	}
+	else
+	{
+		InternalOnActiveGameplayEffectAdded(*AppliedActiveGE);
+	}
 
-	return NewEffect;
+	return AppliedActiveGE;
 }
 
 /** This is called anytime a new ActiveGameplayEffect is added, on both client and server in all cases */
@@ -1604,7 +1847,7 @@ void FActiveGameplayEffectsContainer::AddActiveGameplayEffectGrantedTagsAndModif
 			// Ongoing tags being met. We either calculate magnitude one time, or its done via OnDirty calls (or potentially a frequency timer one day)
 					
 			FAggregator* Aggregator = FindOrCreateAttributeAggregator(Effect.Spec.Def->Modifiers[ModIdx].Attribute).Get();
-			Aggregator->AddMod(Mod.GetEvaluatedMagnitude(), ModInfo.ModifierOp, &ModInfo.SourceTags, &ModInfo.TargetTags, Effect.PredictionKey.WasLocallyGenerated(), Effect.Handle);
+			Aggregator->AddAggregatorMod(Mod.GetEvaluatedMagnitude() * static_cast<float>(Effect.Spec.StackCount), ModInfo.ModifierOp, &ModInfo.SourceTags, &ModInfo.TargetTags, Effect.PredictionKey.WasLocallyGenerated(), Effect.Handle);
 		}
 	}
 
@@ -1623,7 +1866,7 @@ void FActiveGameplayEffectsContainer::AddActiveGameplayEffectGrantedTagsAndModif
 }
 
 /** Called on server to remove a GameplayEffect */
-bool FActiveGameplayEffectsContainer::RemoveActiveGameplayEffect(FActiveGameplayEffectHandle Handle)
+bool FActiveGameplayEffectsContainer::RemoveActiveGameplayEffect(FActiveGameplayEffectHandle Handle, int32 StacksToRemove)
 {
 	// Could make this a map for quicker lookup
 	for (int32 Idx = 0; Idx < GameplayEffects.Num(); ++Idx)
@@ -1633,10 +1876,11 @@ bool FActiveGameplayEffectsContainer::RemoveActiveGameplayEffect(FActiveGameplay
 			//Block removal if we're scope-locked
 			if (FScopedActiveGameplayEffectLock::IsLockInEffect())
 			{
-				FScopedActiveGameplayEffectLock::AddAction()->InitForRemoveGE(TWeakObjectPtr<UAbilitySystemComponent>(Owner), Handle);
+				FScopedActiveGameplayEffectLock::AddAction()->InitForRemoveGE(TWeakObjectPtr<UAbilitySystemComponent>(Owner), Handle, StacksToRemove);
 				return true;		//We are assuming the "ensure(Idx < GameplayEffects.Num())" would be passed.
 			}
-			return InternalRemoveActiveGameplayEffect(Idx);
+			InternalRemoveActiveGameplayEffect(Idx, StacksToRemove);
+			return true;
 		}
 	}
 	ABILITY_LOG(Warning, TEXT("RemoveActiveGameplayEffect called with invalid Handle: %s"), *Handle.ToString());
@@ -1644,15 +1888,23 @@ bool FActiveGameplayEffectsContainer::RemoveActiveGameplayEffect(FActiveGameplay
 }
 
 /** Called by server to actually remove a GameplayEffect */
-bool FActiveGameplayEffectsContainer::InternalRemoveActiveGameplayEffect(int32 Idx)
+bool FActiveGameplayEffectsContainer::InternalRemoveActiveGameplayEffect(int32 Idx, int32 StacksToRemove)
 {
 	SCOPE_CYCLE_COUNTER(STAT_RemoveActiveGameplayEffect);
 
+	check(!FScopedActiveGameplayEffectLock::IsLockInEffect());
+	
 	if (ensure(Idx < GameplayEffects.Num()))
 	{
 		FActiveGameplayEffect& Effect = GameplayEffects[Idx];
 
-		check(!FScopedActiveGameplayEffectLock::IsLockInEffect());
+		if (StacksToRemove > 0 && Effect.Spec.StackCount > StacksToRemove)
+		{
+			// This won't be a full remove, only a change in StackCount.
+			Effect.Spec.StackCount -= StacksToRemove;
+			OnStackCountChange(Effect);
+			return false;
+		}
 
 		bool ShouldInvokeGameplayCueEvent = true;
 		const bool bIsNetAuthority = IsNetAuthority();
@@ -1709,8 +1961,8 @@ bool FActiveGameplayEffectsContainer::InternalRemoveActiveGameplayEffect(int32 I
 void FActiveGameplayEffectsContainer::InternalOnActiveGameplayEffectRemoved(const FActiveGameplayEffect& Effect)
 {
 	// Remove our tag requirements from the dependancy map
-	RemoveActiveEffectTagDependancy(Effect.Spec.Def->OngoingTagRequirements.IgnoreTags, Effect.Handle);
-	RemoveActiveEffectTagDependancy(Effect.Spec.Def->OngoingTagRequirements.RequireTags, Effect.Handle);
+	RemoveActiveEffectTagDependency(Effect.Spec.Def->OngoingTagRequirements.IgnoreTags, Effect.Handle);
+	RemoveActiveEffectTagDependency(Effect.Spec.Def->OngoingTagRequirements.RequireTags, Effect.Handle);
 
 	if (Effect.Spec.Def)
 	{
@@ -1736,7 +1988,7 @@ void FActiveGameplayEffectsContainer::RemoveActiveGameplayEffectGrantedTagsAndMo
 				FAggregatorRef* RefPtr = AttributeAggregatorMap.Find(Mod.Attribute);
 				if(RefPtr)
 				{
-					RefPtr->Get()->RemoveMod(Effect.Handle);
+					RefPtr->Get()->RemoveAggregatorMod(Effect.Handle);
 				}
 			}
 		}
@@ -1757,7 +2009,7 @@ void FActiveGameplayEffectsContainer::RemoveActiveGameplayEffectGrantedTagsAndMo
 	}
 }
 
-void FActiveGameplayEffectsContainer::RemoveActiveEffectTagDependancy(const FGameplayTagContainer& Tags, FActiveGameplayEffectHandle Handle)
+void FActiveGameplayEffectsContainer::RemoveActiveEffectTagDependency(const FGameplayTagContainer& Tags, FActiveGameplayEffectHandle Handle)
 {
 	for(const FGameplayTag& Tag : Tags)
 	{
@@ -1775,7 +2027,7 @@ void FActiveGameplayEffectsContainer::RemoveActiveEffectTagDependancy(const FGam
 
 void FActiveGameplayEffectsContainer::OnOwnerTagChange(FGameplayTag TagChange, int32 NewCount)
 {
-	// It may be beneficial to do a scoped lock on attribute reeevalluation during this dunction
+	// It may be beneficial to do a scoped lock on attribute re-evaluation during this function
 
 	auto Ptr = ActiveEffectTagDependencies.Find(TagChange);
 	if (Ptr)
@@ -1880,7 +2132,7 @@ void FActiveGameplayEffectsContainer::CheckDuration(FActiveGameplayEffectHandle
 					TimerManager.ClearTimer(Effect.PeriodHandle);
 				}
 
-				InternalRemoveActiveGameplayEffect(Idx);
+				InternalRemoveActiveGameplayEffect(Idx, -1);
 			}
 			else
 			{
@@ -1977,7 +2229,7 @@ TArray<float> FActiveGameplayEffectsContainer::GetActiveEffectsDuration(const FA
 	return ReturnList;
 }
 
-void FActiveGameplayEffectsContainer::RemoveActiveEffects(const FActiveGameplayEffectQuery Query)
+void FActiveGameplayEffectsContainer::RemoveActiveEffects(const FActiveGameplayEffectQuery Query, int32 StacksToRemove)
 {
 	const bool bLockInEffect = FScopedActiveGameplayEffectLock::IsLockInEffect();
 
@@ -1989,12 +2241,14 @@ void FActiveGameplayEffectsContainer::RemoveActiveEffects(const FActiveGameplayE
 			// Defer removal if we're scope-locked
 			if (bLockInEffect)
 			{
-				FScopedActiveGameplayEffectLock::AddAction()->InitForRemoveGE(TWeakObjectPtr<UAbilitySystemComponent>(Owner), Effect.Handle);
+				FScopedActiveGameplayEffectLock::AddAction()->InitForRemoveGE(TWeakObjectPtr<UAbilitySystemComponent>(Owner), Effect.Handle, StacksToRemove);
 			}
 			else
 			{
-				InternalRemoveActiveGameplayEffect(idx);
-				idx--;
+				if (InternalRemoveActiveGameplayEffect(idx, StacksToRemove))
+				{
+					idx--;
+				}
 			}
 		}
 	}
@@ -2085,6 +2339,11 @@ void FActiveGameplayEffectHandle::RemoveFromGlobalMap()
 
 bool FActiveGameplayEffectQuery::Matches(const FActiveGameplayEffect& Effect) const
 {
+	if (CustomMatch.IsBound())
+	{
+		return CustomMatch.Execute(Effect);
+	}
+
 	// if we are looking for owning tags check them on the Granted Tags and Owned Tags Container
 	if (OwningTagContainer)
 	{
@@ -2208,7 +2467,7 @@ void FActiveGameplayEffectAction::PerformAction()
 	{
 		if (bRemove)
 		{
-			OwningASC.Get()->RemoveActiveGameplayEffect(Handle);
+			OwningASC.Get()->RemoveActiveGameplayEffect(Handle, StacksToRemove);
 		}
 		else
 		{
diff --git a/Engine/Source/Runtime/GameplayAbilities/Private/GameplayEffectAggregator.cpp b/Engine/Source/Runtime/GameplayAbilities/Private/GameplayEffectAggregator.cpp
--- a/Engine/Source/Runtime/GameplayAbilities/Private/GameplayEffectAggregator.cpp
+++ b/Engine/Source/Runtime/GameplayAbilities/Private/GameplayEffectAggregator.cpp
@@ -166,7 +166,7 @@ float FAggregator::SumMods(const TArray<FAggregatorMod> &Mods, float Bias, const
 	return Sum;
 }
 
-void FAggregator::AddMod(float EvaluatedMagnitude, TEnumAsByte<EGameplayModOp::Type> ModifierOp, const FGameplayTagRequirements* SourceTagReqs, const FGameplayTagRequirements* TargetTagReqs, bool IsPredicted, FActiveGameplayEffectHandle ActiveHandle)
+void FAggregator::AddAggregatorMod(float EvaluatedMagnitude, TEnumAsByte<EGameplayModOp::Type> ModifierOp, const FGameplayTagRequirements* SourceTagReqs, const FGameplayTagRequirements* TargetTagReqs, bool IsPredicted, FActiveGameplayEffectHandle ActiveHandle)
 {
 	TArray<FAggregatorMod> &ModList = Mods[ModifierOp];
 
@@ -180,12 +180,14 @@ void FAggregator::AddMod(float EvaluatedMagnitude, TEnumAsByte<EGameplayModOp::T
 	NewMod.IsPredicted = IsPredicted;
 
 	if (IsPredicted)
+	{
 		NumPredictiveMods++;
+	}
 
 	BroadcastOnDirty();
 }
 
-void FAggregator::RemoveMod(FActiveGameplayEffectHandle ActiveHandle)
+void FAggregator::RemoveAggregatorMod(FActiveGameplayEffectHandle ActiveHandle)
 {
 	if (ActiveHandle.IsValid())
 	{
@@ -206,9 +208,14 @@ void FAggregator::AddModsFrom(const FAggregator& SourceAggregator)
 	}
 }
 
-void FAggregator::AddDependant(FActiveGameplayEffectHandle Handle)
+void FAggregator::AddDependent(FActiveGameplayEffectHandle Handle)
 {
-	Dependants.Add(Handle);
+	Dependents.Add(Handle);
+}
+
+void FAggregator::RemoveDependent(FActiveGameplayEffectHandle Handle)
+{
+	Dependents.Remove(Handle);
 }
 
 bool FAggregator::HasPredictedMods() const
@@ -225,7 +232,9 @@ void FAggregator::RemoveModsWithActiveHandle(TArray<FAggregatorMod>& Mods, FActi
 		if (Mods[idx].ActiveHandle == ActiveHandle)
 		{
 			if (Mods[idx].IsPredicted)
+			{
 				NumPredictiveMods--;
+			}
 
 			Mods.RemoveAtSwap(idx, 1, false);
 		}
@@ -243,8 +252,8 @@ void FAggregator::TakeSnapshotOf(const FAggregator& AggToSnapshot)
 
 void FAggregator::BroadcastOnDirty()
 {
-	// If we are batching on Dirty calls (and we actually have dependants registered with us) then early out.
-	if (FScopedAggregatorOnDirtyBatch::GlobalBatchCount > 0 && (Dependants.Num() > 0 || OnDirty.IsBound()))
+	// If we are batching on Dirty calls (and we actually have dependents registered with us) then early out.
+	if (FScopedAggregatorOnDirtyBatch::GlobalBatchCount > 0 && (Dependents.Num() > 0 || OnDirty.IsBound()))
 	{
 		FScopedAggregatorOnDirtyBatch::DirtyAggregators.Add(this);
 		return;
@@ -253,7 +262,7 @@ void FAggregator::BroadcastOnDirty()
 	if (IsBroadcastingDirty)
 	{
 		// Apologies for the vague warning but its very hard from this spot to call out what data has caused this. If this frequently happens we should improve this.
-		ABILITY_LOG(Warning, TEXT("FAggregator detected cyclic attribute dependancies. We are skipping a recursive dirty call. Its possible the resulting attribute values are not what you expect!"));
+		ABILITY_LOG(Warning, TEXT("FAggregator detected cyclic attribute dependencies. We are skipping a recursive dirty call. Its possible the resulting attribute values are not what you expect!"));
 		return;
 	}
 
@@ -262,17 +271,17 @@ void FAggregator::BroadcastOnDirty()
 	OnDirty.Broadcast(this);
 
 
-	TArray<FActiveGameplayEffectHandle>	ValidDependants;
-	for (FActiveGameplayEffectHandle Handle : Dependants)
+	TArray<FActiveGameplayEffectHandle>	ValidDependents;
+	for (FActiveGameplayEffectHandle Handle : Dependents)
 	{
 		UAbilitySystemComponent* ASC = Handle.GetOwningAbilitySystemComponent();
 		if (ASC)
 		{
-			ASC->OnMagnitudeDependancyChange(Handle, this);
-			ValidDependants.Add(Handle);
+			ASC->OnMagnitudeDependencyChange(Handle, this);
+			ValidDependents.Add(Handle);
 		}
 	}
-	Dependants = ValidDependants;
+	Dependents = ValidDependents;
 
 }
 
diff --git a/Engine/Source/Runtime/GameplayAbilities/Private/GameplayEffectExecutionCalculation.cpp b/Engine/Source/Runtime/GameplayAbilities/Private/GameplayEffectExecutionCalculation.cpp
--- a/Engine/Source/Runtime/GameplayAbilities/Private/GameplayEffectExecutionCalculation.cpp
+++ b/Engine/Source/Runtime/GameplayAbilities/Private/GameplayEffectExecutionCalculation.cpp
@@ -35,7 +35,7 @@ FGameplayEffectCustomExecutionParameters::FGameplayEffectCustomExecutionParamete
 		float ModEvalValue = 0.f;
 		if (ScopedAggregator && CurScopedMod.ModifierMagnitude.AttemptCalculateMagnitude(InOwningSpec, ModEvalValue))
 		{
-			ScopedAggregator->AddMod(ModEvalValue, CurScopedMod.ModifierOp, &CurScopedMod.SourceTags, &CurScopedMod.TargetTags, false, ModifierHandle);
+			ScopedAggregator->AddAggregatorMod(ModEvalValue, CurScopedMod.ModifierOp, &CurScopedMod.SourceTags, &CurScopedMod.TargetTags, false, ModifierHandle);
 		}
 		else
 		{
diff --git a/Engine/Source/Runtime/GameplayAbilities/Private/GameplayEffectTypes.cpp b/Engine/Source/Runtime/GameplayAbilities/Private/GameplayEffectTypes.cpp
--- a/Engine/Source/Runtime/GameplayAbilities/Private/GameplayEffectTypes.cpp
+++ b/Engine/Source/Runtime/GameplayAbilities/Private/GameplayEffectTypes.cpp
@@ -144,7 +144,6 @@ bool FGameplayEffectContextHandle::NetSerialize(FArchive& Ar, class UPackageMap*
 			// For now, just always reset/reallocate the data when loading.
 			// Longer term if we want to generalize this and use it for property replication, we should support
 			// only reallocating when necessary
-			check(!Data.IsValid());
 
 			FGameplayEffectContext * NewData = (FGameplayEffectContext*)FMemory::Malloc(ScriptStruct->GetCppStructOps()->GetSize());
 			ScriptStruct->InitializeStruct(NewData);
@@ -194,15 +193,13 @@ bool FGameplayEffectContextHandle::NetSerialize(FArchive& Ar, class UPackageMap*
 FString EGameplayModOpToString(int32 Type)
 {
 	static UEnum *e = FindObject<UEnum>(ANY_PACKAGE, TEXT("EGameplayModOp"));
-	FString Right;
-	e->GetEnum(Type).ToString().Split(TEXT("::"), nullptr, &Right);
-	return Right;
+	return e->GetEnumName(Type);
 }
 
 FString EGameplayModEffectToString(int32 Type)
 {
 	static UEnum *e = FindObject<UEnum>(ANY_PACKAGE, TEXT("EGameplayModEffect"));
-	return e->GetEnum(Type).ToString();
+	return e->GetEnumName(Type);
 }
 
 bool FGameplayTagCountContainer::HasMatchingGameplayTag(FGameplayTag TagToCheck) const
diff --git a/Engine/Source/Runtime/GameplayAbilities/Private/GameplayPrediction.cpp b/Engine/Source/Runtime/GameplayAbilities/Private/GameplayPrediction.cpp
--- a/Engine/Source/Runtime/GameplayAbilities/Private/GameplayPrediction.cpp
+++ b/Engine/Source/Runtime/GameplayAbilities/Private/GameplayPrediction.cpp
@@ -49,7 +49,7 @@ void FPredictionKey::GenerateNewPredictionKey()
 	bIsStale = false;
 }
 
-void FPredictionKey::GenerateDependantPredictionKey()
+void FPredictionKey::GenerateDependentPredictionKey()
 {
 	KeyType Previous = 0;
 	if (Base == 0)
@@ -65,7 +65,7 @@ void FPredictionKey::GenerateDependantPredictionKey()
 
 	if (Previous > 0)
 	{
-		FPredictionKeyDelegates::AddDependancy(Current, Previous);
+		FPredictionKeyDelegates::AddDependency(Current, Previous);
 	}
 }
 
@@ -186,7 +186,7 @@ void FPredictionKeyDelegates::CaughtUp(FPredictionKey::KeyType Key)
 	}
 }
 
-void FPredictionKeyDelegates::AddDependancy(FPredictionKey::KeyType ThisKey, FPredictionKey::KeyType DependsOn)
+void FPredictionKeyDelegates::AddDependency(FPredictionKey::KeyType ThisKey, FPredictionKey::KeyType DependsOn)
 {
 	NewRejectedDelegate(DependsOn).BindStatic(&FPredictionKeyDelegates::Reject, ThisKey);
 	NewCaughtUpDelegate(DependsOn).BindStatic(&FPredictionKeyDelegates::CaughtUp, ThisKey);
@@ -227,7 +227,7 @@ FScopedPredictionWindow::FScopedPredictionWindow(UAbilitySystemComponent* InAbil
 	{
 		ClearScopedPredictionKey = true;
 		RestoreKey = InAbilitySystemComponent->ScopedPredictionKey;
-		InAbilitySystemComponent->ScopedPredictionKey.GenerateDependantPredictionKey();
+		InAbilitySystemComponent->ScopedPredictionKey.GenerateDependentPredictionKey();
 		
 	}
 }
diff --git a/Engine/Source/Runtime/GameplayAbilities/Public/Abilities/GameplayAbility.h b/Engine/Source/Runtime/GameplayAbilities/Public/Abilities/GameplayAbility.h
--- a/Engine/Source/Runtime/GameplayAbilities/Public/Abilities/GameplayAbility.h
+++ b/Engine/Source/Runtime/GameplayAbilities/Public/Abilities/GameplayAbility.h
@@ -227,6 +227,8 @@ public:
 	/** Is this ability triggered from TriggerData (or is it triggered explicitly through input/game code) */
 	bool IsTriggered() const;
 
+	bool HasAuthority(const FGameplayAbilityActivationInfo* ActivationInfo) const;
+
 	bool HasAuthorityOrPredictionKey(const FGameplayAbilityActorInfo* ActorInfo, const FGameplayAbilityActivationInfo* ActivationInfo) const;
 
 	/** Called when the ability is given to an AbilitySystemComponent */
@@ -385,6 +387,10 @@ protected:
 	//	
 	// -------------------------------------
 
+	// ---------
+	// Apply Self
+	// ---------
+
 	UFUNCTION(BlueprintCallable, Category = Ability, FriendlyName="ApplyGameplayEffectToOwner")
 	FActiveGameplayEffectHandle BP_ApplyGameplayEffectToOwner(TSubclassOf<UGameplayEffect> GameplayEffectClass, int32 GameplayEffectLevel = 1);
 
@@ -399,6 +405,8 @@ protected:
 
 	FActiveGameplayEffectHandle ApplyGameplayEffectSpecToOwner(const FGameplayAbilitySpecHandle AbilityHandle, const FGameplayAbilityActorInfo* ActorInfo, const FGameplayAbilityActivationInfo ActivationInfo, const FGameplayEffectSpecHandle SpecHandle);
 
+	// ---------
+	// Apply Target
 	// ---------
 
 	UFUNCTION(BlueprintCallable, Category = Ability, FriendlyName = "ApplyGameplayEffectToTarget")
@@ -415,6 +423,17 @@ protected:
 
 	TArray<FActiveGameplayEffectHandle> ApplyGameplayEffectSpecToTarget(const FGameplayAbilitySpecHandle AbilityHandle, const FGameplayAbilityActorInfo* ActorInfo, const FGameplayAbilityActivationInfo ActivationInfo, const FGameplayEffectSpecHandle SpecHandle, FGameplayAbilityTargetDataHandle TargetData);
 
+	// ---------
+	// Remove Self
+	// ---------
+	
+	/** Removes GAmeplayEffects from owner which match the given asset level tags. */
+	UFUNCTION(BlueprintCallable, Category = Ability, FriendlyName="RemoveGameplayEffectFromOwnerWithAssetTags")
+	void BP_RemoveGameplayEffectFromOwnerWithAssetTags(FGameplayTagContainer WithAssetTags, int32 StacksToRemove = -1);
+
+	UFUNCTION(BlueprintCallable, Category = Ability, FriendlyName="RemoveGameplayEffectFromOwnerWithGrantedTags")
+	void BP_RemoveGameplayEffectFromOwnerWithGrantedTags(FGameplayTagContainer WithGrantedTags, int32 StacksToRemove = -1);
+
 	// -------------------------------------
 	//	GameplayCue
 	//	Abilities can invoke GameplayCues without having to create GameplayEffects
diff --git a/Engine/Source/Runtime/GameplayAbilities/Public/Abilities/GameplayAbilityTargetActor.h b/Engine/Source/Runtime/GameplayAbilities/Public/Abilities/GameplayAbilityTargetActor.h
--- a/Engine/Source/Runtime/GameplayAbilities/Public/Abilities/GameplayAbilityTargetActor.h
+++ b/Engine/Source/Runtime/GameplayAbilities/Public/Abilities/GameplayAbilityTargetActor.h
@@ -24,10 +24,6 @@ class GAMEPLAYABILITIES_API AGameplayAbilityTargetActor : public AActor
 
 public:
 
-	/** Native classes can set this to call StaticGetTargetData to instantly get TargetData, instead of instantiating the GameplayTargetActor and calling StartTargeting. */
-	UPROPERTY()
-	bool StaticTargetFunction;
-
 	/** The TargetData this class produces can be entirely generated on the server. We don't require the client to send us full or partial TargetData (possibly just a 'confirm') */
 	UPROPERTY(EditAnywhere, Category=Advanced)
 	bool ShouldProduceTargetDataOnServer;
@@ -36,8 +32,6 @@ public:
 	//UPROPERTY(BlueprintReadOnly, meta=(ExposeOnSpawn=true), Category=Targeting)
 	UPROPERTY(BlueprintReadOnly, meta = (ExposeOnSpawn = true), Replicated, Category = Targeting)
 	FGameplayAbilityTargetingLocationInfo StartLocation;
-
-	virtual FGameplayAbilityTargetDataHandle StaticGetTargetData(UWorld * World, const FGameplayAbilityActorInfo* ActorInfo, FGameplayAbilityActivationInfo ActivationInfo) const;
 	
 	/** Initialize and begin targeting logic  */
 	virtual void StartTargeting(UGameplayAbility* Ability);
diff --git a/Engine/Source/Runtime/GameplayAbilities/Public/Abilities/GameplayAbilityTargetActor_Trace.h b/Engine/Source/Runtime/GameplayAbilities/Public/Abilities/GameplayAbilityTargetActor_Trace.h
--- a/Engine/Source/Runtime/GameplayAbilities/Public/Abilities/GameplayAbilityTargetActor_Trace.h
+++ b/Engine/Source/Runtime/GameplayAbilities/Public/Abilities/GameplayAbilityTargetActor_Trace.h
@@ -21,8 +21,6 @@ public:
 	/** Sweeps as normal, but will manually filter all hit actors */
 	static void SweepWithFilter(FHitResult& OutHitResult, const UWorld* World, const FGameplayTargetDataFilterHandle FilterHandle, const FVector& Start, const FVector& End, const FQuat& Rotation, const FCollisionShape CollisionShape, FName ProfileName, const FCollisionQueryParams Params);
 
-	virtual FGameplayAbilityTargetDataHandle StaticGetTargetData(UWorld * World, const FGameplayAbilityActorInfo* ActorInfo, FGameplayAbilityActivationInfo ActivationInfo) const override;
-
 	void AimWithPlayerController(const AActor* InSourceActor, FCollisionQueryParams Params, const FVector& TraceStart, FVector& OutTraceEnd, bool bIgnorePitch = false) const;
 
 	static bool ClipCameraRayToAbilityRange(FVector CameraLocation, FVector CameraDirection, FVector AbilityCenter, float AbilityRange, FVector& ClippedPosition);
diff --git a/Engine/Source/Runtime/GameplayAbilities/Public/Abilities/Tasks/AbilityTask_VisualizeTargeting.h b/Engine/Source/Runtime/GameplayAbilities/Public/Abilities/Tasks/AbilityTask_VisualizeTargeting.h
--- a/Engine/Source/Runtime/GameplayAbilities/Public/Abilities/Tasks/AbilityTask_VisualizeTargeting.h
+++ b/Engine/Source/Runtime/GameplayAbilities/Public/Abilities/Tasks/AbilityTask_VisualizeTargeting.h
@@ -34,9 +34,7 @@ protected:
 
 	virtual void OnDestroy(bool AbilityEnded) override;
 
-	TSubclassOf<AGameplayAbilityTargetActor> TargetClass;
-
-	/** The TargetActor that we spawned, or the class CDO if this is a static targeting task */
+	/** The TargetActor that we spawned */
 	TWeakObjectPtr<AGameplayAbilityTargetActor>	MyTargetActor;
 
 	/** Handle for efficient management of OnTimeElapsed timer */
diff --git a/Engine/Source/Runtime/GameplayAbilities/Public/Abilities/Tasks/AbilityTask_WaitTargetData.h b/Engine/Source/Runtime/GameplayAbilities/Public/Abilities/Tasks/AbilityTask_WaitTargetData.h
--- a/Engine/Source/Runtime/GameplayAbilities/Public/Abilities/Tasks/AbilityTask_WaitTargetData.h
+++ b/Engine/Source/Runtime/GameplayAbilities/Public/Abilities/Tasks/AbilityTask_WaitTargetData.h
@@ -57,7 +57,7 @@ protected:
 
 	TSubclassOf<AGameplayAbilityTargetActor> TargetClass;
 
-	/** The TargetActor that we spawned, or the class CDO if this is a static targeting task */
+	/** The TargetActor that we spawned */
 	TWeakObjectPtr<AGameplayAbilityTargetActor>	MyTargetActor;
 
 	TEnumAsByte<EGameplayTargetingConfirmation::Type> ConfirmationType;
diff --git a/Engine/Source/Runtime/GameplayAbilities/Public/AbilitySystemComponent.h b/Engine/Source/Runtime/GameplayAbilities/Public/AbilitySystemComponent.h
--- a/Engine/Source/Runtime/GameplayAbilities/Public/AbilitySystemComponent.h
+++ b/Engine/Source/Runtime/GameplayAbilities/Public/AbilitySystemComponent.h
@@ -219,8 +219,9 @@ class GAMEPLAYABILITIES_API UAbilitySystemComponent : public UActorComponent, pu
 	FActiveGameplayEffectHandle ApplyGameplayEffectSpecToTarget(OUT FGameplayEffectSpec& GameplayEffect, UAbilitySystemComponent *Target, FPredictionKey PredictionKey=FPredictionKey());
 	FActiveGameplayEffectHandle ApplyGameplayEffectSpecToSelf(OUT FGameplayEffectSpec& GameplayEffect, FPredictionKey PredictionKey = FPredictionKey());
 
+	/** Removes GameplayEffect by Handle. StacksToRemove=-1 will remove all stacks. */
 	UFUNCTION(BlueprintCallable, BlueprintAuthorityOnly, Category = GameplayEffects)
-	bool RemoveActiveGameplayEffect(FActiveGameplayEffectHandle Handle);
+	bool RemoveActiveGameplayEffect(FActiveGameplayEffectHandle Handle, int32 StacksToRemove=-1);
 
 	/** Get an outgoing GameplayEffectSpec that is ready to be applied to other things. */
 	UFUNCTION(BlueprintCallable, Category = GameplayEffects)
@@ -238,9 +239,9 @@ class GAMEPLAYABILITIES_API UAbilitySystemComponent : public UActorComponent, pu
 	FGameplayEffectContextHandle GetEffectContext() const;
 
 	/** This only exists so it can be hooked up to a multicast delegate */
-	void RemoveActiveGameplayEffect_NoReturn(FActiveGameplayEffectHandle Handle)
+	void RemoveActiveGameplayEffect_NoReturn(FActiveGameplayEffectHandle Handle, int32 StacksToRemove=-1)
 	{
-		RemoveActiveGameplayEffect(Handle);
+		RemoveActiveGameplayEffect(Handle, StacksToRemove);
 	}
 
 	float GetGameplayEffectDuration(FActiveGameplayEffectHandle Handle) const;
@@ -284,7 +285,7 @@ class GAMEPLAYABILITIES_API UAbilitySystemComponent : public UActorComponent, pu
 
 	void OnAttributeAggregatorDirty(FAggregator* Aggregator, FGameplayAttribute Attribute);
 
-	void OnMagnitudeDependancyChange(FActiveGameplayEffectHandle Handle, const FAggregator* ChangedAggregator);
+	void OnMagnitudeDependencyChange(FActiveGameplayEffectHandle Handle, const FAggregator* ChangedAggregator);
 
 	/** This ASC has successfully applied a GE to something (potentially itself) */
 	void OnGameplayEffectAppliedToTarget(UAbilitySystemComponent* Target, const FGameplayEffectSpec& SpecApplied, FActiveGameplayEffectHandle ActiveHandle);
@@ -392,8 +393,8 @@ class GAMEPLAYABILITIES_API UAbilitySystemComponent : public UActorComponent, pu
 	UFUNCTION(BlueprintCallable, Category = GameplayEffects)
 	void RemoveActiveEffectsWithTags(FGameplayTagContainer Tags);
 
-	/** Removes all active effects that match given query */
-	void RemoveActiveEffects(const FActiveGameplayEffectQuery Query);
+	/** Removes all active effects that match given query. StacksToRemove=-1 will remove all stacks. */
+	void RemoveActiveEffects(const FActiveGameplayEffectQuery Query, int32 StacksToRemove=-1);
 
 	void OnRestackGameplayEffects();	
 	
diff --git a/Engine/Source/Runtime/GameplayAbilities/Public/AbilitySystemStats.h b/Engine/Source/Runtime/GameplayAbilities/Public/AbilitySystemStats.h
--- a/Engine/Source/Runtime/GameplayAbilities/Public/AbilitySystemStats.h
+++ b/Engine/Source/Runtime/GameplayAbilities/Public/AbilitySystemStats.h
@@ -16,7 +16,7 @@ DECLARE_CYCLE_STAT_EXTERN(TEXT("GetActiveEffectsData"), STAT_GameplayEffectsGetA
 DECLARE_CYCLE_STAT_EXTERN(TEXT("GetCooldownTimeRemaining"), STAT_GameplayAbilityGetCooldownTimeRemaining, STATGROUP_AbilitySystem, );
 DECLARE_CYCLE_STAT_EXTERN(TEXT("GetCooldownTimeRemaining"), STAT_GameplayAbilityGetCooldownTimeRemainingAndDuration, STATGROUP_AbilitySystem, );
 DECLARE_CYCLE_STAT_EXTERN(TEXT("RemoveActiveGameplayEffect"), STAT_RemoveActiveGameplayEffect, STATGROUP_AbilitySystem, );
-DECLARE_CYCLE_STAT_EXTERN(TEXT("CreateNewActiveGameplayEffect"), STAT_CreateNewActiveGameplayEffect, STATGROUP_AbilitySystem, );
+DECLARE_CYCLE_STAT_EXTERN(TEXT("ApplyGameplayEffectSpec"), STAT_ApplyGameplayEffectSpec, STATGROUP_AbilitySystem, );
 DECLARE_CYCLE_STAT_EXTERN(TEXT("GetGameplayCueFunction"), STAT_GetGameplayCueFunction, STATGROUP_AbilitySystem, );
 DECLARE_CYCLE_STAT_EXTERN(TEXT("GetOutgoingSpec"), STAT_GetOutgoingSpec, STATGROUP_AbilitySystem, );
 DECLARE_CYCLE_STAT_EXTERN(TEXT("InitAttributeSetDefaults"), STAT_InitAttributeSetDefaults, STATGROUP_AbilitySystem, );
diff --git a/Engine/Source/Runtime/GameplayAbilities/Public/GameplayCueNotify_Actor.h b/Engine/Source/Runtime/GameplayAbilities/Public/GameplayCueNotify_Actor.h
--- a/Engine/Source/Runtime/GameplayAbilities/Public/GameplayCueNotify_Actor.h
+++ b/Engine/Source/Runtime/GameplayAbilities/Public/GameplayCueNotify_Actor.h
@@ -49,6 +49,9 @@ class GAMEPLAYABILITIES_API AGameplayCueNotify_Actor : public AActor
 	UFUNCTION(BlueprintNativeEvent, Category = "GameplayCueNotify")
 	bool OnActive(AActor* MyTarget, FGameplayCueParameters Parameters);
 
+	UFUNCTION(BlueprintNativeEvent, Category = "GameplayCueNotify")
+	bool WhileActive(AActor* MyTarget, FGameplayCueParameters Parameters);
+
 	UFUNCTION(BlueprintNativeEvent, Category = "GameplayCueNotify")
 	bool OnRemove(AActor* MyTarget, FGameplayCueParameters Parameters);
 
diff --git a/Engine/Source/Runtime/GameplayAbilities/Public/GameplayCueNotify_Static.h b/Engine/Source/Runtime/GameplayAbilities/Public/GameplayCueNotify_Static.h
--- a/Engine/Source/Runtime/GameplayAbilities/Public/GameplayCueNotify_Static.h
+++ b/Engine/Source/Runtime/GameplayAbilities/Public/GameplayCueNotify_Static.h
@@ -49,6 +49,9 @@ class GAMEPLAYABILITIES_API UGameplayCueNotify_Static : public UObject
 	UFUNCTION(BlueprintNativeEvent, BlueprintPure, Category = "GameplayCueNotify")
 	bool OnActive(AActor* MyTarget, FGameplayCueParameters Parameters) const;
 
+	UFUNCTION(BlueprintNativeEvent, BlueprintPure, Category = "GameplayCueNotify")
+	bool WhileActive(AActor* MyTarget, FGameplayCueParameters Parameters) const;
+
 	UFUNCTION(BlueprintNativeEvent, BlueprintPure, Category = "GameplayCueNotify")
 	bool OnRemove(AActor* MyTarget, FGameplayCueParameters Parameters) const;
 
diff --git a/Engine/Source/Runtime/GameplayAbilities/Public/GameplayEffect.h b/Engine/Source/Runtime/GameplayAbilities/Public/GameplayEffect.h
--- a/Engine/Source/Runtime/GameplayAbilities/Public/GameplayEffect.h
+++ b/Engine/Source/Runtime/GameplayAbilities/Public/GameplayEffect.h
@@ -19,6 +19,8 @@ class UAbilitySystemComponent;
 class UGameplayModMagnitudeCalculation;
 class UGameplayEffectExecutionCalculation;
 
+DECLARE_DELEGATE_RetVal_OneParam(bool, FActiveGameplayEffectQueryCustomMatch, const FActiveGameplayEffect);
+
 /** Enumeration outlining the possible gameplay effect magnitude calculation policies */
 UENUM()
 enum class EGameplayEffectMagnitudeCalculation : uint8
@@ -218,7 +220,7 @@ public:
 	bool AttemptCalculateMagnitude(const FGameplayEffectSpec& InRelevantSpec, OUT float& OutCalculatedMagnitude) const;
 
 	/** Attempts to recalculate the magnitude given a changed aggregator. This will only recalculate if we are a modifier that is linked (non snapshot) to the given aggregator. */
-	bool AttempteRecalculateMagnitudeFromDependantChange(const FGameplayEffectSpec& InRelevantSpec, OUT float& OutCalculatedMagnitude, const FAggregator* ChangedAggregator) const;
+	bool AttemptRecalculateMagnitudeFromDependentChange(const FGameplayEffectSpec& InRelevantSpec, OUT float& OutCalculatedMagnitude, const FAggregator* ChangedAggregator) const;
 
 	/**
 	 * Gather all of the attribute capture definitions necessary to compute the magnitude and place them into the provided array
@@ -468,6 +470,28 @@ enum class EGameplayEffectDurationType : uint8
 	HasDuration
 };
 
+/** Enumeration of policies for dealing with duration of a gameplay effect while stacking */
+UENUM()
+enum class EGameplayEffectStackingDurationPolicy : uint8
+{
+	/** The duration of the effect will be refreshed from any successful stack application */
+	RefreshOnSuccessfulApplication,
+
+	/** The duration of the effect will never be refreshed */
+	NeverRefresh,
+};
+
+/** Enumeration of policies for dealing with the period of a gameplay effect while stacking */
+UENUM()
+enum class EGameplayEffectStackingPeriodPolicy : uint8
+{
+	/** Any progress toward the next tick of a periodic effect is discarded upon any successful stack application */
+	ResetOnSuccessfulApplication,
+
+	/** The progress toward the next tick of a periodic effect will never be reset, regardless of stack applications */
+	NeverReset,
+};
+
 /**
  * UGameplayEffect
  *	The GameplayEffect definition. This is the data asset defined in the editor that drives everything.
@@ -520,9 +544,13 @@ public:
 	FScalableFloat	Duration;
 
 	/** Period in seconds. 0.0 for non-periodic effects */
-	UPROPERTY(EditDefaultsOnly, Category=GameplayEffect)
+	UPROPERTY(EditDefaultsOnly, BlueprintReadOnly, Category=Period)
 	FScalableFloat	Period;
 	
+	/** If true, the effect executes on application and then at every period interval. If false, no execution occurs until the first period elapses. */
+	UPROPERTY(EditDefaultsOnly, BlueprintReadOnly, Category=Period)
+	bool bExecutePeriodicEffectOnApplication;
+
 	/** Array of modifiers that will affect the target of this effect */
 	UPROPERTY(EditDefaultsOnly, BlueprintReadOnly, Category=GameplayEffect)
 	TArray<FGameplayModifierInfo> Modifiers;
@@ -546,6 +574,18 @@ public:
 	UPROPERTY(VisibleDefaultsOnly, Category = Deprecated)
 	TArray<UGameplayEffect*> TargetEffects;
 
+	/** Effects to apply when a stacking effect "overflows" its stack count through another attempted application. Added whether the overflow application succeeds or not. */
+	UPROPERTY(EditDefaultsOnly, BlueprintReadOnly, Category=Overflow)
+	TArray<TSubclassOf<UGameplayEffect>> OverflowEffects;
+
+	/** If true, stacking attempts made while at the stack count will fail, resulting in the duration and context not being refreshed */
+	UPROPERTY(EditDefaultsOnly, BlueprintReadOnly, Category=Overflow)
+	bool bDenyOverflowApplication;
+
+	/** If true, the entire stack of the effect will be cleared once it overflows */
+	UPROPERTY(EditDefaultsOnly, BlueprintReadOnly, Category=Overflow, meta=(EditCondition="bDenyOverflowApplication"))
+	bool bClearStackOnOverflow;
+
 	void GetTargetEffects(TArray<const UGameplayEffect*>& OutEffects) const;
 
 	// ------------------------------------------------
@@ -609,6 +649,26 @@ public:
 	/** Grants the owner immunity from these source tags. */
 	UPROPERTY(EditDefaultsOnly, BlueprintReadOnly, Category = Tags, meta = (DisplayName = "GrantedApplicationImmunityTags"))
 	FGameplayTagRequirements GrantedApplicationImmunityTags;
+
+	// ----------------------------------------------------------------------
+	//	Stacking
+	// ----------------------------------------------------------------------
+	
+	/** How this GameplayEffect stacks with other instances of this same GameplayEffect */
+	UPROPERTY(EditDefaultsOnly, BlueprintReadOnly, Category = Stacking)
+	EGameplayEffectStackingType	StackingType;
+
+	/** Stack limit for StackingType */
+	UPROPERTY(EditDefaultsOnly, BlueprintReadOnly, Category = Stacking)
+	int32 StackLimitCount;
+
+	/** Policy for how the effect duration should be refreshed while stacking */
+	UPROPERTY(EditDefaultsOnly, BlueprintReadOnly, Category = Stacking)
+	EGameplayEffectStackingDurationPolicy StackDurationRefreshPolicy;
+
+	/** Policy for how the effect period should be reset (or not) while stacking */
+	UPROPERTY(EditDefaultsOnly, BlueprintReadOnly, Category = Stacking)
+	EGameplayEffectStackingPeriodPolicy StackPeriodResetPolicy;
 };
 
 /** Holds evaluated magnitude from a GameplayEffect modifier */
@@ -740,6 +800,9 @@ struct GAMEPLAYABILITIES_API FGameplayEffectAttributeCaptureSpec
 
 	/** Register this handle with linked aggregators */
 	void RegisterLinkedAggregatorCallback(FActiveGameplayEffectHandle Handle) const;
+
+	/** Unregister this handle with linked aggregators */
+	void UnregisterLinkedAggregatorCallback(FActiveGameplayEffectHandle Handle) const;
 	
 	/** Return true if this capture should be recalculated if the given aggregator has changed */
 	bool ShouldRefreshLinkedAggregator(const FAggregator* ChangedAggregator) const;
@@ -803,6 +866,9 @@ public:
 	/** Registers any linked aggregators to notify this active handle if they are dirtied */
 	void RegisterLinkedAggregatorCallbacks(FActiveGameplayEffectHandle Handle) const;
 
+	/** Unregisters any linked aggregators from notifying this active handle if they are dirtied */
+	void UnregisterLinkedAggregatorCallbacks(FActiveGameplayEffectHandle Handle) const;
+
 private:
 
 	/** Captured attributes from the source of a gameplay effect */
@@ -838,6 +904,10 @@ struct GAMEPLAYABILITIES_API FGameplayEffectSpec
 		, Period(UGameplayEffect::NO_PERIOD)
 		, ChanceToApplyToTarget(1.f)
 		, ChanceToExecuteOnGameplayEffect(1.f)
+		, StackCount(1)
+		, bCompletedSourceAttributeCapture(false)
+		, bCompletedTargetAttributeCapture(false)
+		, bDurationLocked(false)
 		, Level(UGameplayEffect::INVALID_LEVEL)
 	{
 
@@ -874,8 +944,17 @@ struct GAMEPLAYABILITIES_API FGameplayEffectSpec
 	/** Deletes any modified attributes that aren't needed. Call before replication */
 	void PruneModifiedAttributes();
 
-	/** Sets duration. This should only be called as the GameplayEffect is being created */
-	void SetDuration(float NewDuration);
+	/**
+	 * Helper function to attempt to calculate the duration of the spec from its GE definition
+	 * 
+	 * @param OutDefDuration	Computed duration of the spec from its GE definition; Not the actual duration of the spec
+	 * 
+	 * @return True if the calculation was successful, false if it was not
+	 */
+	bool AttemptCalculateDurationFromDef(OUT float& OutDefDuration) const;
+
+	/** Sets duration. This should only be called as the GameplayEffect is being created and applied; Ignores calls after attribute capture */
+	void SetDuration(float NewDuration, bool bLockDuration);
 
 	float GetDuration() const;
 	float GetPeriod() const;
@@ -934,6 +1013,21 @@ struct GAMEPLAYABILITIES_API FGameplayEffectSpec
 	
 	UPROPERTY()
 	TArray<FModifierSpec> Modifiers;
+
+	UPROPERTY()
+	int32 StackCount;
+
+	/** Whether the spec has had its source attribute capture completed or not yet */
+	UPROPERTY(NotReplicated)
+	uint32 bCompletedSourceAttributeCapture : 1;
+
+	/** Whether the spec has had its target attribute capture completed or not yet */
+	UPROPERTY(NotReplicated)
+	uint32 bCompletedTargetAttributeCapture : 1;
+
+	/** Whether the duration of the spec is locked or not; If it is, attempts to set it will fail */
+	UPROPERTY(NotReplicated)
+	uint32 bDurationLocked : 1;
 	
 	void CalculateModifierMagnitudes();
 
@@ -953,13 +1047,14 @@ struct GAMEPLAYABILITIES_API FGameplayEffectSpec
 		return EffectContext;
 	}
 
+	void CaptureAttributeDataFromTarget(UAbilitySystemComponent* TargetAbilitySystemComponent);
+
 private:
 
 	/** Map of set by caller magnitudes */
 	TMap<FName, float>	SetByCallerMagnitudes;
 
 	void CaptureDataFromSource();
-	void RecalculateDuration();
 
 	/** Helper function to initialize all of the capture definitions required by the spec */
 	void SetupAttributeCaptureDefinitions();
@@ -986,6 +1081,7 @@ struct FActiveGameplayEffect : public FFastArraySerializerItem
 
 	FActiveGameplayEffect()
 		: StartGameStateTime(0)
+		, CachedStartGameStateTime(0)
 		, StartWorldTime(0.f)
 		, IsInhibited(true)
 	{
@@ -996,6 +1092,7 @@ struct FActiveGameplayEffect : public FFastArraySerializerItem
 		, Spec(InSpec)
 		, PredictionKey(InPredictionKey)
 		, StartGameStateTime(InStartGameStateTime)
+		, CachedStartGameStateTime(InStartGameStateTime)
 		, StartWorldTime(CurrentWorldTime)
 		, IsInhibited(true)
 	{
@@ -1013,6 +1110,10 @@ struct FActiveGameplayEffect : public FFastArraySerializerItem
 	UPROPERTY()
 	int32 StartGameStateTime;
 
+	/** Used for handling duration modifications being replicated */
+	UPROPERTY(NotReplicated)
+	int32 CachedStartGameStateTime;
+
 	UPROPERTY(NotReplicated)
 	float StartWorldTime;
 
@@ -1117,7 +1218,10 @@ struct FActiveGameplayEffectQuery
 	{
 	}
 
-	// Returns true if Effect matches the criteria of this query. Returns false otherwise.
+	/** Bind this to override the default query-matching code. */
+	FActiveGameplayEffectQueryCustomMatch CustomMatch;
+
+	/** Returns true if Effect matches the criteria of this query, which will be overridden by CustomMatch if it is bound. Returns false otherwise. */
 	bool Matches(const FActiveGameplayEffect& Effect) const;
 
 	/** used to match with InheritableOwnedTagsContainer */
@@ -1167,7 +1271,7 @@ struct FActiveGameplayEffectsContainer : public FFastArraySerializer
 
 	void RegisterWithOwner(UAbilitySystemComponent* Owner);	
 	
-	FActiveGameplayEffect& CreateNewActiveGameplayEffect(const FGameplayEffectSpec &Spec, FPredictionKey InPredictionKey);
+	FActiveGameplayEffect* ApplyGameplayEffectSpec(const FGameplayEffectSpec& Spec, FPredictionKey InPredictionKey);
 
 	FActiveGameplayEffect* GetActiveGameplayEffect(const FActiveGameplayEffectHandle Handle);
 		
@@ -1175,7 +1279,7 @@ struct FActiveGameplayEffectsContainer : public FFastArraySerializer
 	
 	void ExecutePeriodicGameplayEffect(FActiveGameplayEffectHandle Handle);	// This should not be outward facing to the skill system API, should only be called by the owning AbilitySystemComponent
 
-	bool RemoveActiveGameplayEffect(FActiveGameplayEffectHandle Handle);
+	bool RemoveActiveGameplayEffect(FActiveGameplayEffectHandle Handle, int32 StacksToRemove);
 	
 	float GetGameplayEffectDuration(FActiveGameplayEffectHandle Handle) const;
 
@@ -1235,7 +1339,7 @@ struct FActiveGameplayEffectsContainer : public FFastArraySerializer
 
 	TArray<float> GetActiveEffectsDuration(const FActiveGameplayEffectQuery Query) const;
 
-	void RemoveActiveEffects(const FActiveGameplayEffectQuery Query);
+	void RemoveActiveEffects(const FActiveGameplayEffectQuery Query, int32 StacksToRemove);
 
 	int32 GetGameStateTime() const;
 
@@ -1271,8 +1375,8 @@ private:
 
 	bool IsNetAuthority() const;
 
-	/** Called internally to actually remove a GameplayEffect */
-	bool InternalRemoveActiveGameplayEffect(int32 Idx);
+	/** Called internally to actually remove a GameplayEffect or to reduce its StackCount. Returns true if we resized our internal GameplayEffect array. */
+	bool InternalRemoveActiveGameplayEffect(int32 Idx, int32 StacksToRemove);
 
 	/** Called both in server side creation and replication creation/deletion */
 	void InternalOnActiveGameplayEffectAdded(FActiveGameplayEffect& Effect);
@@ -1281,8 +1385,8 @@ private:
 	void RemoveActiveGameplayEffectGrantedTagsAndModifiers(const FActiveGameplayEffect& Effect);
 	void AddActiveGameplayEffectGrantedTagsAndModifiers(FActiveGameplayEffect& Effect);
 
-	/** Updates tag dependancy map when a GameplayEffect is removed */
-	void RemoveActiveEffectTagDependancy(const FGameplayTagContainer& Tags, FActiveGameplayEffectHandle Handle);
+	/** Updates tag dependency map when a GameplayEffect is removed */
+	void RemoveActiveEffectTagDependency(const FGameplayTagContainer& Tags, FActiveGameplayEffectHandle Handle);
 
 	// -------------------------------------------------------------------------------------------
 
@@ -1292,13 +1396,32 @@ private:
 
 	TMap<FGameplayTag, TSet<FActiveGameplayEffectHandle> >	ActiveEffectTagDependencies;
 
+	/** A map to manage stacking while we are the source */
+	TMap<TWeakObjectPtr<UGameplayEffect>, TArray<FActiveGameplayEffectHandle> >	SourceStackingMap;
+
 	FGameplayTagCountContainer ApplicationImmunityGameplayTagCountContainer;
 
 	FAggregatorRef& FindOrCreateAttributeAggregator(FGameplayAttribute Attribute);
 
 	void OnAttributeAggregatorDirty(FAggregator* Aggregator, FGameplayAttribute Attribute);
 
-	void OnMagnitudeDependancyChange(FActiveGameplayEffectHandle Handle, const FAggregator* ChangedAgg);
+	void OnMagnitudeDependencyChange(FActiveGameplayEffectHandle Handle, const FAggregator* ChangedAgg);
+
+	void OnStackCountChange(FActiveGameplayEffect& ActiveEffect);
+
+	void UpdateAggregatorModMagnitudes(const TSet<FGameplayAttribute>& AttributesToUpdate, FActiveGameplayEffect& ActiveEffect);
+
+	/** Helper function to find the active GE that the specified spec can stack with, if any */
+	FActiveGameplayEffect* FindStackableActiveGameplayEffect(const FGameplayEffectSpec& Spec);
+
+	/** Helper function to compute the modified duration of a specified spec & base duration value */
+	float ComputeModifiedDurationOfAppliedSpec(const FGameplayEffectSpec& Spec, float BaseValue) const;
+	
+	/** Helper function to handle the case of same-effect stacking overflow; Returns true if the overflow application should apply, false if it should not */
+	bool HandleActiveGameplayEffectStackOverflow(const FActiveGameplayEffect& ActiveStackableGE, const FGameplayEffectSpec& OldSpec, const FGameplayEffectSpec& OverflowingSpec);
+
+	/** After application has gone through, give stacking rules a chance to do something as the source of the gameplay effect (E.g., remove an old version) */
+	void ApplyStackingLogicPostApplyAsSource(UAbilitySystemComponent* Target, const FGameplayEffectSpec& SpecApplied, FActiveGameplayEffectHandle ActiveHandle);
 };
 
 template<>
@@ -1327,10 +1450,11 @@ struct GAMEPLAYABILITIES_API FActiveGameplayEffectAction
 		bIsInitialized = true;
 	}
 
-	void InitForRemoveGE(TWeakObjectPtr<UAbilitySystemComponent> InOwningASC, const FActiveGameplayEffectHandle& InHandle)
+	void InitForRemoveGE(TWeakObjectPtr<UAbilitySystemComponent> InOwningASC, const FActiveGameplayEffectHandle& InHandle, int32 InStacksToRemove)
 	{
 		OwningASC = InOwningASC;
 		Handle = InHandle;
+		StacksToRemove = InStacksToRemove;
 		bRemove = true;
 		bIsInitialized = true;
 	}
@@ -1349,6 +1473,9 @@ struct GAMEPLAYABILITIES_API FActiveGameplayEffectAction
 
 	UPROPERTY()
 	bool bIsInitialized;
+
+	UPROPERTY()
+	int32 StacksToRemove;
 };
 
 
diff --git a/Engine/Source/Runtime/GameplayAbilities/Public/GameplayEffectAggregator.h b/Engine/Source/Runtime/GameplayAbilities/Public/GameplayEffectAggregator.h
--- a/Engine/Source/Runtime/GameplayAbilities/Public/GameplayEffectAggregator.h
+++ b/Engine/Source/Runtime/GameplayAbilities/Public/GameplayEffectAggregator.h
@@ -26,6 +26,7 @@ struct GAMEPLAYABILITIES_API FAggregatorMod
 	const FGameplayTagRequirements*	TargetTagReqs;
 
 	float EvaluatedMagnitude;		// Magnitude this mod was last evaluated at
+	float StackCount;
 
 	FActiveGameplayEffectHandle ActiveHandle;	// Handle of the active GameplayEffect we are tied to (if any)
 	bool IsPredicted;
@@ -47,12 +48,13 @@ struct GAMEPLAYABILITIES_API FAggregator : public TSharedFromThis<FAggregator>
 	void ExecModOnBaseValue(TEnumAsByte<EGameplayModOp::Type> ModifierOp, float EvaluatedMagnitude);
 	static float StaticExecModOnBaseValue(float BaseValue, TEnumAsByte<EGameplayModOp::Type> ModifierOp, float EvaluatedMagnitude);
 
-	void AddMod(float EvaluatedData, TEnumAsByte<EGameplayModOp::Type> ModifierOp, const FGameplayTagRequirements*	SourceTagReqs, const FGameplayTagRequirements* TargetTagReqs, bool IsPredicted, FActiveGameplayEffectHandle ActiveHandle = FActiveGameplayEffectHandle() );
-	void RemoveMod(FActiveGameplayEffectHandle ActiveHandle);
+	void AddAggregatorMod(float EvaluatedData, TEnumAsByte<EGameplayModOp::Type> ModifierOp, const FGameplayTagRequirements*	SourceTagReqs, const FGameplayTagRequirements* TargetTagReqs, bool IsPredicted, FActiveGameplayEffectHandle ActiveHandle = FActiveGameplayEffectHandle() );
+	void RemoveAggregatorMod(FActiveGameplayEffectHandle ActiveHandle);
 
 	/** Evaluates the Aggregator with the internal base value and given parameters */
 	float Evaluate(const FAggregatorEvaluateParameters& Parameters) const;
 
+	/** Works backwards to calculate the base value. Used on clients for doing predictive modifiers */
 	float ReverseEvaluate(float FinalValue, const FAggregatorEvaluateParameters& Parameters) const;
 
 	/** Evaluates the Aggregator with an arbitrary base value */
@@ -67,7 +69,8 @@ struct GAMEPLAYABILITIES_API FAggregator : public TSharedFromThis<FAggregator>
 
 	void AddModsFrom(const FAggregator& SourceAggregator);
 
-	void AddDependant(FActiveGameplayEffectHandle Handle);
+	void AddDependent(FActiveGameplayEffectHandle Handle);
+	void RemoveDependent(FActiveGameplayEffectHandle Handle);
 
 	bool HasPredictedMods() const;
 	
@@ -84,7 +87,7 @@ private:
 	TArray<FAggregatorMod>	Mods[EGameplayModOp::Max];
 
 	/** ActiveGE handles that we need to notify if we change. NOT copied over during snapshots. */
-	TArray<FActiveGameplayEffectHandle>	Dependants;
+	TArray<FActiveGameplayEffectHandle>	Dependents;
 	bool	IsBroadcastingDirty;
 	int32	NumPredictiveMods;
 
diff --git a/Engine/Source/Runtime/GameplayAbilities/Public/GameplayEffectTypes.h b/Engine/Source/Runtime/GameplayAbilities/Public/GameplayEffectTypes.h
--- a/Engine/Source/Runtime/GameplayAbilities/Public/GameplayEffectTypes.h
+++ b/Engine/Source/Runtime/GameplayAbilities/Public/GameplayEffectTypes.h
@@ -50,8 +50,22 @@ namespace EGameplayModOp
 UENUM()
 enum class EGameplayEffectAttributeCaptureSource : uint8
 {
-	Source,	// Source (caster) of the gameplay effect
-	Target	// Target (recipient) of the gameplay effect
+	// Source (caster) of the gameplay effect
+	Source,	
+	// Target (recipient) of the gameplay effect
+	Target	
+};
+
+/** Enumeration for ways a single GameplayEffect asset can stack */
+UENUM()
+enum class EGameplayEffectStackingType : uint8
+{
+	// No stacking. Multiple applications of this GameplayEffect are treated as separate instances.
+	None,
+	// Each caster has its own stack
+	AggregateBySource,
+	// Each target has its own stack
+	AggregateByTarget,
 };
 
 /**
diff --git a/Engine/Source/Runtime/GameplayAbilities/Public/GameplayPrediction.h b/Engine/Source/Runtime/GameplayAbilities/Public/GameplayPrediction.h
--- a/Engine/Source/Runtime/GameplayAbilities/Public/GameplayPrediction.h
+++ b/Engine/Source/Runtime/GameplayAbilities/Public/GameplayPrediction.h
@@ -303,11 +303,11 @@ struct GAMEPLAYABILITIES_API FPredictionKey
 	UPROPERTY(NotReplicated)
 	bool bIsStale;
 
-	/** Construct a new prediction key with no dependancies */
+	/** Construct a new prediction key with no dependencies */
 	static FPredictionKey CreateNewPredictionKey(class UAbilitySystemComponent*);
 
-	/** Create a new dependant prediction key: keep our existing base or use the current key as the base. */
-	void GenerateDependantPredictionKey();
+	/** Create a new dependent prediction key: keep our existing base or use the current key as the base. */
+	void GenerateDependentPredictionKey();
 
 	/** Creates new delegate called only when this key is rejected. */
 	FPredictionKeyEvent& NewRejectedDelegate();
@@ -427,7 +427,7 @@ public:
 	static void Reject(FPredictionKey::KeyType Key);
 	static void CatchUpTo(FPredictionKey::KeyType Key);
 
-	static void AddDependancy(FPredictionKey::KeyType ThisKey, FPredictionKey::KeyType DependsOn);
+	static void AddDependency(FPredictionKey::KeyType ThisKey, FPredictionKey::KeyType DependsOn);
 
 private:
 
