Merge CL2399952 for gameplay effects:
- Temporarily swap back to checking duration value instead of policy to unbreak cooldowns

diff --git a/Engine/Source/Runtime/GameplayAbilities/Private/GameplayEffect.cpp b/Engine/Source/Runtime/GameplayAbilities/Private/GameplayEffect.cpp
--- a/Engine/Source/Runtime/GameplayAbilities/Private/GameplayEffect.cpp
+++ b/Engine/Source/Runtime/GameplayAbilities/Private/GameplayEffect.cpp
@@ -1731,8 +1731,7 @@ FActiveGameplayEffect* FActiveGameplayEffectsContainer::ApplyGameplayEffectSpec(
 	// Register Source and Target non snapshot capture delegates here
 	AppliedEffectSpec.CapturedRelevantAttributes.RegisterLinkedAggregatorCallbacks(AppliedActiveGE->Handle);
 	
-	// Calculate Duration mods if we have a real duration
-	if (AppliedEffectSpec.Def->DurationPolicy == EGameplayEffectDurationType::HasDuration && bSetDuration)
+	if (bSetDuration)
 	{
 		// Re-calculate the duration, as it could rely on target captured attributes
 		float DefCalcDuration = 0.f;
@@ -1743,25 +1742,29 @@ FActiveGameplayEffect* FActiveGameplayEffectsContainer::ApplyGameplayEffectSpec(
 
 		const float DurationBaseValue = AppliedEffectSpec.GetDuration();
 
-		float FinalDuration = ComputeModifiedDurationOfAppliedSpec(AppliedEffectSpec, DurationBaseValue);
-
-		// We cannot mod ourselves into an instant or infinite duration effect
-		if (FinalDuration <= 0.f)
+		// Calculate Duration mods if we have a real duration
+		if (DurationBaseValue > 0.f)
 		{
-			ABILITY_LOG(Error, TEXT("GameplayEffect %s Duration was modified to %.2f. Clamping to 0.1s duration."), *AppliedEffectSpec.Def->GetName(), FinalDuration);
-			FinalDuration = 0.1f;
-		}
+			float FinalDuration = ComputeModifiedDurationOfAppliedSpec(AppliedEffectSpec, DurationBaseValue);
 
-		AppliedEffectSpec.SetDuration(FinalDuration, true);
+			// We cannot mod ourselves into an instant or infinite duration effect
+			if (FinalDuration <= 0.f)
+			{
+				ABILITY_LOG(Error, TEXT("GameplayEffect %s Duration was modified to %.2f. Clamping to 0.1s duration."), *AppliedEffectSpec.Def->GetName(), FinalDuration);
+				FinalDuration = 0.1f;
+			}
 
-		// ABILITY_LOG(Warning, TEXT("SetDuration for %s. Base: %.2f, Final: %.2f"), *NewEffect.Spec.Def->GetName(), DurationBaseValue, FinalDuration);
-		
-		// Register duration callbacks with the timer manager
-		if (Owner)
-		{
-			FTimerManager& TimerManager = Owner->GetWorld()->GetTimerManager();
-			FTimerDelegate Delegate = FTimerDelegate::CreateUObject(Owner, &UAbilitySystemComponent::CheckDurationExpired, AppliedActiveGE->Handle);
-			TimerManager.SetTimer(AppliedActiveGE->DurationHandle, Delegate, FinalDuration, false);
+			AppliedEffectSpec.SetDuration(FinalDuration, true);
+
+			// ABILITY_LOG(Warning, TEXT("SetDuration for %s. Base: %.2f, Final: %.2f"), *NewEffect.Spec.Def->GetName(), DurationBaseValue, FinalDuration);
+
+			// Register duration callbacks with the timer manager
+			if (Owner)
+			{
+				FTimerManager& TimerManager = Owner->GetWorld()->GetTimerManager();
+				FTimerDelegate Delegate = FTimerDelegate::CreateUObject(Owner, &UAbilitySystemComponent::CheckDurationExpired, AppliedActiveGE->Handle);
+				TimerManager.SetTimer(AppliedActiveGE->DurationHandle, Delegate, FinalDuration, false);
+			}
 		}
 	}
 	
