- Add support for handling stack count for executions
- Add new output structure for executions, move decision to trigger conditional effects from return value to inside the struct
- Executions can decide whether they want to manually handle the stack count or allow the GE system to default to multiplying the magnitude of the emitted modifiers by the stack count

diff --git a/Engine/Source/Runtime/GameplayAbilities/Private/GameplayEffect.cpp b/Engine/Source/Runtime/GameplayAbilities/Private/GameplayEffect.cpp
--- a/Engine/Source/Runtime/GameplayAbilities/Private/GameplayEffect.cpp
+++ b/Engine/Source/Runtime/GameplayAbilities/Private/GameplayEffect.cpp
@@ -1103,12 +1103,12 @@ void FActiveGameplayEffectsContainer::ExecuteActiveEffectsFrom(FGameplayEffectSp
 			const UGameplayEffectExecutionCalculation* ExecCDO = CurExecDef.CalculationClass->GetDefaultObject<UGameplayEffectExecutionCalculation>();
 			check(ExecCDO);
 
-			TArray<FGameplayModifierEvaluatedData> OutModifiers;
-
 			// Run the custom execution
 			FGameplayEffectCustomExecutionParameters ExecutionParams(SpecToUse, CurExecDef.CalculationModifiers, Owner);
-			bool bRunConditionalEffects = ExecCDO->Execute(ExecutionParams, OutModifiers);
+			FGameplayEffectCustomExecutionOutput ExecutionOutput;
+			ExecCDO->Execute(ExecutionParams, ExecutionOutput);
 
+			const bool bRunConditionalEffects = ExecutionOutput.ShouldTriggerConditionalGameplayEffects();
 			if (bRunConditionalEffects)
 			{
 				// If successful, apply conditional specs
@@ -1121,8 +1121,19 @@ void FActiveGameplayEffectsContainer::ExecuteActiveEffectsFrom(FGameplayEffectSp
 			}
 
 			// Execute any mods the custom execution yielded
+			TArray<FGameplayModifierEvaluatedData> OutModifiers;
+			ExecutionOutput.GetOutputModifiers(OutModifiers);
+
+			const bool bApplyStackCountToEmittedMods = !ExecutionOutput.IsStackCountHandledManually();
+			const int32 SpecStackCount = SpecToUse.StackCount;
+
 			for (FGameplayModifierEvaluatedData& CurExecMod : OutModifiers)
 			{
+				// If the execution didn't manually handle the stack count, automatically apply it here
+				if (bApplyStackCountToEmittedMods && SpecStackCount > 1)
+				{
+					CurExecMod.Magnitude *= SpecStackCount;
+				}
 				InvokeGameplayCueExecute |= InternalExecuteMod(SpecToUse, CurExecMod);
 			}
 		}
diff --git a/Engine/Source/Runtime/GameplayAbilities/Private/GameplayEffectExecutionCalculation.cpp b/Engine/Source/Runtime/GameplayAbilities/Private/GameplayEffectExecutionCalculation.cpp
--- a/Engine/Source/Runtime/GameplayAbilities/Private/GameplayEffectExecutionCalculation.cpp
+++ b/Engine/Source/Runtime/GameplayAbilities/Private/GameplayEffectExecutionCalculation.cpp
@@ -50,19 +50,12 @@ const FGameplayEffectSpec& FGameplayEffectCustomExecutionParameters::GetOwningSp
 	return *OwningSpec;
 }
 
-FGameplayEffectSpec& FGameplayEffectCustomExecutionParameters::GetOwningSpec()
-{
-	check(OwningSpec);
-	return *OwningSpec;
-}
-
-
 UAbilitySystemComponent* FGameplayEffectCustomExecutionParameters::GetTargetAbilitySystemComponent() const
 {
 	return TargetAbilitySystemComponent.Get();
 }
 
-UAbilitySystemComponent* FGameplayEffectCustomExecutionParameters::GetSourcebilitySystemComponent() const
+UAbilitySystemComponent* FGameplayEffectCustomExecutionParameters::GetSourceAbilitySystemComponent() const
 {
 	check(OwningSpec);
 	return OwningSpec->GetContext().GetInstigatorAbilitySystemComponent();
@@ -178,6 +171,47 @@ bool FGameplayEffectCustomExecutionParameters::AttemptGetCapturedAttributeAggreg
 	return false;
 }
 
+FGameplayEffectCustomExecutionOutput::FGameplayEffectCustomExecutionOutput()
+	: bTriggerConditionalGameplayEffects(false)
+	, bHandledStackCountManually(false)
+{
+}
+
+void FGameplayEffectCustomExecutionOutput::MarkStackCountHandledManually()
+{
+	bHandledStackCountManually = true;
+}
+
+bool FGameplayEffectCustomExecutionOutput::IsStackCountHandledManually() const
+{
+	return bHandledStackCountManually;
+}
+
+void FGameplayEffectCustomExecutionOutput::MarkConditionalGameplayEffectsToTrigger()
+{
+	bTriggerConditionalGameplayEffects = true;
+}
+
+bool FGameplayEffectCustomExecutionOutput::ShouldTriggerConditionalGameplayEffects() const
+{
+	return bTriggerConditionalGameplayEffects;
+}
+
+void FGameplayEffectCustomExecutionOutput::AddOutputModifier(const FGameplayModifierEvaluatedData& InOutputMod)
+{
+	OutputModifiers.Add(InOutputMod);
+}
+
+const TArray<FGameplayModifierEvaluatedData>& FGameplayEffectCustomExecutionOutput::GetOutputModifiers() const
+{
+	return OutputModifiers;
+}
+
+void FGameplayEffectCustomExecutionOutput::GetOutputModifiers(OUT TArray<FGameplayModifierEvaluatedData>& OutOutputModifiers) const
+{
+	OutOutputModifiers.Append(OutputModifiers);
+}
+
 UGameplayEffectExecutionCalculation::UGameplayEffectExecutionCalculation(const FObjectInitializer& ObjectInitializer)
 	: Super(ObjectInitializer)
 {
@@ -199,7 +233,6 @@ void UGameplayEffectExecutionCalculation::GetValidScopedModifierAttributeCapture
 }
 #endif // #if WITH_EDITORONLY_DATA
 
-bool UGameplayEffectExecutionCalculation::Execute_Implementation(FGameplayEffectCustomExecutionParameters& ExecutionParams, OUT TArray<FGameplayModifierEvaluatedData>& OutAdditionalModifiers) const
+void UGameplayEffectExecutionCalculation::Execute_Implementation(const FGameplayEffectCustomExecutionParameters& ExecutionParams, OUT FGameplayEffectCustomExecutionOutput& OutExecutionOutput) const
 {
-	return true;
 }
diff --git a/Engine/Source/Runtime/GameplayAbilities/Public/GameplayEffectExecutionCalculation.h b/Engine/Source/Runtime/GameplayAbilities/Public/GameplayEffectExecutionCalculation.h
--- a/Engine/Source/Runtime/GameplayAbilities/Public/GameplayEffectExecutionCalculation.h
+++ b/Engine/Source/Runtime/GameplayAbilities/Public/GameplayEffectExecutionCalculation.h
@@ -25,14 +25,11 @@ public:
 	/** Simple accessor to owning gameplay spec */
 	const FGameplayEffectSpec& GetOwningSpec() const;
 
-	/** Simple non-const accessor to owning gameplay spec */
-	FGameplayEffectSpec& GetOwningSpec();
-
 	/** Simple accessor to target ability system component */
 	UAbilitySystemComponent* GetTargetAbilitySystemComponent() const;
 
 	/** Simple accessor to source ability system component (could be null!) */
-	UAbilitySystemComponent* GetSourcebilitySystemComponent() const;
+	UAbilitySystemComponent* GetSourceAbilitySystemComponent() const;
 	
 	/**
 	 * Attempts to calculate the magnitude of a captured attribute given the specified parameters. Can fail if the gameplay spec doesn't have
@@ -105,6 +102,53 @@ private:
 	TWeakObjectPtr<UAbilitySystemComponent> TargetAbilitySystemComponent;
 };
 
+/** Struct representing the output of a custom gameplay effect execution. */
+USTRUCT()
+struct GAMEPLAYABILITIES_API FGameplayEffectCustomExecutionOutput
+{
+	GENERATED_USTRUCT_BODY()
+
+public:
+
+	/** Constructor */
+	FGameplayEffectCustomExecutionOutput();
+
+	/** Mark that the execution has manually handled the stack count and the GE system should not attempt to automatically act upon it for emitted modifiers */
+	void MarkStackCountHandledManually();
+
+	/** Simple accessor for determining whether the execution has manually handled the stack count or not */
+	bool IsStackCountHandledManually() const;
+
+	/** Mark that the execution wants conditional gameplay effects to trigger */
+	void MarkConditionalGameplayEffectsToTrigger();
+
+	/** Simple accessor for determining whether the execution wants conditional gameplay effects to trigger or not */
+	bool ShouldTriggerConditionalGameplayEffects() const;
+
+	/** Add the specified evaluated data to the execution's output modifiers */
+	void AddOutputModifier(const FGameplayModifierEvaluatedData& InOutputMod);
+
+	/** Simple accessor to output modifiers of the execution */
+	const TArray<FGameplayModifierEvaluatedData>& GetOutputModifiers() const;
+
+	/** Simple accessor to output modifiers of the execution */
+	void GetOutputModifiers(OUT TArray<FGameplayModifierEvaluatedData>& OutOutputModifiers) const;
+
+private:
+
+	/** Modifiers emitted by the execution */
+	UPROPERTY()
+	TArray<FGameplayModifierEvaluatedData> OutputModifiers;
+
+	/** If true, the execution wants to trigger conditional gameplay effects when it completes */
+	UPROPERTY()
+	uint32 bTriggerConditionalGameplayEffects : 1;
+	
+	/** If true, the execution itself has manually handled the stack count of the effect and the GE system doesn't have to automatically handle it */
+	UPROPERTY()
+	uint32 bHandledStackCountManually : 1;
+};
+
 UCLASS(BlueprintType, Blueprintable, Abstract)
 class GAMEPLAYABILITIES_API UGameplayEffectExecutionCalculation : public UGameplayEffectCalculation
 {
@@ -134,10 +178,9 @@ public:
 	 * 
 	 * @note: Native subclasses should override the auto-generated Execute_Implementation function and NOT this one.
 	 * 
-	 * @param ExecutionParams			Parameters for the custom execution calculation
-	 * @param OutAdditionalModifiers	[OUT] Additional modifiers the custom execution has generated and would like executed upon the target
-	 * @return							True if this was successful and it should apply conditional effects
+	 * @param ExecutionParams		Parameters for the custom execution calculation
+	 * @param OutExecutionOutput	[OUT] Output data populated by the execution detailing further behavior or results of the execution
 	 */
 	UFUNCTION(BlueprintNativeEvent, Category="Calculation")
-	bool Execute(FGameplayEffectCustomExecutionParameters& ExecutionParams, TArray<FGameplayModifierEvaluatedData>& OutAdditionalModifiers) const;
+	void Execute(const FGameplayEffectCustomExecutionParameters& ExecutionParams, FGameplayEffectCustomExecutionOutput& OutExecutionOutput) const;
 };
