Add bTriggerOnce to WaitAttributeChange task

diff --git a/Engine/Source/Runtime/GameplayAbilities/Private/Abilities/Tasks/AbilityTask_WaitAttributeChange.cpp b/Engine/Source/Runtime/GameplayAbilities/Private/Abilities/Tasks/AbilityTask_WaitAttributeChange.cpp
--- a/Engine/Source/Runtime/GameplayAbilities/Private/Abilities/Tasks/AbilityTask_WaitAttributeChange.cpp
+++ b/Engine/Source/Runtime/GameplayAbilities/Private/Abilities/Tasks/AbilityTask_WaitAttributeChange.cpp
@@ -11,20 +11,22 @@
 UAbilityTask_WaitAttributeChange::UAbilityTask_WaitAttributeChange(const FObjectInitializer& ObjectInitializer)
 	: Super(ObjectInitializer)
 {
+	bTriggerOnce = false;
 }
 
-UAbilityTask_WaitAttributeChange* UAbilityTask_WaitAttributeChange::WaitForAttributeChange(UObject* WorldContextObject, FGameplayAttribute InAttribute, FGameplayTag InWithTag, FGameplayTag InWithoutTag)
+UAbilityTask_WaitAttributeChange* UAbilityTask_WaitAttributeChange::WaitForAttributeChange(UObject* WorldContextObject, FGameplayAttribute InAttribute, FGameplayTag InWithTag, FGameplayTag InWithoutTag, bool TriggerOnce)
 {
 	auto MyObj = NewTask<UAbilityTask_WaitAttributeChange>(WorldContextObject);
 	MyObj->WithTag = InWithTag;
 	MyObj->WithoutTag = InWithoutTag;
 	MyObj->Attribute = InAttribute;
 	MyObj->ComparisonType = EWaitAttributeChangeComparison::None;
+	MyObj->bTriggerOnce = TriggerOnce;
 
 	return MyObj;
 }
 
-UAbilityTask_WaitAttributeChange* UAbilityTask_WaitAttributeChange::WaitForAttributeChangeWithComparison(UObject* WorldContextObject, FGameplayAttribute InAttribute, FGameplayTag InWithTag, FGameplayTag InWithoutTag, TEnumAsByte<EWaitAttributeChangeComparison::Type> InComparisonType, float InComparisonValue)
+UAbilityTask_WaitAttributeChange* UAbilityTask_WaitAttributeChange::WaitForAttributeChangeWithComparison(UObject* WorldContextObject, FGameplayAttribute InAttribute, FGameplayTag InWithTag, FGameplayTag InWithoutTag, TEnumAsByte<EWaitAttributeChangeComparison::Type> InComparisonType, float InComparisonValue, bool TriggerOnce)
 {
 	auto MyObj = NewTask<UAbilityTask_WaitAttributeChange>(WorldContextObject);
 	MyObj->WithTag = InWithTag;
@@ -32,6 +34,7 @@ UAbilityTask_WaitAttributeChange* UAbilityTask_WaitAttributeChange::WaitForAttri
 	MyObj->Attribute = InAttribute;
 	MyObj->ComparisonType = InComparisonType;
 	MyObj->ComparisonValue = InComparisonValue;
+	MyObj->bTriggerOnce = TriggerOnce;
 
 	return MyObj;
 }
@@ -92,7 +95,10 @@ void UAbilityTask_WaitAttributeChange::OnAttributeChange(float NewValue, const F
 	if (PassedComparison)
 	{
 		OnChange.Broadcast();
-		EndTask();
+		if (bTriggerOnce)
+		{
+			EndTask();
+		}
 	}
 }
 
diff --git a/Engine/Source/Runtime/GameplayAbilities/Public/Abilities/Tasks/AbilityTask_WaitAttributeChange.h b/Engine/Source/Runtime/GameplayAbilities/Public/Abilities/Tasks/AbilityTask_WaitAttributeChange.h
--- a/Engine/Source/Runtime/GameplayAbilities/Public/Abilities/Tasks/AbilityTask_WaitAttributeChange.h
+++ b/Engine/Source/Runtime/GameplayAbilities/Public/Abilities/Tasks/AbilityTask_WaitAttributeChange.h
@@ -45,20 +45,21 @@ class UAbilityTask_WaitAttributeChange : public UAbilityTask
 
 	/** Wait until an attribute changes. */
 	UFUNCTION(BlueprintCallable, Category="Ability|Tasks", meta = (HidePin = "WorldContextObject", DefaultToSelf = "WorldContextObject", BlueprintInternalUseOnly = "TRUE"))
-	static UAbilityTask_WaitAttributeChange* WaitForAttributeChange(UObject* WorldContextObject, FGameplayAttribute Attribute, FGameplayTag WithSrcTag, FGameplayTag WithoutSrcTag);
+	static UAbilityTask_WaitAttributeChange* WaitForAttributeChange(UObject* WorldContextObject, FGameplayAttribute Attribute, FGameplayTag WithSrcTag, FGameplayTag WithoutSrcTag, bool TriggerOnce=true);
 
 	/** Wait until an attribute changes to pass a given test. */
 	UFUNCTION(BlueprintCallable, Category = "Ability|Tasks", meta = (HidePin = "WorldContextObject", DefaultToSelf = "WorldContextObject", BlueprintInternalUseOnly = "TRUE"))
-	static UAbilityTask_WaitAttributeChange* WaitForAttributeChangeWithComparison(UObject* WorldContextObject, FGameplayAttribute InAttribute, FGameplayTag InWithTag, FGameplayTag InWithoutTag, TEnumAsByte<EWaitAttributeChangeComparison::Type> InComparisonType, float InComparisonValue);
+	static UAbilityTask_WaitAttributeChange* WaitForAttributeChangeWithComparison(UObject* WorldContextObject, FGameplayAttribute InAttribute, FGameplayTag InWithTag, FGameplayTag InWithoutTag, TEnumAsByte<EWaitAttributeChangeComparison::Type> InComparisonType, float InComparisonValue, bool TriggerOnce=true);
 
 	FGameplayTag WithTag;
 	FGameplayTag WithoutTag;
 	FGameplayAttribute	Attribute;
 	TEnumAsByte<EWaitAttributeChangeComparison::Type> ComparisonType;
 	float ComparisonValue;
+	bool bTriggerOnce;
 	FDelegateHandle OnAttributeChangeDelegateHandle;
 
 protected:
 
 	virtual void OnDestroy(bool AbilityEnded) override;
-};
\ No newline at end of file
+};
