ScopedLock protection for GameplayEffect array

diff --git a/Engine/Source/Runtime/GameplayAbilities/Private/Abilities/GameplayAbility.cpp b/Engine/Source/Runtime/GameplayAbilities/Private/Abilities/GameplayAbility.cpp
--- a/Engine/Source/Runtime/GameplayAbilities/Private/Abilities/GameplayAbility.cpp
+++ b/Engine/Source/Runtime/GameplayAbilities/Private/Abilities/GameplayAbility.cpp
@@ -531,7 +531,7 @@ float UGameplayAbility::GetCooldownTimeRemaining(const FGameplayAbilityActorInfo
 
 	check(ActorInfo->AbilitySystemComponent.IsValid());
 	const FGameplayTagContainer* CooldownTags = GetCooldownTags();
-	if (CooldownTags)
+	if (CooldownTags && CooldownTags->Num() > 0)
 	{
 		TArray< float > Durations = ActorInfo->AbilitySystemComponent->GetActiveEffectsTimeRemaining(FActiveGameplayEffectQuery(CooldownTags));
 		if (Durations.Num() > 0)
@@ -554,7 +554,7 @@ void UGameplayAbility::GetCooldownTimeRemainingAndDuration(FGameplayAbilitySpecH
 	CooldownDuration = 0.f;
 	
 	const FGameplayTagContainer* CooldownTags = GetCooldownTags();
-	if (CooldownTags)
+	if (CooldownTags && CooldownTags->Num() > 0)
 	{
 		TArray< float > DurationRemaining = ActorInfo->AbilitySystemComponent->GetActiveEffectsTimeRemaining(FActiveGameplayEffectQuery(CooldownTags));
 		if (DurationRemaining.Num() > 0)
@@ -580,7 +580,7 @@ void UGameplayAbility::GetCooldownTimeRemainingAndDuration(FGameplayAbilitySpecH
 
 const FGameplayTagContainer* UGameplayAbility::GetCooldownTags() const
 {
-	UGameplayEffect* CDGE =GetCooldownGameplayEffect();
+	UGameplayEffect* CDGE = GetCooldownGameplayEffect();
 	return CDGE ? &CDGE->InheritableOwnedTagsContainer.CombinedTags : nullptr;
 }
 
diff --git a/Engine/Source/Runtime/GameplayAbilities/Private/AbilitySystemComponent.cpp b/Engine/Source/Runtime/GameplayAbilities/Private/AbilitySystemComponent.cpp
--- a/Engine/Source/Runtime/GameplayAbilities/Private/AbilitySystemComponent.cpp
+++ b/Engine/Source/Runtime/GameplayAbilities/Private/AbilitySystemComponent.cpp
@@ -640,20 +640,7 @@ FActiveGameplayEffectHandle UAbilitySystemComponent::ApplyGameplayEffectSpecToSe
 
 void UAbilitySystemComponent::ExecutePeriodicEffect(FActiveGameplayEffectHandle	Handle)
 {
-	FActiveGameplayEffect* GameplayEffect = ActiveGameplayEffects.GetActiveGameplayEffect(Handle);
-	if (GameplayEffect)
-	{
-		ABILITY_VLOG(OwnerActor, Log, TEXT("Executed Periodic Effect %s"), *GameplayEffect->Spec.Def->GetFName().ToString());
-
-		for (FGameplayModifierInfo Modifier : GameplayEffect->Spec.Def->Modifiers)
-		{
-			float Magnitude = 0.f;
-			Modifier.ModifierMagnitude.AttemptCalculateMagnitude(GameplayEffect->Spec, Magnitude);
-			ABILITY_VLOG(OwnerActor, Log, TEXT("         %s: %s %f"), *Modifier.Attribute.GetName(), *EGameplayModOpToString(Modifier.ModifierOp), Magnitude);
-		}
-
-		ActiveGameplayEffects.ExecutePeriodicGameplayEffect(Handle);
-	}
+	ActiveGameplayEffects.ExecutePeriodicGameplayEffect(Handle);
 }
 
 void UAbilitySystemComponent::ExecuteGameplayEffect(FGameplayEffectSpec &Spec, FPredictionKey PredictionKey)
@@ -681,19 +668,6 @@ void UAbilitySystemComponent::CheckDurationExpired(FActiveGameplayEffectHandle H
 
 bool UAbilitySystemComponent::RemoveActiveGameplayEffect(FActiveGameplayEffectHandle Handle, int32 StacksToRemove)
 {
-	FActiveGameplayEffect* GameplayEffect = ActiveGameplayEffects.GetActiveGameplayEffect(Handle);
-	if (GameplayEffect)
-	{
-		ABILITY_VLOG(OwnerActor, Log, TEXT("Removed %s"), *GameplayEffect->Spec.Def->GetFName().ToString());
-
-		for (FGameplayModifierInfo Modifier : GameplayEffect->Spec.Def->Modifiers)
-		{
-			float Magnitude = 0.f;
-			Modifier.ModifierMagnitude.AttemptCalculateMagnitude(GameplayEffect->Spec, Magnitude);
-			ABILITY_VLOG(OwnerActor, Log, TEXT("         %s: %s %f"), *Modifier.Attribute.GetName(), *EGameplayModOpToString(Modifier.ModifierOp), Magnitude);
-		}
-	}
-
 	return ActiveGameplayEffects.RemoveActiveGameplayEffect(Handle, StacksToRemove);
 }
 
diff --git a/Engine/Source/Runtime/GameplayAbilities/Private/GameplayEffect.cpp b/Engine/Source/Runtime/GameplayAbilities/Private/GameplayEffect.cpp
--- a/Engine/Source/Runtime/GameplayAbilities/Private/GameplayEffect.cpp
+++ b/Engine/Source/Runtime/GameplayAbilities/Private/GameplayEffect.cpp
@@ -1044,6 +1044,13 @@ void FActiveGameplayEffect::PostReplicatedChange(const struct FActiveGameplayEff
 //	FActiveGameplayEffectsContainer
 //
 // --------------------------------------------------------------------------------------------------------------------------------------------------------
+// 
+FActiveGameplayEffectsContainer::FActiveGameplayEffectsContainer()
+	: Owner(nullptr)
+	, ScopedLockCount(0)
+	, PendingRemoves(0)
+{
+}
 
 void FActiveGameplayEffectsContainer::RegisterWithOwner(UAbilitySystemComponent* InOwner)
 {
@@ -1168,10 +1175,21 @@ void FActiveGameplayEffectsContainer::ExecuteActiveEffectsFrom(FGameplayEffectSp
 
 void FActiveGameplayEffectsContainer::ExecutePeriodicGameplayEffect(FActiveGameplayEffectHandle Handle)
 {
-	FScopedActiveGameplayEffectLock AGELock;
+	GAMEPLAYEFFECT_SCOPE_LOCK();
 	FActiveGameplayEffect* ActiveEffect = GetActiveGameplayEffect(Handle);
 	if (ActiveEffect)
 	{
+		if (UE_LOG_ACTIVE(VLogAbilitySystem, Log))
+		{
+			ABILITY_VLOG(Owner->OwnerActor, Log, TEXT("Executed Periodic Effect %s"), *ActiveEffect->Spec.Def->GetFName().ToString());
+			for (FGameplayModifierInfo Modifier : ActiveEffect->Spec.Def->Modifiers)
+			{
+				float Magnitude = 0.f;
+				Modifier.ModifierMagnitude.AttemptCalculateMagnitude(ActiveEffect->Spec, Magnitude);
+				ABILITY_VLOG(Owner->OwnerActor, Log, TEXT("         %s: %s %f"), *Modifier.Attribute.GetName(), *EGameplayModOpToString(Modifier.ModifierOp), Magnitude);
+			}
+		}
+
 		// Execute
 		ExecuteActiveEffectsFrom(ActiveEffect->Spec);
 	}
@@ -1190,29 +1208,9 @@ FActiveGameplayEffect* FActiveGameplayEffectsContainer::GetActiveGameplayEffect(
 	return nullptr;
 }
 
-int32 FScopedActiveGameplayEffectLock::AGELockCount = 0;
-TArray<TSharedPtr<FActiveGameplayEffectAction>> FScopedActiveGameplayEffectLock::DeferredAGEActions;
-FScopedActiveGameplayEffectLock::FScopedActiveGameplayEffectLock()
-{
-	++AGELockCount;
-}
-
-FScopedActiveGameplayEffectLock::~FScopedActiveGameplayEffectLock()
-{
-	--AGELockCount;
-	if (!IsLockInEffect())
-	{
-		for (int32 i = 0; i < NumActions(); ++i)
-		{
-			GetAction(i)->PerformAction();
-		}
-		ClearActions();
-	}
-}
-
 FAggregatorRef& FActiveGameplayEffectsContainer::FindOrCreateAttributeAggregator(FGameplayAttribute Attribute)
 {
-	FAggregatorRef *RefPtr = AttributeAggregatorMap.Find(Attribute);
+	FAggregatorRef* RefPtr = AttributeAggregatorMap.Find(Attribute);
 	if (RefPtr)
 	{
 		return *RefPtr;
@@ -1292,6 +1290,7 @@ void FActiveGameplayEffectsContainer::OnMagnitudeDependencyChange(FActiveGamepla
 {
 	if (Handle.IsValid())
 	{
+		GAMEPLAYEFFECT_SCOPE_LOCK();
 		FActiveGameplayEffect* ActiveEffect = GetActiveGameplayEffect(Handle);
 		if (ActiveEffect)
 		{
@@ -1331,7 +1330,7 @@ void FActiveGameplayEffectsContainer::OnMagnitudeDependencyChange(FActiveGamepla
 void FActiveGameplayEffectsContainer::OnStackCountChange(FActiveGameplayEffect& ActiveEffect)
 {
 	const FGameplayEffectSpec& Spec = ActiveEffect.Spec;
-	TSet<FGameplayAttribute>	AttributesToUpdate;
+	TSet<FGameplayAttribute> AttributesToUpdate;
 	
 	for (int32 ModIdx = 0; ModIdx < Spec.Modifiers.Num(); ++ModIdx)
 	{
@@ -1474,7 +1473,6 @@ void FActiveGameplayEffectsContainer::SetBaseAttributeValueFromReplication(FGame
 
 float FActiveGameplayEffectsContainer::GetGameplayEffectDuration(FActiveGameplayEffectHandle Handle) const
 {
-	// Could make this a map for quicker lookup
 	for (int32 idx = 0; idx < GameplayEffects.Num(); ++idx)
 	{
 		if (GameplayEffects[idx].Handle == Handle)
@@ -1489,7 +1487,6 @@ float FActiveGameplayEffectsContainer::GetGameplayEffectDuration(FActiveGameplay
 
 float FActiveGameplayEffectsContainer::GetGameplayEffectMagnitude(FActiveGameplayEffectHandle Handle, FGameplayAttribute Attribute) const
 {
-	// Could make this a map for quicker lookup
 	for (FActiveGameplayEffect Effect : GameplayEffects)
 	{
 		if (Effect.Handle == Handle)
@@ -1654,6 +1651,8 @@ FActiveGameplayEffect* FActiveGameplayEffectsContainer::ApplyGameplayEffectSpec(
 {
 	SCOPE_CYCLE_COUNTER(STAT_ApplyGameplayEffectSpec);
 
+	GAMEPLAYEFFECT_SCOPE_LOCK();
+
 	if (Owner && Owner->OwnerActor && IsNetAuthority())
 	{
 		Owner->OwnerActor->FlushNetDormancy();
@@ -1727,6 +1726,17 @@ FActiveGameplayEffect* FActiveGameplayEffectsContainer::ApplyGameplayEffectSpec(
 	}
 	else
 	{
+		if (ScopedLockCount > 0 && GameplayEffects.GetSlack() <= 0)
+		{
+			// TODO: Instead of failing we could add to a pending add list. For that to work though all other places that look at GameplayEffects list, they will need to look at
+			// the pending list as well.
+			
+			ABILITY_LOG(Error, TEXT("******************************************************"));
+			ABILITY_LOG(Error, TEXT("FActiveGameplayEffectsContainer::ApplyGameplayEffectSpec is attemping to add GameplayEffect %s while locked and having no more slack!"));
+			ABILITY_LOG(Error, TEXT("******************************************************"));
+			ensure(false);
+		}
+
 		FActiveGameplayEffectHandle NewHandle = FActiveGameplayEffectHandle::GenerateNewHandle(Owner);
 		AppliedActiveGE = new(GameplayEffects) FActiveGameplayEffect(NewHandle, Spec, GetWorldTime(), GetGameStateTime(), InPredictionKey);
 	}
@@ -1828,6 +1838,8 @@ FActiveGameplayEffect* FActiveGameplayEffectsContainer::ApplyGameplayEffectSpec(
 /** This is called anytime a new ActiveGameplayEffect is added, on both client and server in all cases */
 void FActiveGameplayEffectsContainer::InternalOnActiveGameplayEffectAdded(FActiveGameplayEffect& Effect)
 {
+	GAMEPLAYEFFECT_SCOPE_LOCK();
+
 	// Add our ongoing tag requirements to the dependency map. We will actually check for these tags below.
 	for (const FGameplayTag& Tag : Effect.Spec.Def->OngoingTagRequirements.IgnoreTags)
 	{
@@ -1849,6 +1861,8 @@ void FActiveGameplayEffectsContainer::InternalOnActiveGameplayEffectAdded(FActiv
 
 void FActiveGameplayEffectsContainer::AddActiveGameplayEffectGrantedTagsAndModifiers(FActiveGameplayEffect& Effect)
 {
+	GAMEPLAYEFFECT_SCOPE_LOCK();
+
 	// Register this ActiveGameplayEffects modifiers with our Attribute Aggregators
 	if (Effect.Spec.GetPeriod() <= UGameplayEffect::NO_PERIOD)
 	{
@@ -1887,12 +1901,17 @@ bool FActiveGameplayEffectsContainer::RemoveActiveGameplayEffect(FActiveGameplay
 	{
 		if (GameplayEffects[Idx].Handle == Handle)
 		{
-			//Block removal if we're scope-locked
-			if (FScopedActiveGameplayEffectLock::IsLockInEffect())
+			if (UE_LOG_ACTIVE(VLogAbilitySystem, Log))
 			{
-				FScopedActiveGameplayEffectLock::AddAction()->InitForRemoveGE(TWeakObjectPtr<UAbilitySystemComponent>(Owner), Handle, StacksToRemove);
-				return true;		//We are assuming the "ensure(Idx < GameplayEffects.Num())" would be passed.
+				ABILITY_VLOG(Owner->OwnerActor, Log, TEXT("Removed %s"), *GameplayEffects[Idx].Spec.Def->GetFName().ToString());
+				for (FGameplayModifierInfo Modifier : GameplayEffects[Idx].Spec.Def->Modifiers)
+				{
+					float Magnitude = 0.f;
+					Modifier.ModifierMagnitude.AttemptCalculateMagnitude(GameplayEffects[Idx].Spec, Magnitude);
+					ABILITY_VLOG(Owner->OwnerActor, Log, TEXT("         %s: %s %f"), *Modifier.Attribute.GetName(), *EGameplayModOpToString(Modifier.ModifierOp), Magnitude);
+				}
 			}
+
 			InternalRemoveActiveGameplayEffect(Idx, StacksToRemove);
 			return true;
 		}
@@ -1905,8 +1924,6 @@ bool FActiveGameplayEffectsContainer::RemoveActiveGameplayEffect(FActiveGameplay
 bool FActiveGameplayEffectsContainer::InternalRemoveActiveGameplayEffect(int32 Idx, int32 StacksToRemove)
 {
 	SCOPE_CYCLE_COUNTER(STAT_RemoveActiveGameplayEffect);
-
-	check(!FScopedActiveGameplayEffectLock::IsLockInEffect());
 	
 	if (ensure(Idx < GameplayEffects.Num()))
 	{
@@ -1955,16 +1972,29 @@ bool FActiveGameplayEffectsContainer::InternalRemoveActiveGameplayEffect(int32 I
 		// Remove this handle from the global map
 		Effect.Handle.RemoveFromGlobalMap();
 
+		bool ModifiedArray = false;
+
 		// Finally remove the ActiveGameplayEffect
-		GameplayEffects.RemoveAtSwap(Idx);
-		MarkArrayDirty();
+		if (ScopedLockCount > 0)
+		{
+			// We are locked, so this removal is now pending.
+			PendingRemoves++;
+			Effect.IsPendingRemove = true;
+		}
+		else
+		{
+			// Not locked, so do the removal right away.
+			GameplayEffects.RemoveAtSwap(Idx);
+			MarkArrayDirty();
+			ModifiedArray = true;
+		}
 
 		// Hack: force netupdate on owner. This isn't really necessary in real gameplay but is nice
 		// during debugging where breakpoints or pausing can mess up network update times. Open issue
 		// with network team.
 		Owner->GetOwner()->ForceNetUpdate();
 		
-		return true;
+		return ModifiedArray;
 	}
 
 	ABILITY_LOG(Warning, TEXT("InternalRemoveActiveGameplayEffect called with invalid index: %d"), Idx);
@@ -2046,6 +2076,8 @@ void FActiveGameplayEffectsContainer::OnOwnerTagChange(FGameplayTag TagChange, i
 	auto Ptr = ActiveEffectTagDependencies.Find(TagChange);
 	if (Ptr)
 	{
+		GAMEPLAYEFFECT_SCOPE_LOCK();
+
 		FGameplayTagContainer OwnerTags;
 		Owner->GetOwnedGameplayTags(OwnerTags);
 
@@ -2101,8 +2133,8 @@ void FActiveGameplayEffectsContainer::PreDestroy()
 
 int32 FActiveGameplayEffectsContainer::GetGameStateTime() const
 {
-	UWorld *World = Owner->GetWorld();
-	AGameState * GameState = World->GetGameState<AGameState>();
+	UWorld* World = Owner->GetWorld();
+	AGameState* GameState = World->GetGameState<AGameState>();
 	if (GameState)
 	{
 		return GameState->ElapsedTime;
@@ -2119,6 +2151,7 @@ float FActiveGameplayEffectsContainer::GetWorldTime() const
 
 void FActiveGameplayEffectsContainer::CheckDuration(FActiveGameplayEffectHandle Handle)
 {
+	GAMEPLAYEFFECT_SCOPE_LOCK();
 	for (int32 Idx = 0; Idx < GameplayEffects.Num(); ++Idx)
 	{
 		FActiveGameplayEffect& Effect = GameplayEffects[Idx];
@@ -2244,26 +2277,13 @@ TArray<float> FActiveGameplayEffectsContainer::GetActiveEffectsDuration(const FA
 }
 
 void FActiveGameplayEffectsContainer::RemoveActiveEffects(const FActiveGameplayEffectQuery Query, int32 StacksToRemove)
-{
-	const bool bLockInEffect = FScopedActiveGameplayEffectLock::IsLockInEffect();
-
-	for (int32 idx=0; idx < GameplayEffects.Num(); ++idx)
+{	
+	for (int32 idx=GameplayEffects.Num()-1; idx >= 0; --idx)
 	{
 		const FActiveGameplayEffect& Effect = GameplayEffects[idx];
 		if (Query.Matches(Effect))
 		{
-			// Defer removal if we're scope-locked
-			if (bLockInEffect)
-			{
-				FScopedActiveGameplayEffectLock::AddAction()->InitForRemoveGE(TWeakObjectPtr<UAbilitySystemComponent>(Owner), Effect.Handle, StacksToRemove);
-			}
-			else
-			{
-				if (InternalRemoveActiveGameplayEffect(idx, StacksToRemove))
-				{
-					idx--;
-				}
-			}
+			InternalRemoveActiveGameplayEffect(idx, StacksToRemove);			
 		}
 	}
 }
@@ -2474,22 +2494,44 @@ void FInheritedTagContainer::RemoveTag(FGameplayTag TagToRemove)
 	CombinedTags.RemoveTag(TagToRemove);
 }
 
-void FActiveGameplayEffectAction::PerformAction()
+FScopedActiveGameplayEffectLock::FScopedActiveGameplayEffectLock(FActiveGameplayEffectsContainer& InContainer)
+	: Container(InContainer)
 {
-	check(bIsInitialized);
-	if (OwningASC.IsValid())
+	if (Container.ScopedLockCount++ == 0)
 	{
-		if (bRemove)
+		// If we transitioned from unlocked to locked, make sure we have enough slack for additions.
+		// 4 is arbitrary determined. We should add instrumentation to get a better ideal/default number.
+		static const int32 GAMEPLAYEFFECT_MIN_SLACK = 4;
+
+		if (Container.GameplayEffects.GetSlack() < GAMEPLAYEFFECT_MIN_SLACK)
 		{
-			OwningASC.Get()->RemoveActiveGameplayEffect(Handle, StacksToRemove);
+			Container.GameplayEffects.Reserve( Container.GameplayEffects.Num() + GAMEPLAYEFFECT_MIN_SLACK );
 		}
-		else
+	}
+}
+
+FScopedActiveGameplayEffectLock::~FScopedActiveGameplayEffectLock()
+{
+	if (--Container.ScopedLockCount == 0 && Container.PendingRemoves > 0)
+	{
+		// If we are transitioned from locked to unlocked, we now need to actually delete and pending kill gameplay effects
+		for (int32 idx=Container.GameplayEffects.Num()-1; idx >= 0 && Container.PendingRemoves > 0; --idx)
+		{
+			FActiveGameplayEffect& Effect = Container.GameplayEffects[idx];
+
+			if (Effect.IsPendingRemove)
+			{
+				Container.GameplayEffects.RemoveAtSwap(idx, 1, false);
+				Container.PendingRemoves--;
+			}
+		}
+
+		if (!ensure(Container.PendingRemoves == 0))
 		{
-			UAbilitySystemComponent* ASC = OwningASC.Get();
-			check(!ASC->ActiveGameplayEffects.GameplayEffectsPendingAdd.IsEmpty());
-			FActiveGameplayEffect TempEffect;
-			ASC->ActiveGameplayEffects.GameplayEffectsPendingAdd.Dequeue(TempEffect);
-			ASC->ActiveGameplayEffects.GameplayEffects.Add(TempEffect);
+			ABILITY_LOG(Warning, TEXT("~FScopedActiveGameplayEffectLock has %d pending removes after a scope lock removal"), Container.PendingRemoves);
+			Container.PendingRemoves = 0;
 		}
+
+		Container.MarkArrayDirty();
 	}
 }
diff --git a/Engine/Source/Runtime/GameplayAbilities/Public/GameplayEffect.h b/Engine/Source/Runtime/GameplayAbilities/Public/GameplayEffect.h
--- a/Engine/Source/Runtime/GameplayAbilities/Public/GameplayEffect.h
+++ b/Engine/Source/Runtime/GameplayAbilities/Public/GameplayEffect.h
@@ -1084,6 +1084,7 @@ struct FActiveGameplayEffect : public FFastArraySerializerItem
 		, CachedStartGameStateTime(0)
 		, StartWorldTime(0.f)
 		, IsInhibited(true)
+		, IsPendingRemove(false)
 	{
 	}
 
@@ -1095,6 +1096,7 @@ struct FActiveGameplayEffect : public FFastArraySerializerItem
 		, CachedStartGameStateTime(InStartGameStateTime)
 		, StartWorldTime(CurrentWorldTime)
 		, IsInhibited(true)
+		, IsPendingRemove(false)
 	{
 	}
 
@@ -1121,6 +1123,8 @@ struct FActiveGameplayEffect : public FFastArraySerializerItem
 	UPROPERTY()
 	bool IsInhibited;
 
+	bool IsPendingRemove;
+
 	FOnActiveGameplayEffectRemoved	OnRemovedDelegate;
 
 	FTimerHandle PeriodHandle;
@@ -1259,16 +1263,13 @@ struct FActiveGameplayEffectsContainer : public FFastArraySerializer
 
 	friend struct FActiveGameplayEffect;
 	friend class UAbilitySystemComponent;
+	friend struct FScopedActiveGameplayEffectLock;
+	friend class AAbilitySystemDebugHUD;
 
-	FActiveGameplayEffectsContainer() {};
+	FActiveGameplayEffectsContainer();
 
 	UAbilitySystemComponent* Owner;
 
-	UPROPERTY()
-	TArray<FActiveGameplayEffect>	GameplayEffects;
-
-	TQueue<FActiveGameplayEffect>	GameplayEffectsPendingAdd;
-
 	void RegisterWithOwner(UAbilitySystemComponent* Owner);	
 	
 	FActiveGameplayEffect* ApplyGameplayEffectSpec(const FGameplayEffectSpec& Spec, FPredictionKey InPredictionKey);
@@ -1361,6 +1362,9 @@ struct FActiveGameplayEffectsContainer : public FFastArraySerializer
 
 private:
 
+	UPROPERTY()
+	TArray<FActiveGameplayEffect>	GameplayEffects;
+
 	void InternalUpdateNumericalAttribute(FGameplayAttribute Attribute, float NewValue, const FGameplayEffectModCallbackData* ModData);
 	
 	/**
@@ -1422,6 +1426,9 @@ private:
 
 	/** After application has gone through, give stacking rules a chance to do something as the source of the gameplay effect (E.g., remove an old version) */
 	void ApplyStackingLogicPostApplyAsSource(UAbilitySystemComponent* Target, const FGameplayEffectSpec& SpecApplied, FActiveGameplayEffectHandle ActiveHandle);
+
+	mutable int32 ScopedLockCount;
+	int32 PendingRemoves;
 };
 
 template<>
@@ -1434,89 +1441,14 @@ struct TStructOpsTypeTraits< FActiveGameplayEffectsContainer > : public TStructO
 };
 
 
-USTRUCT()
-struct GAMEPLAYABILITIES_API FActiveGameplayEffectAction
-{
-	GENERATED_USTRUCT_BODY()
-	FActiveGameplayEffectAction()
-	: bIsInitialized(false)
-	{
-	}
-
-	void InitForAddGE(TWeakObjectPtr<UAbilitySystemComponent> InOwningASC)
-	{
-		OwningASC = InOwningASC;
-		bRemove = false;
-		bIsInitialized = true;
-	}
-
-	void InitForRemoveGE(TWeakObjectPtr<UAbilitySystemComponent> InOwningASC, const FActiveGameplayEffectHandle& InHandle, int32 InStacksToRemove)
-	{
-		OwningASC = InOwningASC;
-		Handle = InHandle;
-		StacksToRemove = InStacksToRemove;
-		bRemove = true;
-		bIsInitialized = true;
-	}
-
-	void PerformAction();
-
-	UPROPERTY()
-	TWeakObjectPtr<UAbilitySystemComponent> OwningASC;
-
-	/** This is needed for removals, and does not exist for additions. */
-	UPROPERTY()
-	FActiveGameplayEffectHandle Handle;
-
-	UPROPERTY()
-	bool bRemove;
-
-	UPROPERTY()
-	bool bIsInitialized;
-
-	UPROPERTY()
-	int32 StacksToRemove;
-};
-
-
 struct GAMEPLAYABILITIES_API FScopedActiveGameplayEffectLock
 {
-	FScopedActiveGameplayEffectLock();
+	FScopedActiveGameplayEffectLock(FActiveGameplayEffectsContainer& InContainer);
 	~FScopedActiveGameplayEffectLock();
 
 private:
-	static int32 AGELockCount;
-	static TArray<TSharedPtr<FActiveGameplayEffectAction>> DeferredAGEActions;
-
-public:
-	static bool IsLockInEffect()
-	{
-		return (AGELockCount ? true : false);
-	}
-
-	static FActiveGameplayEffectAction* AddAction()
-	{
-		check(IsLockInEffect());
-		FActiveGameplayEffectAction* DataPtr = new FActiveGameplayEffectAction();
-		DeferredAGEActions.Add(TSharedPtr<FActiveGameplayEffectAction>(DataPtr));
-		return DataPtr;
-	}
-
-	void ClearActions()
-	{
-		DeferredAGEActions.Reset();
-	}
-
-	int32 NumActions()
-	{
-		return DeferredAGEActions.Num();
-	}
-
-	FActiveGameplayEffectAction* GetAction(int32 Index)
-	{
-		check((Index >= 0) && (Index < DeferredAGEActions.Num()));
-		return DeferredAGEActions[Index].Get();
-	}
+	FActiveGameplayEffectsContainer& Container;
 };
 
 
+#define GAMEPLAYEFFECT_SCOPE_LOCK()	FScopedActiveGameplayEffectLock ActiveScopeLock(*this);
\ No newline at end of file
