Fix non-unity build

diff --git a/Engine/Source/Runtime/GameplayAbilities/Private/GameplayEffect.cpp b/Engine/Source/Runtime/GameplayAbilities/Private/GameplayEffect.cpp
--- a/Engine/Source/Runtime/GameplayAbilities/Private/GameplayEffect.cpp
+++ b/Engine/Source/Runtime/GameplayAbilities/Private/GameplayEffect.cpp
@@ -9,6 +9,7 @@
 #include "AbilitySystemInterface.h"
 #include "GameplayModMagnitudeCalculation.h"
 #include "GameplayEffectExecutionCalculation.h"
+#include "VisualLogger.h"
 
 const float UGameplayEffect::INFINITE_DURATION = -1.f;
 const float UGameplayEffect::INSTANT_APPLICATION = 0.f;
