Integration(2406340) - Ability bug fixes;

- Changed 'MakeTargetLocationInfoFromOwnerSkeletalMeshComponent()' to be const.

- Removed 'EditAnywhere' from some 'AGameplayAbilityTargetActor' properties.

- Added warning to fix crash in 'UGameplayStatics::BeginSpawningActorFromClass()' when the world is NULL.

diff --git a/Engine/Source/Runtime/Engine/Private/GameplayStatics.cpp b/Engine/Source/Runtime/Engine/Private/GameplayStatics.cpp
--- a/Engine/Source/Runtime/Engine/Private/GameplayStatics.cpp
+++ b/Engine/Source/Runtime/Engine/Private/GameplayStatics.cpp
@@ -316,7 +316,14 @@ AActor* UGameplayStatics::BeginSpawningActorFromClass(UObject* WorldContextObjec
 		}
 
 		UWorld* World = GEngine->GetWorldFromContextObject(WorldContextObject);
-		NewActor = World->SpawnActorDeferred<AActor>(Class, SpawnLoc, SpawnRot, NULL, AutoInstigator, bNoCollisionFail);
+		if (World)
+		{
+			NewActor = World->SpawnActorDeferred<AActor>(Class, SpawnLoc, SpawnRot, NULL, AutoInstigator, bNoCollisionFail);
+		}
+		else
+		{
+			UE_LOG(LogScript, Warning, TEXT("UGameplayStatics::BeginSpawningActorFromClass: %s can not be spawned in NULL world"), *Class->GetName());		
+		}
 	}
 
 	return NewActor;
diff --git a/Engine/Source/Runtime/GameplayAbilities/Private/Abilities/GameplayAbility.cpp b/Engine/Source/Runtime/GameplayAbilities/Private/Abilities/GameplayAbility.cpp
--- a/Engine/Source/Runtime/GameplayAbilities/Private/Abilities/GameplayAbility.cpp
+++ b/Engine/Source/Runtime/GameplayAbilities/Private/Abilities/GameplayAbility.cpp
@@ -768,7 +768,7 @@ FGameplayAbilityTargetingLocationInfo UGameplayAbility::MakeTargetLocationInfoFr
 	return ReturnLocation;
 }
 
-FGameplayAbilityTargetingLocationInfo UGameplayAbility::MakeTargetLocationInfoFromOwnerSkeletalMeshComponent(FName SocketName)
+FGameplayAbilityTargetingLocationInfo UGameplayAbility::MakeTargetLocationInfoFromOwnerSkeletalMeshComponent(FName SocketName) const
 {
 	FGameplayAbilityTargetingLocationInfo ReturnLocation;
 	ReturnLocation.LocationType = EGameplayAbilityTargetingLocationType::SocketTransform;
diff --git a/Engine/Source/Runtime/GameplayAbilities/Public/Abilities/GameplayAbility.h b/Engine/Source/Runtime/GameplayAbilities/Public/Abilities/GameplayAbility.h
--- a/Engine/Source/Runtime/GameplayAbilities/Public/Abilities/GameplayAbility.h
+++ b/Engine/Source/Runtime/GameplayAbilities/Public/Abilities/GameplayAbility.h
@@ -668,7 +668,7 @@ protected:
 	FGameplayAbilityTargetingLocationInfo MakeTargetLocationInfoFromOwnerActor();
 
 	UFUNCTION(BlueprintPure, Category = Ability, meta = (HidePin = "WorldContextObject", DefaultToSelf = "WorldContextObject"))
-	FGameplayAbilityTargetingLocationInfo MakeTargetLocationInfoFromOwnerSkeletalMeshComponent(FName SocketName);
+	FGameplayAbilityTargetingLocationInfo MakeTargetLocationInfoFromOwnerSkeletalMeshComponent(FName SocketName) const;
 
 	// ----------------------------------------------------------------------------------------------------------------
 	//
diff --git a/Engine/Source/Runtime/GameplayAbilities/Public/Abilities/GameplayAbilityTargetActor.h b/Engine/Source/Runtime/GameplayAbilities/Public/Abilities/GameplayAbilityTargetActor.h
--- a/Engine/Source/Runtime/GameplayAbilities/Public/Abilities/GameplayAbilityTargetActor.h
+++ b/Engine/Source/Runtime/GameplayAbilities/Public/Abilities/GameplayAbilityTargetActor.h
@@ -81,14 +81,14 @@ public:
 	AActor* SourceActor;
 
 	/** Parameters for world reticle. Usage of these parameters is dependent on the reticle. */
-	UPROPERTY(BlueprintReadWrite, EditAnywhere, meta = (ExposeOnSpawn = true), Category = Targeting)
+	UPROPERTY(BlueprintReadWrite, meta = (ExposeOnSpawn = true), Category = Targeting)
 	FWorldReticleParameters ReticleParams;
 
 	/** Reticle that will appear on top of acquired targets. Reticles will be spawned/despawned as targets are acquired/lost. */
 	UPROPERTY(BlueprintReadWrite, EditAnywhere, meta = (ExposeOnSpawn = true), Category = Targeting)
 	TSubclassOf<AGameplayAbilityWorldReticle> ReticleClass;		//Using a special class for replication purposes.
 
-	UPROPERTY(BlueprintReadWrite, EditAnywhere, Replicated, meta = (ExposeOnSpawn = true), Category = Targeting)
+	UPROPERTY(BlueprintReadWrite, Replicated, meta = (ExposeOnSpawn = true), Category = Targeting)
 	FGameplayTargetDataFilterHandle Filter;
 
 	/** Draw the debug information (if applicable) for this targeting actor. */
