Integration(2406350) - Target actor tasks;

- Added ability tasks "VisualizeTargetingUsingActor' and 'WaitTargetDataUsingActor' for cases where we want to spawn target actors outside of the tasks.

- Cleaned up 'UAbilityTask_VisualizeTargeting' and 'UAbilityTask_WaitTargetData' classes.

diff --git a/Engine/Source/Runtime/GameplayAbilities/Private/Abilities/Tasks/AbilityTask_VisualizeTargeting.cpp b/Engine/Source/Runtime/GameplayAbilities/Private/Abilities/Tasks/AbilityTask_VisualizeTargeting.cpp
--- a/Engine/Source/Runtime/GameplayAbilities/Private/Abilities/Tasks/AbilityTask_VisualizeTargeting.cpp
+++ b/Engine/Source/Runtime/GameplayAbilities/Private/Abilities/Tasks/AbilityTask_VisualizeTargeting.cpp
@@ -14,51 +14,69 @@ UAbilityTask_VisualizeTargeting::UAbilityTask_VisualizeTargeting(const FObjectIn
 UAbilityTask_VisualizeTargeting* UAbilityTask_VisualizeTargeting::VisualizeTargeting(UObject* WorldContextObject, TSubclassOf<AGameplayAbilityTargetActor> InTargetClass, FName TaskInstanceName, float Duration)
 {
 	UAbilityTask_VisualizeTargeting* MyObj = NewTask<UAbilityTask_VisualizeTargeting>(WorldContextObject, TaskInstanceName);		//Register for task list here, providing a given FName as a key
+	MyObj->TargetClass = InTargetClass;
+	MyObj->TargetActor = NULL;
 	MyObj->SetDuration(Duration);
 	return MyObj;
 }
 
+UAbilityTask_VisualizeTargeting* UAbilityTask_VisualizeTargeting::VisualizeTargetingUsingActor(UObject* WorldContextObject, AGameplayAbilityTargetActor* InTargetActor, FName TaskInstanceName, float Duration)
+{
+	UAbilityTask_VisualizeTargeting* MyObj = NewTask<UAbilityTask_VisualizeTargeting>(WorldContextObject, TaskInstanceName);		//Register for task list here, providing a given FName as a key
+	MyObj->TargetClass = NULL;
+	MyObj->TargetActor = InTargetActor;
+	MyObj->SetDuration(Duration);
+	return MyObj;
+}
 
-void UAbilityTask_VisualizeTargeting::SetDuration(const float Duration)
+void UAbilityTask_VisualizeTargeting::Activate()
 {
-	if (Duration > 0.f)
+	if (Ability.IsValid() && TargetClass == NULL)
 	{
-		GetWorld()->GetTimerManager().SetTimer(TimerHandle_OnTimeElapsed, this, &UAbilityTask_VisualizeTargeting::OnTimeElapsed, Duration, false);
+		check(TargetActor.IsValid());
+		AGameplayAbilityTargetActor* SpawnedActor = TargetActor.Get();
+
+		TargetClass = SpawnedActor->GetClass();
+
+		if (ShouldSpawnTargetActor())
+		{
+			InitializeTargetActor(SpawnedActor);
+			FinalizeTargetActor(SpawnedActor);
+		}
+		else
+		{
+			TargetActor = NULL;
+
+			// We may need a better solution here.  We don't know the target actor isn't needed till after it's already been spawned.
+			SpawnedActor->Destroy();
+			SpawnedActor = nullptr;
+		}
 	}
 }
 
-bool UAbilityTask_VisualizeTargeting::BeginSpawningActor(UObject* WorldContextObject, TSubclassOf<AGameplayAbilityTargetActor> TargetClass, AGameplayAbilityTargetActor*& SpawnedActor)
+bool UAbilityTask_VisualizeTargeting::BeginSpawningActor(UObject* WorldContextObject, TSubclassOf<AGameplayAbilityTargetActor> InTargetClass, AGameplayAbilityTargetActor*& SpawnedActor)
 {
 	SpawnedActor = nullptr;
 
-	UGameplayAbility* MyAbility = Ability.Get();
-	if (MyAbility)
+	if (Ability.IsValid())
 	{
-		const AGameplayAbilityTargetActor* CDO = CastChecked<AGameplayAbilityTargetActor>(TargetClass->GetDefaultObject());
-
-		const bool bReplicates = CDO->GetReplicates();
-		const bool bIsLocallyControlled = MyAbility->GetCurrentActorInfo()->IsLocallyControlled();
-
-		// Spawn the actor if this is a locally controlled ability (always) or if this is a replicating targeting mode.
-		// (E.g., server will spawn this target actor to replicate to all non owning clients)
-		if (bReplicates || bIsLocallyControlled)
+		if (ShouldSpawnTargetActor())
 		{
-			UClass* Class = *TargetClass;
+			UClass* Class = *InTargetClass;
 			if (Class != NULL)
 			{
 				UWorld* World = GEngine->GetWorldFromContextObject(WorldContextObject);
 				SpawnedActor = World->SpawnActorDeferred<AGameplayAbilityTargetActor>(Class, FVector::ZeroVector, FRotator::ZeroRotator, NULL, NULL, true);
 			}
 
-			MyTargetActor = SpawnedActor;
-
-			AGameplayAbilityTargetActor* TargetActor = CastChecked<AGameplayAbilityTargetActor>(SpawnedActor);
-			if (TargetActor)
+			if (SpawnedActor)
 			{
-				TargetActor->MasterPC = MyAbility->GetCurrentActorInfo()->PlayerController.Get();
+				TargetActor = SpawnedActor;
+				InitializeTargetActor(SpawnedActor);
 			}
 		}
 	}
+
 	return (SpawnedActor != nullptr);
 }
 
@@ -66,21 +84,63 @@ void UAbilityTask_VisualizeTargeting::FinishSpawningActor(UObject* WorldContextO
 {
 	if (SpawnedActor)
 	{
-		check(MyTargetActor == SpawnedActor);
+		check(TargetActor == SpawnedActor);
 
 		const FTransform SpawnTransform = AbilitySystemComponent->GetOwner()->GetTransform();
 
 		SpawnedActor->FinishSpawning(SpawnTransform);
-		AbilitySystemComponent->SpawnedTargetActors.Push(SpawnedActor);
-		SpawnedActor->StartTargeting(Ability.Get());
+
+		FinalizeTargetActor(SpawnedActor);
+	}
+}
+
+void UAbilityTask_VisualizeTargeting::SetDuration(const float Duration)
+{
+	if (Duration > 0.f)
+	{
+		GetWorld()->GetTimerManager().SetTimer(TimerHandle_OnTimeElapsed, this, &UAbilityTask_VisualizeTargeting::OnTimeElapsed, Duration, false);
 	}
 }
 
+bool UAbilityTask_VisualizeTargeting::ShouldSpawnTargetActor() const
+{
+	check(TargetClass);
+	check(Ability.IsValid());
+
+	// Spawn the actor if this is a locally controlled ability (always) or if this is a replicating targeting mode.
+	// (E.g., server will spawn this target actor to replicate to all non owning clients)
+
+	const AGameplayAbilityTargetActor* CDO = CastChecked<AGameplayAbilityTargetActor>(TargetClass->GetDefaultObject());
+
+	const bool bReplicates = CDO->GetReplicates();
+	const bool bIsLocallyControlled = Ability->GetCurrentActorInfo()->IsLocallyControlled();
+
+	return (bReplicates || bIsLocallyControlled);
+}
+
+void UAbilityTask_VisualizeTargeting::InitializeTargetActor(AGameplayAbilityTargetActor* SpawnedActor) const
+{
+	check(SpawnedActor);
+	check(Ability.IsValid());
+
+	SpawnedActor->MasterPC = Ability->GetCurrentActorInfo()->PlayerController.Get();
+}
+
+void UAbilityTask_VisualizeTargeting::FinalizeTargetActor(AGameplayAbilityTargetActor* SpawnedActor) const
+{
+	check(SpawnedActor);
+	check(Ability.IsValid());
+
+	AbilitySystemComponent->SpawnedTargetActors.Push(SpawnedActor);
+
+	SpawnedActor->StartTargeting(Ability.Get());
+}
+
 void UAbilityTask_VisualizeTargeting::OnDestroy(bool AbilityEnded)
 {
-	if (MyTargetActor.IsValid())
+	if (TargetActor.IsValid())
 	{
-		MyTargetActor->Destroy();
+		TargetActor->Destroy();
 	}
 
 	GetWorld()->GetTimerManager().ClearTimer(TimerHandle_OnTimeElapsed);
diff --git a/Engine/Source/Runtime/GameplayAbilities/Private/Abilities/Tasks/AbilityTask_WaitTargetData.cpp b/Engine/Source/Runtime/GameplayAbilities/Private/Abilities/Tasks/AbilityTask_WaitTargetData.cpp
--- a/Engine/Source/Runtime/GameplayAbilities/Private/Abilities/Tasks/AbilityTask_WaitTargetData.cpp
+++ b/Engine/Source/Runtime/GameplayAbilities/Private/Abilities/Tasks/AbilityTask_WaitTargetData.cpp
@@ -15,76 +15,73 @@ UAbilityTask_WaitTargetData* UAbilityTask_WaitTargetData::WaitTargetData(UObject
 {
 	auto MyObj = NewTask<UAbilityTask_WaitTargetData>(WorldContextObject, TaskInstanceName);		//Register for task list here, providing a given FName as a key
 	MyObj->TargetClass = InTargetClass;
+	MyObj->TargetActor = NULL;
 	MyObj->ConfirmationType = ConfirmationType;
 	return MyObj;
 }
 
-bool UAbilityTask_WaitTargetData::BeginSpawningActor(UObject* WorldContextObject, TSubclassOf<AGameplayAbilityTargetActor> TargetClass, AGameplayAbilityTargetActor*& SpawnedActor)
+UAbilityTask_WaitTargetData* UAbilityTask_WaitTargetData::WaitTargetDataUsingActor(UObject* WorldContextObject, FName TaskInstanceName, TEnumAsByte<EGameplayTargetingConfirmation::Type> ConfirmationType, AGameplayAbilityTargetActor* InTargetActor)
 {
-	SpawnedActor = nullptr;
+	auto MyObj = NewTask<UAbilityTask_WaitTargetData>(WorldContextObject, TaskInstanceName);		//Register for task list here, providing a given FName as a key
+	MyObj->TargetClass = NULL;
+	MyObj->TargetActor = InTargetActor;
+	MyObj->ConfirmationType = ConfirmationType;
+	return MyObj;
+}
 
-	UGameplayAbility* MyAbility = Ability.Get();
-	if (MyAbility)
+void UAbilityTask_WaitTargetData::Activate()
+{
+	// Need to handle case where target actor was passed into task
+	if (Ability.IsValid() && TargetClass == NULL)
 	{
-		const AGameplayAbilityTargetActor* CDO = CastChecked<AGameplayAbilityTargetActor>(TargetClass->GetDefaultObject());
+		check(TargetActor.IsValid());
+		AGameplayAbilityTargetActor* SpawnedActor = TargetActor.Get();
 
-		const bool bReplicates = CDO->GetReplicates();
-		const bool bIsLocallyControlled = MyAbility->GetCurrentActorInfo()->IsLocallyControlled();
+		TargetClass = SpawnedActor->GetClass();
 
-		// Spawn the actor if this is a locally controlled ability (always) or if this is a replicating targeting mode.
-		// (E.g., server will spawn this target actor to replicate to all non owning clients)
-		if (bReplicates || bIsLocallyControlled || CDO->ShouldProduceTargetDataOnServer)
+		if (ShouldSpawnTargetActor())
 		{
-			UClass* Class = *TargetClass;
-			if (Class != NULL)
-			{
-				UWorld* World = GEngine->GetWorldFromContextObject(WorldContextObject);
-				SpawnedActor = World->SpawnActorDeferred<AGameplayAbilityTargetActor>(Class, FVector::ZeroVector, FRotator::ZeroRotator, NULL, NULL, true);
-			}
+			InitializeTargetActor(SpawnedActor);
+			FinalizeTargetActor(SpawnedActor);
+		}
+		else
+		{
+			TargetActor = NULL;
 
-			// If we spawned the target actor, always register the callbacks for when the data is ready.
-			SpawnedActor->TargetDataReadyDelegate.AddUObject(this, &UAbilityTask_WaitTargetData::OnTargetDataReadyCallback);
-			SpawnedActor->CanceledDelegate.AddUObject(this, &UAbilityTask_WaitTargetData::OnTargetDataCancelledCallback);
+			// We may need a better solution here.  We don't know the target actor isn't needed till after it's already been spawned.
+			SpawnedActor->Destroy();
+			SpawnedActor = nullptr;
+		}
 
-			MyTargetActor = SpawnedActor;
+		RegisterTargetDataCallbacks();
+	}
+}
+
+bool UAbilityTask_WaitTargetData::BeginSpawningActor(UObject* WorldContextObject, TSubclassOf<AGameplayAbilityTargetActor> InTargetClass, AGameplayAbilityTargetActor*& SpawnedActor)
+{
+	SpawnedActor = nullptr;
 
-			AGameplayAbilityTargetActor* TargetActor = CastChecked<AGameplayAbilityTargetActor>(SpawnedActor);
-			if (TargetActor)
+	if (Ability.IsValid())
+	{
+		if (ShouldSpawnTargetActor())
+		{
+			UClass* Class = *InTargetClass;
+			if (Class != NULL)
 			{
-				TargetActor->MasterPC = MyAbility->GetCurrentActorInfo()->PlayerController.Get();
+				UWorld* World = GEngine->GetWorldFromContextObject(WorldContextObject);
+				SpawnedActor = World->SpawnActorDeferred<AGameplayAbilityTargetActor>(Class, FVector::ZeroVector, FRotator::ZeroRotator, NULL, NULL, true);
 			}
-		}		
 
-		// If not locally controlled (server for remote client), see if TargetData was already sent
-		// else register callback for when it does get here.
-		if (!bIsLocallyControlled)
-		{
-			// Register with the TargetData callbacks if we are expecting client to send them
-			if (!CDO->ShouldProduceTargetDataOnServer)
+			if (SpawnedActor)
 			{
-				// Problem here - if there's targeting data just sitting around, fire events don't get hooked up because of this if-else, even if we don't end the task.
-				if (AbilitySystemComponent->ReplicatedTargetData.IsValid(0))
-				{
-					ValidData.Broadcast(AbilitySystemComponent->ReplicatedTargetData);
-					if (ConfirmationType == EGameplayTargetingConfirmation::CustomMulti)
-					{
-						//Since multifire is supported, we still need to hook up the callbacks
-						OnTargetDataReplicatedCallbackDelegateHandle = AbilitySystemComponent->ReplicatedTargetDataDelegate.AddUObject(this, &UAbilityTask_WaitTargetData::OnTargetDataReplicatedCallback);
-						AbilitySystemComponent->ReplicatedTargetDataCancelledDelegate.AddDynamic(this, &UAbilityTask_WaitTargetData::OnTargetDataReplicatedCancelledCallback);
-					}
-					else
-					{
-						EndTask();
-					}
-				}
-				else
-				{
-					OnTargetDataReplicatedCallbackDelegateHandle = AbilitySystemComponent->ReplicatedTargetDataDelegate.AddUObject(this, &UAbilityTask_WaitTargetData::OnTargetDataReplicatedCallback);
-					AbilitySystemComponent->ReplicatedTargetDataCancelledDelegate.AddDynamic(this, &UAbilityTask_WaitTargetData::OnTargetDataReplicatedCancelledCallback);
-				}
+				TargetActor = SpawnedActor;
+				InitializeTargetActor(SpawnedActor);
 			}
 		}
+
+		RegisterTargetDataCallbacks();
 	}
+
 	return (SpawnedActor != nullptr);
 }
 
@@ -92,46 +89,124 @@ void UAbilityTask_WaitTargetData::FinishSpawningActor(UObject* WorldContextObjec
 {
 	if (SpawnedActor && !SpawnedActor->IsPendingKill())
 	{
-		check(MyTargetActor == SpawnedActor);
+		check(TargetActor == SpawnedActor);
 
 		const FTransform SpawnTransform = AbilitySystemComponent->GetOwner()->GetTransform();
 
 		SpawnedActor->FinishSpawning(SpawnTransform);
 
-		// User ability activation is inhibited while this is active
-		AbilitySystemComponent->SetUserAbilityActivationInhibited(true);
-		AbilitySystemComponent->SpawnedTargetActors.Push(SpawnedActor);
+		FinalizeTargetActor(SpawnedActor);
+	}
+}
+
+bool UAbilityTask_WaitTargetData::ShouldSpawnTargetActor() const
+{
+	check(TargetClass);
+	check(Ability.IsValid());
+
+	// Spawn the actor if this is a locally controlled ability (always) or if this is a replicating targeting mode.
+	// (E.g., server will spawn this target actor to replicate to all non owning clients)
 
-		SpawnedActor->StartTargeting(Ability.Get());
+	const AGameplayAbilityTargetActor* CDO = CastChecked<AGameplayAbilityTargetActor>(TargetClass->GetDefaultObject());
 
-		if (SpawnedActor->ShouldProduceTargetData())
+	const bool bReplicates = CDO->GetReplicates();
+	const bool bIsLocallyControlled = Ability->GetCurrentActorInfo()->IsLocallyControlled();
+	const bool bShouldProduceTargetDataOnServer = CDO->ShouldProduceTargetDataOnServer;
+
+	return (bReplicates || bIsLocallyControlled || bShouldProduceTargetDataOnServer);
+}
+
+void UAbilityTask_WaitTargetData::InitializeTargetActor(AGameplayAbilityTargetActor* SpawnedActor) const
+{
+	check(SpawnedActor);
+	check(Ability.IsValid());
+
+	SpawnedActor->MasterPC = Ability->GetCurrentActorInfo()->PlayerController.Get();
+
+	// If we spawned the target actor, always register the callbacks for when the data is ready.
+	SpawnedActor->TargetDataReadyDelegate.AddUObject(this, &UAbilityTask_WaitTargetData::OnTargetDataReadyCallback);
+	SpawnedActor->CanceledDelegate.AddUObject(this, &UAbilityTask_WaitTargetData::OnTargetDataCancelledCallback);
+}
+
+void UAbilityTask_WaitTargetData::FinalizeTargetActor(AGameplayAbilityTargetActor* SpawnedActor) const
+{
+	check(SpawnedActor);
+	check(Ability.IsValid());
+
+	// User ability activation is inhibited while this is active
+	AbilitySystemComponent->SetUserAbilityActivationInhibited(true);
+	AbilitySystemComponent->SpawnedTargetActors.Push(SpawnedActor);
+
+	SpawnedActor->StartTargeting(Ability.Get());
+
+	if (SpawnedActor->ShouldProduceTargetData())
+	{
+		// If instant confirm, then stop targeting immediately.
+		// Note this is kind of bad: we should be able to just call a static func on the CDO to do this. 
+		// But then we wouldn't get to set ExposeOnSpawnParameters.
+		if (ConfirmationType == EGameplayTargetingConfirmation::Instant)
+		{
+			SpawnedActor->ConfirmTargeting();
+		}
+		else if (ConfirmationType == EGameplayTargetingConfirmation::UserConfirmed)
 		{
-			// If instant confirm, then stop targeting immediately.
-			// Note this is kind of bad: we should be able to just call a static func on the CDO to do this. 
-			// But then we wouldn't get to set ExposeOnSpawnParameters.
-			if (ConfirmationType == EGameplayTargetingConfirmation::Instant)
+			// If not locally controlled, check if we already have a confirm/cancel repped to us
+			if (!Ability->GetCurrentActorInfo()->IsLocallyControlled())
 			{
-				SpawnedActor->ConfirmTargeting();
+				if (AbilitySystemComponent->ReplicatedConfirmAbility)
+				{
+					SpawnedActor->ConfirmTargeting();
+					return;
+				}
+				else if (AbilitySystemComponent->ReplicatedCancelAbility)
+				{
+					SpawnedActor->CancelTargeting();
+					return;
+				}
 			}
-			else if (ConfirmationType == EGameplayTargetingConfirmation::UserConfirmed)
+
+			// Bind to the Cancel/Confirm Delegates (called from local confirm or from repped confirm)
+			SpawnedActor->BindToConfirmCancelInputs();
+		}
+	}
+}
+
+void UAbilityTask_WaitTargetData::RegisterTargetDataCallbacks()
+{
+	check(TargetClass);
+	check(Ability.IsValid());
+
+	const AGameplayAbilityTargetActor* CDO = CastChecked<AGameplayAbilityTargetActor>(TargetClass->GetDefaultObject());
+
+	const bool bIsLocallyControlled = Ability->GetCurrentActorInfo()->IsLocallyControlled();
+	const bool bShouldProduceTargetDataOnServer = CDO->ShouldProduceTargetDataOnServer;
+
+	// If not locally controlled (server for remote client), see if TargetData was already sent
+	// else register callback for when it does get here.
+	if (!bIsLocallyControlled)
+	{
+		// Register with the TargetData callbacks if we are expecting client to send them
+		if (!bShouldProduceTargetDataOnServer)
+		{
+			// Problem here - if there's targeting data just sitting around, fire events don't get hooked up because of this if-else, even if we don't end the task.
+			if (AbilitySystemComponent->ReplicatedTargetData.IsValid(0))
 			{
-				// If not locally controlled, check if we already have a confirm/cancel repped to us
-				if (!Ability->GetCurrentActorInfo()->IsLocallyControlled())
+				ValidData.Broadcast(AbilitySystemComponent->ReplicatedTargetData);
+				if (ConfirmationType == EGameplayTargetingConfirmation::CustomMulti)
 				{
-					if (AbilitySystemComponent->ReplicatedConfirmAbility)
-					{
-						SpawnedActor->ConfirmTargeting();
-						return;
-					}
-					else if (AbilitySystemComponent->ReplicatedCancelAbility)
-					{
-						SpawnedActor->CancelTargeting();
-						return;
-					}
+					//Since multifire is supported, we still need to hook up the callbacks
+					AbilitySystemComponent->ReplicatedTargetDataDelegate.AddUObject(this, &UAbilityTask_WaitTargetData::OnTargetDataReplicatedCallback);
+					AbilitySystemComponent->ReplicatedTargetDataCancelledDelegate.AddDynamic(this, &UAbilityTask_WaitTargetData::OnTargetDataReplicatedCancelledCallback);
 				}
-
-				// Bind to the Cancel/Confirm Delegates (called from local confirm or from repped confirm)
-				SpawnedActor->BindToConfirmCancelInputs();
+				else
+				{
+					EndTask();
+				}
+			}
+			else
+			{
+				AbilitySystemComponent->ReplicatedTargetDataDelegate.AddUObject(this, &UAbilityTask_WaitTargetData::OnTargetDataReplicatedCallback);
+				AbilitySystemComponent->ReplicatedTargetDataCancelledDelegate.AddDynamic(this, &UAbilityTask_WaitTargetData::OnTargetDataReplicatedCancelledCallback);
 			}
 		}
 	}
@@ -151,7 +226,7 @@ void UAbilityTask_WaitTargetData::OnTargetDataReplicatedCallback(FGameplayAbilit
 	 *	explicitly, the client is basically just sending a 'confirm' and the server is now going to do the work
 	 *	in OnReplicatedTargetDataReceived.
 	 */
-	if (MyTargetActor.IsValid() && !MyTargetActor->OnReplicatedTargetDataReceived(Data))
+	if (TargetActor.IsValid() && !TargetActor->OnReplicatedTargetDataReceived(Data))
 	{
 		Cancelled.Broadcast(Data);
 	}
@@ -211,9 +286,9 @@ void UAbilityTask_WaitTargetData::OnTargetDataCancelledCallback(FGameplayAbility
 void UAbilityTask_WaitTargetData::ExternalConfirm(bool bEndTask)
 {
 	check(AbilitySystemComponent.IsValid());
-	if (MyTargetActor.IsValid())
+	if (TargetActor.IsValid())
 	{
-		AGameplayAbilityTargetActor* CachedTargetActor = MyTargetActor.Get();
+		AGameplayAbilityTargetActor* CachedTargetActor = TargetActor.Get();
 		if (CachedTargetActor->ShouldProduceTargetData())
 		{
 			CachedTargetActor->ConfirmTargetingAndContinue();
@@ -236,12 +311,12 @@ void UAbilityTask_WaitTargetData::OnDestroy(bool AbilityEnded)
 	AbilitySystemComponent->ConsumeAbilityConfirmCancel();
 	AbilitySystemComponent->SetUserAbilityActivationInhibited(false);
 
-	AbilitySystemComponent->ReplicatedTargetDataDelegate.Remove(OnTargetDataReplicatedCallbackDelegateHandle);
+	AbilitySystemComponent->ReplicatedTargetDataDelegate.RemoveUObject(this, &UAbilityTask_WaitTargetData::OnTargetDataReplicatedCallback);
 	AbilitySystemComponent->ReplicatedTargetDataCancelledDelegate.RemoveDynamic(this, &UAbilityTask_WaitTargetData::OnTargetDataReplicatedCancelledCallback);
 
-	if (MyTargetActor.IsValid())
+	if (TargetActor.IsValid())
 	{
-		MyTargetActor->Destroy();
+		TargetActor->Destroy();
 	}
 
 	Super::OnDestroy(AbilityEnded);
@@ -249,9 +324,12 @@ void UAbilityTask_WaitTargetData::OnDestroy(bool AbilityEnded)
 
 bool UAbilityTask_WaitTargetData::ShouldReplicateDataToServer() const
 {
+	check(Ability.IsValid());
+	check(TargetActor.IsValid()); // Target actor should always be valid on client
+
 	// Send TargetData to the server IFF we are the client and this isn't a GameplayTargetActor that can produce data on the server	
 	const FGameplayAbilityActorInfo* Info = Ability->GetCurrentActorInfo();
-	if (!Info->IsNetAuthority() && !TargetClass->GetDefaultObject<AGameplayAbilityTargetActor>()->ShouldProduceTargetDataOnServer)
+	if (!Info->IsNetAuthority() && !TargetActor->ShouldProduceTargetDataOnServer)
 	{
 		return true;
 	}
diff --git a/Engine/Source/Runtime/GameplayAbilities/Public/Abilities/Tasks/AbilityTask_VisualizeTargeting.h b/Engine/Source/Runtime/GameplayAbilities/Public/Abilities/Tasks/AbilityTask_VisualizeTargeting.h
--- a/Engine/Source/Runtime/GameplayAbilities/Public/Abilities/Tasks/AbilityTask_VisualizeTargeting.h
+++ b/Engine/Source/Runtime/GameplayAbilities/Public/Abilities/Tasks/AbilityTask_VisualizeTargeting.h
@@ -1,5 +1,6 @@
 // Copyright 1998-2015 Epic Games, Inc. All Rights Reserved.
 #pragma once
+
 #include "TimerManager.h"
 #include "AbilityTask.h"
 #include "AbilityTask_VisualizeTargeting.generated.h"
@@ -18,12 +19,16 @@ class UAbilityTask_VisualizeTargeting: public UAbilityTask
 
 	void OnTimeElapsed();
 
-	void SetDuration(const float Duration);
-
-	/** Spawns Targeting actor. */
+	/** Spawns target actor and uses it for visualization. */
 	UFUNCTION(BlueprintCallable, meta=(HidePin = "WorldContextObject", DefaultToSelf = "WorldContextObject", BlueprintInternalUseOnly = "true", HideSpawnParms="Instigator"), Category="Ability|Tasks")
 	static UAbilityTask_VisualizeTargeting* VisualizeTargeting(UObject* WorldContextObject, TSubclassOf<AGameplayAbilityTargetActor> Class, FName TaskInstanceName, float Duration = -1.0f);
 
+	/** Visualize target using a specified target actor. */
+	UFUNCTION(BlueprintCallable, meta = (HidePin = "WorldContextObject", DefaultToSelf = "WorldContextObject", BlueprintInternalUseOnly = "true"), Category = "Ability|Tasks")
+	static UAbilityTask_VisualizeTargeting* VisualizeTargetingUsingActor(UObject* WorldContextObject, AGameplayAbilityTargetActor* TargetActor, FName TaskInstanceName, float Duration = -1.0f);
+
+	virtual void Activate();
+
 	UFUNCTION(BlueprintCallable, meta = (HidePin = "WorldContextObject", DefaultToSelf = "WorldContextObject", BlueprintInternalUseOnly = "true"), Category = "Abilities")
 	bool BeginSpawningActor(UObject* WorldContextObject, TSubclassOf<AGameplayAbilityTargetActor> Class, AGameplayAbilityTargetActor*& SpawnedActor);
 
@@ -32,10 +37,20 @@ class UAbilityTask_VisualizeTargeting: public UAbilityTask
 
 protected:
 
+	void SetDuration(const float Duration);
+
+	bool ShouldSpawnTargetActor() const;
+	void InitializeTargetActor(AGameplayAbilityTargetActor* SpawnedActor) const;
+	void FinalizeTargetActor(AGameplayAbilityTargetActor* SpawnedActor) const;
+
 	virtual void OnDestroy(bool AbilityEnded) override;
 
+protected:
+
+	TSubclassOf<AGameplayAbilityTargetActor> TargetClass;
+
 	/** The TargetActor that we spawned */
-	TWeakObjectPtr<AGameplayAbilityTargetActor>	MyTargetActor;
+	TWeakObjectPtr<AGameplayAbilityTargetActor>	TargetActor;
 
 	/** Handle for efficient management of OnTimeElapsed timer */
 	FTimerHandle TimerHandle_OnTimeElapsed;
diff --git a/Engine/Source/Runtime/GameplayAbilities/Public/Abilities/Tasks/AbilityTask_WaitTargetData.h b/Engine/Source/Runtime/GameplayAbilities/Public/Abilities/Tasks/AbilityTask_WaitTargetData.h
--- a/Engine/Source/Runtime/GameplayAbilities/Public/Abilities/Tasks/AbilityTask_WaitTargetData.h
+++ b/Engine/Source/Runtime/GameplayAbilities/Public/Abilities/Tasks/AbilityTask_WaitTargetData.h
@@ -33,10 +33,16 @@ class GAMEPLAYABILITIES_API UAbilityTask_WaitTargetData: public UAbilityTask
 	UFUNCTION()
 	void OnTargetDataCancelledCallback(FGameplayAbilityTargetDataHandle Data);
 
-	/** Spawns Targeting actor and waits for it to return valid data or to be canceled. */
+	/** Spawns target actor and waits for it to return valid data or to be canceled. */
 	UFUNCTION(BlueprintCallable, meta=(HidePin = "WorldContextObject", DefaultToSelf = "WorldContextObject", BlueprintInternalUseOnly = "true", HideSpawnParms="Instigator"), Category="Ability|Tasks")
 	static UAbilityTask_WaitTargetData* WaitTargetData(UObject* WorldContextObject, FName TaskInstanceName, TEnumAsByte<EGameplayTargetingConfirmation::Type> ConfirmationType, TSubclassOf<AGameplayAbilityTargetActor> Class);
 
+	/** Uses specified target actor and waits for it to return valid data or to be canceled. */
+	UFUNCTION(BlueprintCallable, meta = (HidePin = "WorldContextObject", DefaultToSelf = "WorldContextObject", BlueprintInternalUseOnly = "true", HideSpawnParms = "Instigator"), Category = "Ability|Tasks")
+	static UAbilityTask_WaitTargetData* WaitTargetDataUsingActor(UObject* WorldContextObject, FName TaskInstanceName, TEnumAsByte<EGameplayTargetingConfirmation::Type> ConfirmationType, AGameplayAbilityTargetActor* TargetActor);
+
+	virtual void Activate();
+
 	UFUNCTION(BlueprintCallable, meta = (HidePin = "WorldContextObject", DefaultToSelf = "WorldContextObject", BlueprintInternalUseOnly = "true"), Category = "Abilities")
 	bool BeginSpawningActor(UObject* WorldContextObject, TSubclassOf<AGameplayAbilityTargetActor> Class, AGameplayAbilityTargetActor*& SpawnedActor);
 
@@ -51,14 +57,22 @@ class GAMEPLAYABILITIES_API UAbilityTask_WaitTargetData: public UAbilityTask
 
 protected:
 
+	bool ShouldSpawnTargetActor() const;
+	void InitializeTargetActor(AGameplayAbilityTargetActor* SpawnedActor) const;
+	void FinalizeTargetActor(AGameplayAbilityTargetActor* SpawnedActor) const;
+
+	void RegisterTargetDataCallbacks();
+
 	virtual void OnDestroy(bool AbilityEnded) override;
 
 	bool ShouldReplicateDataToServer() const;
 
+protected:
+
 	TSubclassOf<AGameplayAbilityTargetActor> TargetClass;
 
 	/** The TargetActor that we spawned */
-	TWeakObjectPtr<AGameplayAbilityTargetActor>	MyTargetActor;
+	TWeakObjectPtr<AGameplayAbilityTargetActor>	TargetActor;
 
 	TEnumAsByte<EGameplayTargetingConfirmation::Type> ConfirmationType;
 
@@ -73,4 +87,4 @@ protected:
 *			-This function should spawn the actor with SpawnActorDeferred and return true/false if it spawned something.
 *		-Have a function named FinishSpawningActor w/ an AActor* of the class you spawned
 *			-This function *must* call ExecuteConstruction + PostActorConstruction
-*/
\ No newline at end of file
+*/
