Integration(2407338): Fixed crash when using HitScanTube tasks;

diff --git a/Engine/Source/Runtime/GameplayAbilities/Private/Abilities/Tasks/AbilityTask_VisualizeTargeting.cpp b/Engine/Source/Runtime/GameplayAbilities/Private/Abilities/Tasks/AbilityTask_VisualizeTargeting.cpp
--- a/Engine/Source/Runtime/GameplayAbilities/Private/Abilities/Tasks/AbilityTask_VisualizeTargeting.cpp
+++ b/Engine/Source/Runtime/GameplayAbilities/Private/Abilities/Tasks/AbilityTask_VisualizeTargeting.cpp
@@ -31,9 +31,9 @@ UAbilityTask_VisualizeTargeting* UAbilityTask_VisualizeTargeting::VisualizeTarge
 
 void UAbilityTask_VisualizeTargeting::Activate()
 {
-	if (Ability.IsValid() && TargetClass == NULL)
+	// Need to handle case where target actor was passed into task
+	if (Ability.IsValid() && (TargetClass == NULL) && TargetActor.IsValid())
 	{
-		check(TargetActor.IsValid());
 		AGameplayAbilityTargetActor* SpawnedActor = TargetActor.Get();
 
 		TargetClass = SpawnedActor->GetClass();
diff --git a/Engine/Source/Runtime/GameplayAbilities/Private/Abilities/Tasks/AbilityTask_WaitTargetData.cpp b/Engine/Source/Runtime/GameplayAbilities/Private/Abilities/Tasks/AbilityTask_WaitTargetData.cpp
--- a/Engine/Source/Runtime/GameplayAbilities/Private/Abilities/Tasks/AbilityTask_WaitTargetData.cpp
+++ b/Engine/Source/Runtime/GameplayAbilities/Private/Abilities/Tasks/AbilityTask_WaitTargetData.cpp
@@ -32,9 +32,8 @@ UAbilityTask_WaitTargetData* UAbilityTask_WaitTargetData::WaitTargetDataUsingAct
 void UAbilityTask_WaitTargetData::Activate()
 {
 	// Need to handle case where target actor was passed into task
-	if (Ability.IsValid() && TargetClass == NULL)
+	if (Ability.IsValid() && (TargetClass == NULL) && TargetActor.IsValid())
 	{
-		check(TargetActor.IsValid());
 		AGameplayAbilityTargetActor* SpawnedActor = TargetActor.Get();
 
 		TargetClass = SpawnedActor->GetClass();
