Fix scopelock crash on clients by locking during net updates

diff --git a/Engine/Source/Runtime/GameplayAbilities/Private/AbilitySystemComponent.cpp b/Engine/Source/Runtime/GameplayAbilities/Private/AbilitySystemComponent.cpp
--- a/Engine/Source/Runtime/GameplayAbilities/Private/AbilitySystemComponent.cpp
+++ b/Engine/Source/Runtime/GameplayAbilities/Private/AbilitySystemComponent.cpp
@@ -1039,6 +1039,16 @@ void UAbilitySystemComponent::GetSubobjectsWithStableNamesForNetworking(TArray<U
 	}
 }
 
+void UAbilitySystemComponent::PreNetReceive()
+{
+	ActiveGameplayEffects.IncrementLock();
+}
+	
+void UAbilitySystemComponent::PostNetReceive()
+{
+	ActiveGameplayEffects.DecrementLock();
+}
+
 void UAbilitySystemComponent::OnRep_GameplayEffects()
 {
 
diff --git a/Engine/Source/Runtime/GameplayAbilities/Private/GameplayEffect.cpp b/Engine/Source/Runtime/GameplayAbilities/Private/GameplayEffect.cpp
--- a/Engine/Source/Runtime/GameplayAbilities/Private/GameplayEffect.cpp
+++ b/Engine/Source/Runtime/GameplayAbilities/Private/GameplayEffect.cpp
@@ -2617,44 +2617,54 @@ void FInheritedTagContainer::RemoveTag(FGameplayTag TagToRemove)
 	CombinedTags.RemoveTag(TagToRemove);
 }
 
-FScopedActiveGameplayEffectLock::FScopedActiveGameplayEffectLock(FActiveGameplayEffectsContainer& InContainer)
-	: Container(InContainer)
+void FActiveGameplayEffectsContainer::IncrementLock()
 {
-	if (Container.ScopedLockCount++ == 0)
+	if (ScopedLockCount++ == 0)
 	{
 		// If we transitioned from unlocked to locked, make sure we have enough slack for additions.
 		// 4 is arbitrary determined. We should add instrumentation to get a better ideal/default number.
 		static const int32 GAMEPLAYEFFECT_MIN_SLACK = 4;
 
-		if (Container.GameplayEffects.GetSlack() < GAMEPLAYEFFECT_MIN_SLACK)
+		if (GameplayEffects.GetSlack() < GAMEPLAYEFFECT_MIN_SLACK)
 		{
-			Container.GameplayEffects.Reserve( Container.GameplayEffects.Num() + GAMEPLAYEFFECT_MIN_SLACK );
+			GameplayEffects.Reserve( GameplayEffects.Num() + GAMEPLAYEFFECT_MIN_SLACK );
 		}
 	}
 }
 
-FScopedActiveGameplayEffectLock::~FScopedActiveGameplayEffectLock()
+void FActiveGameplayEffectsContainer::DecrementLock()
 {
-	if (--Container.ScopedLockCount == 0 && Container.PendingRemoves > 0)
+	if (--ScopedLockCount == 0 && PendingRemoves > 0)
 	{
 		// If we are transitioned from locked to unlocked, we now need to actually delete and pending kill gameplay effects
-		for (int32 idx=Container.GameplayEffects.Num()-1; idx >= 0 && Container.PendingRemoves > 0; --idx)
+		for (int32 idx=GameplayEffects.Num()-1; idx >= 0 && PendingRemoves > 0; --idx)
 		{
-			FActiveGameplayEffect& Effect = Container.GameplayEffects[idx];
+			FActiveGameplayEffect& Effect = GameplayEffects[idx];
 
 			if (Effect.IsPendingRemove)
 			{
-				Container.GameplayEffects.RemoveAtSwap(idx, 1, false);
-				Container.PendingRemoves--;
+				GameplayEffects.RemoveAtSwap(idx, 1, false);
+				PendingRemoves--;
 			}
 		}
 
-		if (!ensure(Container.PendingRemoves == 0))
+		if (!ensure(PendingRemoves == 0))
 		{
-			ABILITY_LOG(Warning, TEXT("~FScopedActiveGameplayEffectLock has %d pending removes after a scope lock removal"), Container.PendingRemoves);
-			Container.PendingRemoves = 0;
+			ABILITY_LOG(Warning, TEXT("~FScopedActiveGameplayEffectLock has %d pending removes after a scope lock removal"), PendingRemoves);
+			PendingRemoves = 0;
 		}
 
-		Container.MarkArrayDirty();
+		MarkArrayDirty();
 	}
 }
+
+FScopedActiveGameplayEffectLock::FScopedActiveGameplayEffectLock(FActiveGameplayEffectsContainer& InContainer)
+	: Container(InContainer)
+{
+	Container.IncrementLock();
+}
+
+FScopedActiveGameplayEffectLock::~FScopedActiveGameplayEffectLock()
+{
+	Container.DecrementLock();
+}
diff --git a/Engine/Source/Runtime/GameplayAbilities/Public/AbilitySystemComponent.h b/Engine/Source/Runtime/GameplayAbilities/Public/AbilitySystemComponent.h
--- a/Engine/Source/Runtime/GameplayAbilities/Public/AbilitySystemComponent.h
+++ b/Engine/Source/Runtime/GameplayAbilities/Public/AbilitySystemComponent.h
@@ -153,6 +153,10 @@ class GAMEPLAYABILITIES_API UAbilitySystemComponent : public UActorComponent, pu
 
 	virtual void GetSubobjectsWithStableNamesForNetworking(TArray<UObject*>& Objs) override;
 
+	virtual void PreNetReceive() override;
+	
+	virtual void PostNetReceive() override;
+
 	/** PredictionKeys, see more info in GameplayPrediction.h */
 	UPROPERTY(ReplicatedUsing=OnRep_PredictionKey)
 	FPredictionKey	ReplicatedPredictionKey;
diff --git a/Engine/Source/Runtime/GameplayAbilities/Public/GameplayEffect.h b/Engine/Source/Runtime/GameplayAbilities/Public/GameplayEffect.h
--- a/Engine/Source/Runtime/GameplayAbilities/Public/GameplayEffect.h
+++ b/Engine/Source/Runtime/GameplayAbilities/Public/GameplayEffect.h
@@ -1355,6 +1355,9 @@ struct FActiveGameplayEffectsContainer : public FFastArraySerializer
 
 	bool HasApplicationImmunityToSpec(const FGameplayEffectSpec& SpecToApply) const;
 
+	void IncrementLock();
+	void DecrementLock();
+
 private:
 
 	UPROPERTY()
