Fixed deprecated warning caused by integration;

diff --git a/Engine/Source/Runtime/GameplayAbilities/Private/Abilities/Tasks/AbilityTask_WaitTargetData.cpp b/Engine/Source/Runtime/GameplayAbilities/Private/Abilities/Tasks/AbilityTask_WaitTargetData.cpp
--- a/Engine/Source/Runtime/GameplayAbilities/Private/Abilities/Tasks/AbilityTask_WaitTargetData.cpp
+++ b/Engine/Source/Runtime/GameplayAbilities/Private/Abilities/Tasks/AbilityTask_WaitTargetData.cpp
@@ -310,7 +310,7 @@ void UAbilityTask_WaitTargetData::OnDestroy(bool AbilityEnded)
 	AbilitySystemComponent->ConsumeAbilityConfirmCancel();
 	AbilitySystemComponent->SetUserAbilityActivationInhibited(false);
 
-	AbilitySystemComponent->ReplicatedTargetDataDelegate.RemoveUObject(this, &UAbilityTask_WaitTargetData::OnTargetDataReplicatedCallback);
+	AbilitySystemComponent->ReplicatedTargetDataDelegate.Remove(OnTargetDataReplicatedCallbackDelegateHandle);
 	AbilitySystemComponent->ReplicatedTargetDataCancelledDelegate.RemoveDynamic(this, &UAbilityTask_WaitTargetData::OnTargetDataReplicatedCancelledCallback);
 
 	if (TargetActor.IsValid())
