Integration(2408223): Target actor task fixes; Reticle snapping;

* Fixed 'WaitTargetDataUsingActor' and 'VisualizeTargetingUsingActor' to end the task if an invalid actor is specified.

* Added option to AGameplayAbilityWorldReticle for snapping to the target actor's location instead of the hit location. AGameplayAbilityTargetActor_SingleLineTrace is currently the only targeting actor that uses this.

diff --git a/Engine/Source/Runtime/GameplayAbilities/Private/Abilities/GameplayAbilityTargetActor_SingleLineTrace.cpp b/Engine/Source/Runtime/GameplayAbilities/Private/Abilities/GameplayAbilityTargetActor_SingleLineTrace.cpp
--- a/Engine/Source/Runtime/GameplayAbilities/Private/Abilities/GameplayAbilityTargetActor_SingleLineTrace.cpp
+++ b/Engine/Source/Runtime/GameplayAbilities/Private/Abilities/GameplayAbilityTargetActor_SingleLineTrace.cpp
@@ -45,8 +45,11 @@ FHitResult AGameplayAbilityTargetActor_SingleLineTrace::PerformTrace(AActor* InS
 	}
 	if (AGameplayAbilityWorldReticle* LocalReticleActor = ReticleActor.Get())
 	{
-		LocalReticleActor->SetActorLocation(ReturnHitResult.Location);
-		LocalReticleActor->SetIsTargetAnActor(ReturnHitResult.bBlockingHit && (ReturnHitResult.Actor != NULL));
+		const bool bHitActor = (ReturnHitResult.bBlockingHit && (ReturnHitResult.Actor != NULL));
+		const FVector ReticleLocation = (bHitActor && LocalReticleActor->bSnapToTargetedActor) ? ReturnHitResult.Actor->GetActorLocation() : ReturnHitResult.Location;
+
+		LocalReticleActor->SetActorLocation(ReticleLocation);
+		LocalReticleActor->SetIsTargetAnActor(bHitActor);
 	}
 
 	if (bDebug)
diff --git a/Engine/Source/Runtime/GameplayAbilities/Private/Abilities/Tasks/AbilityTask_VisualizeTargeting.cpp b/Engine/Source/Runtime/GameplayAbilities/Private/Abilities/Tasks/AbilityTask_VisualizeTargeting.cpp
--- a/Engine/Source/Runtime/GameplayAbilities/Private/Abilities/Tasks/AbilityTask_VisualizeTargeting.cpp
+++ b/Engine/Source/Runtime/GameplayAbilities/Private/Abilities/Tasks/AbilityTask_VisualizeTargeting.cpp
@@ -32,24 +32,31 @@ UAbilityTask_VisualizeTargeting* UAbilityTask_VisualizeTargeting::VisualizeTarge
 void UAbilityTask_VisualizeTargeting::Activate()
 {
 	// Need to handle case where target actor was passed into task
-	if (Ability.IsValid() && (TargetClass == NULL) && TargetActor.IsValid())
+	if (Ability.IsValid() && (TargetClass == NULL))
 	{
-		AGameplayAbilityTargetActor* SpawnedActor = TargetActor.Get();
+		if (TargetActor.IsValid())
+		{
+			AGameplayAbilityTargetActor* SpawnedActor = TargetActor.Get();
 
-		TargetClass = SpawnedActor->GetClass();
+			TargetClass = SpawnedActor->GetClass();
 
-		if (ShouldSpawnTargetActor())
-		{
-			InitializeTargetActor(SpawnedActor);
-			FinalizeTargetActor(SpawnedActor);
+			if (ShouldSpawnTargetActor())
+			{
+				InitializeTargetActor(SpawnedActor);
+				FinalizeTargetActor(SpawnedActor);
+			}
+			else
+			{
+				TargetActor = NULL;
+
+				// We may need a better solution here.  We don't know the target actor isn't needed till after it's already been spawned.
+				SpawnedActor->Destroy();
+				SpawnedActor = nullptr;
+			}
 		}
 		else
 		{
-			TargetActor = NULL;
-
-			// We may need a better solution here.  We don't know the target actor isn't needed till after it's already been spawned.
-			SpawnedActor->Destroy();
-			SpawnedActor = nullptr;
+			EndTask();
 		}
 	}
 }
diff --git a/Engine/Source/Runtime/GameplayAbilities/Private/Abilities/Tasks/AbilityTask_WaitTargetData.cpp b/Engine/Source/Runtime/GameplayAbilities/Private/Abilities/Tasks/AbilityTask_WaitTargetData.cpp
--- a/Engine/Source/Runtime/GameplayAbilities/Private/Abilities/Tasks/AbilityTask_WaitTargetData.cpp
+++ b/Engine/Source/Runtime/GameplayAbilities/Private/Abilities/Tasks/AbilityTask_WaitTargetData.cpp
@@ -32,27 +32,34 @@ UAbilityTask_WaitTargetData* UAbilityTask_WaitTargetData::WaitTargetDataUsingAct
 void UAbilityTask_WaitTargetData::Activate()
 {
 	// Need to handle case where target actor was passed into task
-	if (Ability.IsValid() && (TargetClass == NULL) && TargetActor.IsValid())
+	if (Ability.IsValid() && (TargetClass == NULL))
 	{
-		AGameplayAbilityTargetActor* SpawnedActor = TargetActor.Get();
+		if (TargetActor.IsValid())
+		{
+			AGameplayAbilityTargetActor* SpawnedActor = TargetActor.Get();
 
-		TargetClass = SpawnedActor->GetClass();
+			TargetClass = SpawnedActor->GetClass();
 
-		if (ShouldSpawnTargetActor())
-		{
-			InitializeTargetActor(SpawnedActor);
-			FinalizeTargetActor(SpawnedActor);
+			if (ShouldSpawnTargetActor())
+			{
+				InitializeTargetActor(SpawnedActor);
+				FinalizeTargetActor(SpawnedActor);
+			}
+			else
+			{
+				TargetActor = NULL;
+
+				// We may need a better solution here.  We don't know the target actor isn't needed till after it's already been spawned.
+				SpawnedActor->Destroy();
+				SpawnedActor = nullptr;
+			}
+
+			RegisterTargetDataCallbacks();
 		}
 		else
 		{
-			TargetActor = NULL;
-
-			// We may need a better solution here.  We don't know the target actor isn't needed till after it's already been spawned.
-			SpawnedActor->Destroy();
-			SpawnedActor = nullptr;
+			EndTask();
 		}
-
-		RegisterTargetDataCallbacks();
 	}
 }
 
diff --git a/Engine/Source/Runtime/GameplayAbilities/Public/Abilities/GameplayAbilityWorldReticle.h b/Engine/Source/Runtime/GameplayAbilities/Public/Abilities/GameplayAbilityWorldReticle.h
--- a/Engine/Source/Runtime/GameplayAbilities/Public/Abilities/GameplayAbilityWorldReticle.h
+++ b/Engine/Source/Runtime/GameplayAbilities/Public/Abilities/GameplayAbilityWorldReticle.h
@@ -69,6 +69,10 @@ public:
 	UPROPERTY(EditAnywhere, BlueprintReadWrite, meta = (ExposeOnSpawn = "true"), Category = "Reticle")
 	bool bFaceOwnerFlat;
 
+	// If the target is an actor snap to it's location 
+	UPROPERTY(EditAnywhere, BlueprintReadWrite, meta = (ExposeOnSpawn = "true"), Category = "Reticle")
+	bool bSnapToTargetedActor;
+
 protected:
 	/** This indicates whether or not the targeting actor considers the current target to be valid. Defaults to true. */
 	UPROPERTY(BlueprintReadOnly, Category = "Network")
diff --git a/Engine/Source/Runtime/GameplayAbilities/Public/Abilities/Tasks/AbilityTask_VisualizeTargeting.h b/Engine/Source/Runtime/GameplayAbilities/Public/Abilities/Tasks/AbilityTask_VisualizeTargeting.h
--- a/Engine/Source/Runtime/GameplayAbilities/Public/Abilities/Tasks/AbilityTask_VisualizeTargeting.h
+++ b/Engine/Source/Runtime/GameplayAbilities/Public/Abilities/Tasks/AbilityTask_VisualizeTargeting.h
@@ -27,7 +27,7 @@ class UAbilityTask_VisualizeTargeting: public UAbilityTask
 	UFUNCTION(BlueprintCallable, meta = (HidePin = "WorldContextObject", DefaultToSelf = "WorldContextObject", BlueprintInternalUseOnly = "true"), Category = "Ability|Tasks")
 	static UAbilityTask_VisualizeTargeting* VisualizeTargetingUsingActor(UObject* WorldContextObject, AGameplayAbilityTargetActor* TargetActor, FName TaskInstanceName, float Duration = -1.0f);
 
-	virtual void Activate();
+	virtual void Activate() override;
 
 	UFUNCTION(BlueprintCallable, meta = (HidePin = "WorldContextObject", DefaultToSelf = "WorldContextObject", BlueprintInternalUseOnly = "true"), Category = "Abilities")
 	bool BeginSpawningActor(UObject* WorldContextObject, TSubclassOf<AGameplayAbilityTargetActor> Class, AGameplayAbilityTargetActor*& SpawnedActor);
diff --git a/Engine/Source/Runtime/GameplayAbilities/Public/Abilities/Tasks/AbilityTask_WaitTargetData.h b/Engine/Source/Runtime/GameplayAbilities/Public/Abilities/Tasks/AbilityTask_WaitTargetData.h
--- a/Engine/Source/Runtime/GameplayAbilities/Public/Abilities/Tasks/AbilityTask_WaitTargetData.h
+++ b/Engine/Source/Runtime/GameplayAbilities/Public/Abilities/Tasks/AbilityTask_WaitTargetData.h
@@ -41,7 +41,7 @@ class GAMEPLAYABILITIES_API UAbilityTask_WaitTargetData: public UAbilityTask
 	UFUNCTION(BlueprintCallable, meta = (HidePin = "WorldContextObject", DefaultToSelf = "WorldContextObject", BlueprintInternalUseOnly = "true", HideSpawnParms = "Instigator"), Category = "Ability|Tasks")
 	static UAbilityTask_WaitTargetData* WaitTargetDataUsingActor(UObject* WorldContextObject, FName TaskInstanceName, TEnumAsByte<EGameplayTargetingConfirmation::Type> ConfirmationType, AGameplayAbilityTargetActor* TargetActor);
 
-	virtual void Activate();
+	virtual void Activate() override;
 
 	UFUNCTION(BlueprintCallable, meta = (HidePin = "WorldContextObject", DefaultToSelf = "WorldContextObject", BlueprintInternalUseOnly = "true"), Category = "Abilities")
 	bool BeginSpawningActor(UObject* WorldContextObject, TSubclassOf<AGameplayAbilityTargetActor> Class, AGameplayAbilityTargetActor*& SpawnedActor);
