Additional nullptr checks

diff --git a/Engine/Source/Runtime/GameplayAbilities/Private/GameplayEffect.cpp b/Engine/Source/Runtime/GameplayAbilities/Private/GameplayEffect.cpp
--- a/Engine/Source/Runtime/GameplayAbilities/Private/GameplayEffect.cpp
+++ b/Engine/Source/Runtime/GameplayAbilities/Private/GameplayEffect.cpp
@@ -1849,6 +1849,12 @@ FActiveGameplayEffect* FActiveGameplayEffectsContainer::ApplyGameplayEffectSpec(
 /** This is called anytime a new ActiveGameplayEffect is added, on both client and server in all cases */
 void FActiveGameplayEffectsContainer::InternalOnActiveGameplayEffectAdded(FActiveGameplayEffect& Effect)
 {
+	if (Effect.Spec.Def == nullptr)
+	{
+		ABILITY_LOG(Error, TEXT("FActiveGameplayEffectsContainer serialized new GameplayEffect with NULL Def!"));
+		return;
+	}
+
 	GAMEPLAYEFFECT_SCOPE_LOCK();
 
 	// Add our ongoing tag requirements to the dependency map. We will actually check for these tags below.
@@ -1872,6 +1878,12 @@ void FActiveGameplayEffectsContainer::InternalOnActiveGameplayEffectAdded(FActiv
 
 void FActiveGameplayEffectsContainer::AddActiveGameplayEffectGrantedTagsAndModifiers(FActiveGameplayEffect& Effect)
 {
+	if (Effect.Spec.Def == nullptr)
+	{
+		ABILITY_LOG(Error, TEXT("AddActiveGameplayEffectGrantedTagsAndModifiers called with null Def!"));
+		return;
+	}
+
 	GAMEPLAYEFFECT_SCOPE_LOCK();
 
 	// Register this ActiveGameplayEffects modifiers with our Attribute Aggregators
