Fix GameplayCueManager to work properly in cooked builds when UBlueprint is not available outside of editor (bDontLoadBlueprintOutsideEditor=true)

diff --git a/Engine/Source/Runtime/GameplayAbilities/Private/GameplayCueManager.cpp b/Engine/Source/Runtime/GameplayAbilities/Private/GameplayCueManager.cpp
--- a/Engine/Source/Runtime/GameplayAbilities/Private/GameplayCueManager.cpp
+++ b/Engine/Source/Runtime/GameplayAbilities/Private/GameplayCueManager.cpp
@@ -74,8 +74,8 @@ void UGameplayCueManager::HandleGameplayCueNotify_Internal(AActor* TargetActor,
 			}
 
 			// See if the object is loaded but just not hooked up here
-			UObject* FoundObject = FindObject<UObject>(nullptr, *CueData.GameplayCueNotifyObj.ToString());
-			if (FoundObject == nullptr)
+			CueData.LoadedGameplayCueClass = FindObject<UClass>(nullptr, *CueData.GameplayCueNotifyObj.ToString());
+			if (CueData.LoadedGameplayCueClass == nullptr)
 			{
 				// Not loaded: start async loading and return
 				StreamableManager.SimpleAsyncLoad(CueData.GameplayCueNotifyObj);
@@ -83,30 +83,6 @@ void UGameplayCueManager::HandleGameplayCueNotify_Internal(AActor* TargetActor,
 				ABILITY_LOG(Warning, TEXT("GameplayCueNotify %s was not loaded when GameplayCue was invoked. Starting async loading."), *CueData.GameplayCueNotifyObj.ToString());
 				return;
 			}
-			else
-			{
-				// Found it - did we load a non-instanced version (UGameplayCueNotify_Static) or an instanced version (AGameplayCueNotify_Actor)
-				UObject* LoadedGameplayCueNotify = Cast<AGameplayCueNotify_Actor>(FoundObject);
-				if (!LoadedGameplayCueNotify)
-				{
-					//Try the other class
-					LoadedGameplayCueNotify = Cast<UGameplayCueNotify_Static>(FoundObject);
-				}
-				if (!LoadedGameplayCueNotify)
-				{
-					// Not a dataasset - maybe a blueprint
-					UBlueprint* GameplayCueBlueprint = Cast<UBlueprint>(FoundObject);
-					if (GameplayCueBlueprint && GameplayCueBlueprint->GeneratedClass)
-					{
-						CueData.LoadedGameplayCueClass = GameplayCueBlueprint->GeneratedClass;
-					}
-					else
-					{
-						ABILITY_LOG(Warning, TEXT("GameplayCueNotify %s loaded object %s that is not a GameplayCueNotify"), *CueData.GameplayCueNotifyObj.ToString(), *FoundObject->GetName());
-						return;
-					}
-				}
-			}
 		}
 
 		// Handle the Notify if we found something
@@ -255,46 +231,40 @@ void UGameplayCueManager::LoadObjectLibrary_Internal()
 	TArray<FAssetData> StaticAssetDatas;
 	GameplayCueNotifyStaticObjectLibrary->GetAssetDataList(StaticAssetDatas);
 
-	IGameplayTagsModule& GameplayTagsModule = IGameplayTagsModule::Get();
-
 	GameplayCueData.Empty();
 	GameplayCueDataMap.Empty();
 
-	for (FAssetData Data: ActorAssetDatas)
+	AddAssetDataList_Internal(ActorAssetDatas, GET_MEMBER_NAME_CHECKED(AGameplayCueNotify_Actor, GameplayCueName));
+	AddAssetDataList_Internal(StaticAssetDatas, GET_MEMBER_NAME_CHECKED(UGameplayCueNotify_Static, GameplayCueName));
+
+	BuildAccelerationMap_Internal();
+}
+
+void UGameplayCueManager::AddAssetDataList_Internal(const TArray<FAssetData>& AssetDataList, FName TagPropertyName)
+{
+	IGameplayTagsModule& GameplayTagsModule = IGameplayTagsModule::Get();
+
+	for (FAssetData Data: AssetDataList)
 	{
-		const FString* FoundGameplayTag = Data.TagsAndValues.Find(GET_MEMBER_NAME_CHECKED(AGameplayCueNotify_Actor, GameplayCueName));
+		const FString* FoundGameplayTag = Data.TagsAndValues.Find(TagPropertyName);
 		if (FoundGameplayTag && FoundGameplayTag->Equals(TEXT("None")) == false)
 		{
-			ABILITY_LOG(Warning, TEXT("Found: %s"), **FoundGameplayTag);
-
-			FGameplayTag  GameplayCueTag = GameplayTagsModule.GetGameplayTagsManager().RequestGameplayTag(FName(**FoundGameplayTag), false);
-			if (GameplayCueTag.IsValid())
+			const FString* GeneratedClassTag = Data.TagsAndValues.Find(TEXT("GeneratedClass"));
+			if (GeneratedClassTag == nullptr)
 			{
-				// Add a new NotifyData entry to our flat list for this one
-				FStringAssetReference StringRef;
-				StringRef.AssetLongPathname = Data.ObjectPath.ToString();
-				AddGameplayCueData_Internal(GameplayCueTag, StringRef);				
-			}
-			else
-			{
-				ABILITY_LOG(Warning, TEXT("Found GameplayCue tag %s in asset %s but there is no corresponding tag in the GameplayTagMAnager."), **FoundGameplayTag, *Data.PackageName.ToString());
+				ABILITY_LOG(Warning, TEXT("Unable to find GeneratedClass value for AssetData %s"), *Data.ObjectPath.ToString());
+				continue;
 			}
-		}
-	}
-	for (FAssetData Data : StaticAssetDatas)
-	{
-		const FString* FoundGameplayTag = Data.TagsAndValues.Find(GET_MEMBER_NAME_CHECKED(UGameplayCueNotify_Static, GameplayCueName));
-		if (FoundGameplayTag && FoundGameplayTag->Equals(TEXT("None")) == false)
-		{
-			ABILITY_LOG(Warning, TEXT("Found: %s"), **FoundGameplayTag);
+
+			ABILITY_LOG(Log, TEXT("GameplayCueManager Found: %s / %s"), **FoundGameplayTag, **GeneratedClassTag);
 
 			FGameplayTag  GameplayCueTag = GameplayTagsModule.GetGameplayTagsManager().RequestGameplayTag(FName(**FoundGameplayTag), false);
 			if (GameplayCueTag.IsValid())
 			{
 				// Add a new NotifyData entry to our flat list for this one
 				FStringAssetReference StringRef;
-				StringRef.AssetLongPathname = Data.ObjectPath.ToString();
-				AddGameplayCueData_Internal(GameplayCueTag, StringRef);
+				StringRef.AssetLongPathname = *GeneratedClassTag;
+				AddGameplayCueData_Internal(GameplayCueTag, StringRef);				
 			}
 			else
 			{
@@ -302,8 +272,6 @@ void UGameplayCueManager::LoadObjectLibrary_Internal()
 			}
 		}
 	}
-
-	BuildAccelerationMap_Internal();
 }
 
 void UGameplayCueManager::AddGameplayCueData_Internal(FGameplayTag  GameplayCueTag, FStringAssetReference StringRef)
@@ -322,6 +290,8 @@ void UGameplayCueManager::AddGameplayCueData_Internal(FGameplayTag  GameplayCueT
 	NewData.GameplayCueNotifyObj = StringRef;
 	NewData.GameplayCueTag = GameplayCueTag;
 
+	UObject* FoundObject = FindObject<UObject>(nullptr, *NewData.GameplayCueNotifyObj.ToString());
+
 	GameplayCueData.Add(NewData);
 }
 
@@ -377,21 +347,20 @@ void UGameplayCueManager::HandleAssetAdded(UObject *Object)
 	if (Blueprint && Blueprint->GeneratedClass)
 	{
 		UGameplayCueNotify_Static* StaticCDO = Cast<UGameplayCueNotify_Static>(Blueprint->GeneratedClass->ClassDefaultObject);
+		AGameplayCueNotify_Actor* ActorCDO = Cast<AGameplayCueNotify_Actor>(Blueprint->GeneratedClass->ClassDefaultObject);
+		
 		FStringAssetReference StringRef;
-		StringRef.AssetLongPathname = Blueprint->GetPathName();
-		if (StaticCDO && StaticCDO->GameplayCueTag.IsValid())
+		StringRef.AssetLongPathname = Blueprint->GeneratedClass->GetPathName();
+		
+		if (StaticCDO)
 		{
 			AddGameplayCueData_Internal(StaticCDO->GameplayCueTag, StringRef);
-			BuildAccelerationMap_Internal();		//RICKH Could this just be done by setting the acceleration map rebuild bool to true?
+			BuildAccelerationMap_Internal();
 		}
-		else
+		else if (ActorCDO)
 		{
-			AGameplayCueNotify_Actor* ActorCDO = Cast<AGameplayCueNotify_Actor>(Blueprint->GeneratedClass->ClassDefaultObject);
-			if (ActorCDO && ActorCDO->GameplayCueTag.IsValid())
-			{
-				AddGameplayCueData_Internal(ActorCDO->GameplayCueTag, StringRef);
-				BuildAccelerationMap_Internal();		//RICKH Could this just be done by setting the acceleration map rebuild bool to true?
-			}
+			AddGameplayCueData_Internal(ActorCDO->GameplayCueTag, StringRef);
+			BuildAccelerationMap_Internal();
 		}
 	}
 }
@@ -400,7 +369,6 @@ void UGameplayCueManager::HandleAssetAdded(UObject *Object)
 void UGameplayCueManager::HandleAssetDeleted(UObject *Object)
 {
 	FStringAssetReference StringRefToRemove;
-
 	UBlueprint* Blueprint = Cast<UBlueprint>(Object);
 	if (Blueprint && Blueprint->GeneratedClass)
 	{
@@ -409,7 +377,7 @@ void UGameplayCueManager::HandleAssetDeleted(UObject *Object)
 		
 		if (StaticCDO || ActorCDO)
 		{
-			StringRefToRemove.AssetLongPathname = Blueprint->GetPathName();
+			StringRefToRemove.AssetLongPathname = Blueprint->GeneratedClass->GetPathName();
 		}
 	}
 
diff --git a/Engine/Source/Runtime/GameplayAbilities/Public/GameplayCueManager.h b/Engine/Source/Runtime/GameplayAbilities/Public/GameplayCueManager.h
--- a/Engine/Source/Runtime/GameplayAbilities/Public/GameplayCueManager.h
+++ b/Engine/Source/Runtime/GameplayAbilities/Public/GameplayCueManager.h
@@ -148,6 +148,8 @@ private:
 
 	void AddGameplayCueData_Internal(FGameplayTag  GameplayCueTag, FStringAssetReference StringRef);
 
+	void AddAssetDataList_Internal(const TArray<FAssetData>& AssetDataList, FName TagPropertyName);
+
 	void BuildAccelerationMap_Internal();
 
 	TArray<FString>	LoadedPaths;
