#UE4 Fixing CIS. Forward declare UDataTable and provide include for FCurveTableRowHandle

diff --git a/Engine/Source/Runtime/GameplayAbilities/Public/AttributeSet.h b/Engine/Source/Runtime/GameplayAbilities/Public/AttributeSet.h
--- a/Engine/Source/Runtime/GameplayAbilities/Public/AttributeSet.h
+++ b/Engine/Source/Runtime/GameplayAbilities/Public/AttributeSet.h
@@ -2,6 +2,7 @@
 
 #pragma once
 
+#include "Engine/CurveTable.h"
 #include "AttributeSet.generated.h"
 
 class UAbilitySystemComponent;
@@ -114,7 +115,7 @@ public:
 	/** This signifies the attribute set can be ID'd by name over the network. */
 	void SetNetAddressable();
 
-	void InitFromMetaDataTable(const UDataTable* DataTable);
+	void InitFromMetaDataTable(const class UDataTable* DataTable);
 
 	UAbilitySystemComponent* GetOwningAbilitySystemComponent() const;
 	FGameplayAbilityActorInfo* GetActorInfo() const;
