#UE4 Fixing a few more CIS errors. It seems header includes have shifted.

diff --git a/Engine/Source/Developer/GameplayDebugger/Private/GameplayDebuggingComponent.cpp b/Engine/Source/Developer/GameplayDebugger/Private/GameplayDebuggingComponent.cpp
--- a/Engine/Source/Developer/GameplayDebugger/Private/GameplayDebuggingComponent.cpp
+++ b/Engine/Source/Developer/GameplayDebugger/Private/GameplayDebuggingComponent.cpp
@@ -16,6 +16,7 @@
 #include "BrainComponent.h"
 #include "BehaviorTreeDelegates.h"
 #include "GameFramework/PlayerState.h"
+#include "GameFramework/Character.h"
 #include "Engine/Channel.h"
 
 #if WITH_EDITOR
diff --git a/Engine/Source/Runtime/GameplayAbilities/Public/Abilities/GameplayAbilityTypes.h b/Engine/Source/Runtime/GameplayAbilities/Public/Abilities/GameplayAbilityTypes.h
--- a/Engine/Source/Runtime/GameplayAbilities/Public/Abilities/GameplayAbilityTypes.h
+++ b/Engine/Source/Runtime/GameplayAbilities/Public/Abilities/GameplayAbilityTypes.h
@@ -5,6 +5,7 @@
 #include "GameplayEffect.h"
 #include "GameplayPrediction.h"
 #include "AttributeSet.h"
+#include "GameFramework/MovementComponent.h"
 #include "GameplayAbilityTypes.generated.h"
 
 class UGameplayEffect;
diff --git a/Engine/Source/Runtime/GameplayAbilities/Public/AttributeSet.h b/Engine/Source/Runtime/GameplayAbilities/Public/AttributeSet.h
--- a/Engine/Source/Runtime/GameplayAbilities/Public/AttributeSet.h
+++ b/Engine/Source/Runtime/GameplayAbilities/Public/AttributeSet.h
@@ -3,6 +3,7 @@
 #pragma once
 
 #include "Engine/CurveTable.h"
+#include "Engine/DataTable.h"
 #include "AttributeSet.generated.h"
 
 class UAbilitySystemComponent;
@@ -115,7 +116,7 @@ public:
 	/** This signifies the attribute set can be ID'd by name over the network. */
 	void SetNetAddressable();
 
-	void InitFromMetaDataTable(const class UDataTable* DataTable);
+	void InitFromMetaDataTable(const UDataTable* DataTable);
 
 	UAbilitySystemComponent* GetOwningAbilitySystemComponent() const;
 	FGameplayAbilityActorInfo* GetActorInfo() const;
