- Fix bug in stacking multiplicative modifier math calculations
- Fix bug with client-side stacking math
- Rename the versions of SetMagnitude and GetMagnitude on GESpec that take an FName to Set/GetSetByCallerMagnitude to make it more obvious they're explicitly for SetByCaller and aren't mis-used elsewhere; Still not great names, but more clear that old ones
- Remove the ChanceToExecute setting, as it was completely un-used after the last refactor; Can re-add it later w/o too much difficulty if needed
- Added new GetModifierMagnitude function to GESpec that takes an index and whether to factor in the spec's stack count or not
- Add new utility functions and a new GameplayEffectUtilities namespace; Currently house common math functions to resolve stacking

diff --git a/Engine/Source/Editor/GameplayAbilitiesEditor/Private/AssetTypeActions_GameplayEffect.cpp b/Engine/Source/Editor/GameplayAbilitiesEditor/Private/AssetTypeActions_GameplayEffect.cpp
--- a/Engine/Source/Editor/GameplayAbilitiesEditor/Private/AssetTypeActions_GameplayEffect.cpp
+++ b/Engine/Source/Editor/GameplayAbilitiesEditor/Private/AssetTypeActions_GameplayEffect.cpp
@@ -24,7 +24,6 @@ void CopyInfoFromDataGE(const UGameplayEffect& InGE, UGameplayEffect& OutGE)
 	}
 
 	OutGE.ChanceToApplyToTarget = InGE.ChanceToApplyToTarget;
-	OutGE.ChanceToExecuteOnGameplayEffect = InGE.ChanceToExecuteOnGameplayEffect;
 
 	// Don't copy over TargetEffects. These will need to get fixed up later
 	// TArray<UDataGameplayEffect*> TargetEffects;
diff --git a/Engine/Source/Editor/GameplayAbilitiesEditor/Private/SGraphNodeK2GameplayEffectVar.cpp b/Engine/Source/Editor/GameplayAbilitiesEditor/Private/SGraphNodeK2GameplayEffectVar.cpp
--- a/Engine/Source/Editor/GameplayAbilitiesEditor/Private/SGraphNodeK2GameplayEffectVar.cpp
+++ b/Engine/Source/Editor/GameplayAbilitiesEditor/Private/SGraphNodeK2GameplayEffectVar.cpp
@@ -80,10 +80,9 @@ void SGraphNodeK2GameplayEffectVar::UpdateGraphNode()
 			SNew(SSpacer)
 		];
 
-	if (GameplayEffect && (GameplayEffect->ChanceToApplyToTarget.Value < 1.f || GameplayEffect->ChanceToExecuteOnGameplayEffect.Value < 1.f))
+	if (GameplayEffect && (GameplayEffect->ChanceToApplyToTarget.Value < 1.f))
 	{
 		FText ApplyToTargetText = FText::AsNumber(GameplayEffect->ChanceToApplyToTarget.Value);
-		FText ExecuteOnGEText = FText::AsNumber(GameplayEffect->ChanceToExecuteOnGameplayEffect.Value);
 
 		RightNodeBox->AddSlot()
 			[
@@ -102,23 +101,6 @@ void SGraphNodeK2GameplayEffectVar::UpdateGraphNode()
 					]
 			];
 
-		RightNodeBox->AddSlot()
-			[
-				SNew(SHorizontalBox)
-				+ SHorizontalBox::Slot()
-				.HAlign(HAlign_Left)
-				[
-					SNew(STextBlock)
-					.Text(LOCTEXT("ChanceToExecuteOnGE", "Chance to Execute on GE:"))
-				]
-				+ SHorizontalBox::Slot()
-					.HAlign(HAlign_Right)
-					[
-						SNew(STextBlock)
-						.Text(ExecuteOnGEText)
-					]
-			];
-
 		RightNodeBox->AddSlot()
 			.FillHeight(0.3f)
 			[
diff --git a/Engine/Source/Runtime/GameplayAbilities/Private/AbilitySystemBlueprintLibrary.cpp b/Engine/Source/Runtime/GameplayAbilities/Private/AbilitySystemBlueprintLibrary.cpp
--- a/Engine/Source/Runtime/GameplayAbilities/Private/AbilitySystemBlueprintLibrary.cpp
+++ b/Engine/Source/Runtime/GameplayAbilities/Private/AbilitySystemBlueprintLibrary.cpp
@@ -409,7 +409,7 @@ FGameplayEffectSpecHandle UAbilitySystemBlueprintLibrary::AssignSetByCallerMagni
 	FGameplayEffectSpec* Spec = SpecHandle.Data.Get();
 	if (Spec)
 	{
-		Spec->SetMagnitude(DataName, Magnitude);
+		Spec->SetSetByCallerMagnitude(DataName, Magnitude);
 	}
 	else
 	{
diff --git a/Engine/Source/Runtime/GameplayAbilities/Private/GameplayEffect.cpp b/Engine/Source/Runtime/GameplayAbilities/Private/GameplayEffect.cpp
--- a/Engine/Source/Runtime/GameplayAbilities/Private/GameplayEffect.cpp
+++ b/Engine/Source/Runtime/GameplayAbilities/Private/GameplayEffect.cpp
@@ -29,7 +29,6 @@ UGameplayEffect::UGameplayEffect(const FObjectInitializer& ObjectInitializer)
 	DurationPolicy = EGameplayEffectDurationType::Instant;
 	bExecutePeriodicEffectOnApplication = true;
 	ChanceToApplyToTarget.SetValue(1.f);
-	ChanceToExecuteOnGameplayEffect.SetValue(1.f);
 	StackingType = EGameplayEffectStackingType::None;
 	StackLimitCount = 0;
 	StackDurationRefreshPolicy = EGameplayEffectStackingDurationPolicy::RefreshOnSuccessfulApplication;
@@ -262,7 +261,7 @@ bool FGameplayEffectModifierMagnitude::AttemptCalculateMagnitude(const FGameplay
 			break;
 
 			case EGameplayEffectMagnitudeCalculation::SetByCaller:
-				OutCalculatedMagnitude = InRelevantSpec.GetMagnitude(SetByCallerMagnitude.DataName);
+				OutCalculatedMagnitude = InRelevantSpec.GetSetByCallerMagnitude(SetByCallerMagnitude.DataName);
 				break;
 			default:
 				ABILITY_LOG(Error, TEXT("Unknown MagnitudeCalculationType %d in AttemptCalculateMagnitude"), (int32)MagnitudeCalculationType);
@@ -525,7 +524,6 @@ void FGameplayEffectSpec::SetLevel(float InLevel)
 		
 		Period = Def->Period.GetValueAtLevel(InLevel);
 		ChanceToApplyToTarget = Def->ChanceToApplyToTarget.GetValueAtLevel(InLevel);
-		ChanceToExecuteOnGameplayEffect = Def->ChanceToExecuteOnGameplayEffect.GetValueAtLevel(InLevel);
 	}
 }
 
@@ -564,9 +562,19 @@ float FGameplayEffectSpec::GetChanceToApplyToTarget() const
 	return ChanceToApplyToTarget;
 }
 
-float FGameplayEffectSpec::GetChanceToExecuteOnGameplayEffect() const
+float FGameplayEffectSpec::GetModifierMagnitude(int32 ModifierIdx, bool bFactorInStackCount) const
 {
-	return ChanceToExecuteOnGameplayEffect;
+	check(Modifiers.IsValidIndex(ModifierIdx) && Def && Def->Modifiers.IsValidIndex(ModifierIdx));
+
+	const float SingleEvaluatedMagnitude = Modifiers[ModifierIdx].GetEvaluatedMagnitude();
+
+	float ModMagnitude = SingleEvaluatedMagnitude;
+	if (bFactorInStackCount)
+	{
+		ModMagnitude = GameplayEffectUtilities::ComputeStackedModifierMagnitude(SingleEvaluatedMagnitude, StackCount, Def->Modifiers[ModifierIdx].ModifierOp);
+	}
+
+	return ModMagnitude;
 }
 
 void FGameplayEffectSpec::CalculateModifierMagnitudes()
@@ -584,24 +592,6 @@ void FGameplayEffectSpec::CalculateModifierMagnitudes()
 	}
 }
 
-float FGameplayEffectSpec::GetMagnitude(const FGameplayAttribute &Attribute) const
-{
-	for (int32 ModIdx = 0; ModIdx < Modifiers.Num(); ++ModIdx)
-	{
-		const FGameplayModifierInfo& ModDef = Def->Modifiers[ModIdx];
-		const FModifierSpec& ModSpec = Modifiers[ModIdx];
-	
-		if (ModDef.Attribute != Attribute)
-		{
-			continue;
-		}
-
-		return ModSpec.EvaluatedMagnitude;
-	}
-
-	return 0.f;
-}
-
 bool FGameplayEffectSpec::HasValidCapturedAttributes(const TArray<FGameplayEffectAttributeCaptureDefinition>& InCaptureDefsToCheck) const
 {
 	return CapturedRelevantAttributes.HasValidCapturedAttributes(InCaptureDefsToCheck);
@@ -676,7 +666,7 @@ void FGameplayEffectSpec::GetAllGrantedTags(OUT FGameplayTagContainer& Container
 	Container.AppendTags(Def->InheritableOwnedTagsContainer.CombinedTags);
 }
 
-void FGameplayEffectSpec::SetMagnitude(FName DataName, float Magnitude)
+void FGameplayEffectSpec::SetSetByCallerMagnitude(FName DataName, float Magnitude)
 {
 #if !(UE_BUILD_SHIPPING || UE_BUILD_TEST)
 	const float* CurrentValue = SetByCallerMagnitudes.Find(DataName);
@@ -689,7 +679,7 @@ void FGameplayEffectSpec::SetMagnitude(FName DataName, float Magnitude)
 	SetByCallerMagnitudes.Add(DataName) = Magnitude;
 }
 
-float FGameplayEffectSpec::GetMagnitude(FName DataName) const
+float FGameplayEffectSpec::GetSetByCallerMagnitude(FName DataName) const
 {
 	float Magnitude = 0.f;
 	const float* Ptr = SetByCallerMagnitudes.Find(DataName);
@@ -1038,6 +1028,9 @@ void FActiveGameplayEffect::PostReplicatedChange(const struct FActiveGameplayEff
 		StartWorldTime = InArray.GetWorldTime() - static_cast<float>(InArray.GetGameStateTime() - StartGameStateTime);
 		CachedStartGameStateTime = StartGameStateTime;
 	}
+
+	// Const cast is ok. It is there to prevent mutation of the GameplayEffects array, which this wont do.
+	const_cast<FActiveGameplayEffectsContainer&>(InArray).UpdateAllAggregatorModMagnitudes(*this);
 }
 
 // --------------------------------------------------------------------------------------------------------------------------------------------------------
@@ -1073,8 +1066,6 @@ void FActiveGameplayEffectsContainer::ExecuteActiveEffectsFrom(FGameplayEffectSp
 	
 	FGameplayEffectSpec& SpecToUse = Spec;
 
-	float ChanceToExecute = SpecToUse.GetChanceToExecuteOnGameplayEffect();		// Not implemented? Should we just remove ChanceToExecute?
-
 	// Capture our own tags.
 	// TODO: We should only capture them if we need to. We may have snapshotted target tags (?) (in the case of dots with exotic setups?)
 
@@ -1091,9 +1082,8 @@ void FActiveGameplayEffectsContainer::ExecuteActiveEffectsFrom(FGameplayEffectSp
 	for (int32 ModIdx = 0; ModIdx < SpecToUse.Modifiers.Num(); ++ModIdx)
 	{
 		const FGameplayModifierInfo& ModDef = SpecToUse.Def->Modifiers[ModIdx];
-		const FModifierSpec& ModSpec = SpecToUse.Modifiers[ModIdx];
-
-		FGameplayModifierEvaluatedData EvalData(ModDef.Attribute, ModDef.ModifierOp, ModSpec.GetEvaluatedMagnitude());
+		
+		FGameplayModifierEvaluatedData EvalData(ModDef.Attribute, ModDef.ModifierOp, SpecToUse.GetModifierMagnitude(ModIdx, true));
 		InvokeGameplayCueExecute |= InternalExecuteMod(SpecToUse, EvalData);
 	}
 
@@ -1140,7 +1130,7 @@ void FActiveGameplayEffectsContainer::ExecuteActiveEffectsFrom(FGameplayEffectSp
 				// If the execution didn't manually handle the stack count, automatically apply it here
 				if (bApplyStackCountToEmittedMods && SpecStackCount > 1)
 				{
-					CurExecMod.Magnitude *= SpecStackCount;
+					CurExecMod.Magnitude = GameplayEffectUtilities::ComputeStackedModifierMagnitude(CurExecMod.Magnitude, SpecStackCount, CurExecMod.ModifierOp);
 				}
 				InvokeGameplayCueExecute |= InternalExecuteMod(SpecToUse, CurExecMod);
 			}
@@ -1329,10 +1319,15 @@ void FActiveGameplayEffectsContainer::OnMagnitudeDependencyChange(FActiveGamepla
 }
 
 void FActiveGameplayEffectsContainer::OnStackCountChange(FActiveGameplayEffect& ActiveEffect)
+{
+	UpdateAllAggregatorModMagnitudes(ActiveEffect);
+}
+
+void FActiveGameplayEffectsContainer::UpdateAllAggregatorModMagnitudes(FActiveGameplayEffect& ActiveEffect)
 {
 	const FGameplayEffectSpec& Spec = ActiveEffect.Spec;
 	TSet<FGameplayAttribute> AttributesToUpdate;
-	
+
 	for (int32 ModIdx = 0; ModIdx < Spec.Modifiers.Num(); ++ModIdx)
 	{
 		const FGameplayModifierInfo& ModDef = Spec.Def->Modifiers[ModIdx];
@@ -1354,15 +1349,14 @@ void FActiveGameplayEffectsContainer::UpdateAggregatorModMagnitudes(const TSet<F
 		// Todo: it would be nice to have a better way of updating this. Some ideas: instead of remove, call an 'invalidate' to NAN out all floats for a given handle, then call an 'update' to set real values.
 		Aggregator->RemoveAggregatorMod(ActiveEffect.Handle);
 
-		// Now readd ALL of our mods
+		// Now re-add ALL of our mods
 		for(int32 ModIdx = 0; ModIdx < Spec.Modifiers.Num(); ++ModIdx)
 		{
 			const FGameplayModifierInfo& ModDef = Spec.Def->Modifiers[ModIdx];
-			const FModifierSpec& ModSpec = Spec.Modifiers[ModIdx];
 
 			if (ModDef.Attribute == Attribute)
 			{
-				Aggregator->AddAggregatorMod(ModSpec.GetEvaluatedMagnitude() * static_cast<float>(Spec.StackCount), ModDef.ModifierOp, &ModDef.SourceTags, &ModDef.TargetTags, ActiveEffect.PredictionKey.WasLocallyGenerated(), ActiveEffect.Handle);
+				Aggregator->AddAggregatorMod(Spec.GetModifierMagnitude(ModIdx, true), ModDef.ModifierOp, &ModDef.SourceTags, &ModDef.TargetTags, ActiveEffect.PredictionKey.WasLocallyGenerated(), ActiveEffect.Handle);
 			}
 		}
 	}
@@ -1891,14 +1885,13 @@ void FActiveGameplayEffectsContainer::AddActiveGameplayEffectGrantedTagsAndModif
 	{
 		for (int32 ModIdx = 0; ModIdx < Effect.Spec.Modifiers.Num(); ++ModIdx)
 		{
-			const FModifierSpec &Mod = Effect.Spec.Modifiers[ModIdx];
 			const FGameplayModifierInfo &ModInfo = Effect.Spec.Def->Modifiers[ModIdx];
 
 			// Note we assume the EvaluatedMagnitude is up to do. There is no case currently where we should recalculate magnitude based on
 			// Ongoing tags being met. We either calculate magnitude one time, or its done via OnDirty calls (or potentially a frequency timer one day)
 					
 			FAggregator* Aggregator = FindOrCreateAttributeAggregator(Effect.Spec.Def->Modifiers[ModIdx].Attribute).Get();
-			Aggregator->AddAggregatorMod(Mod.GetEvaluatedMagnitude() * static_cast<float>(Effect.Spec.StackCount), ModInfo.ModifierOp, &ModInfo.SourceTags, &ModInfo.TargetTags, Effect.PredictionKey.WasLocallyGenerated(), Effect.Handle);
+			Aggregator->AddAggregatorMod(Effect.Spec.GetModifierMagnitude(ModIdx, true), ModInfo.ModifierOp, &ModInfo.SourceTags, &ModInfo.TargetTags, Effect.PredictionKey.WasLocallyGenerated(), Effect.Handle);
 		}
 	}
 
diff --git a/Engine/Source/Runtime/GameplayAbilities/Private/GameplayEffectAggregator.cpp b/Engine/Source/Runtime/GameplayAbilities/Private/GameplayEffectAggregator.cpp
--- a/Engine/Source/Runtime/GameplayAbilities/Private/GameplayEffectAggregator.cpp
+++ b/Engine/Source/Runtime/GameplayAbilities/Private/GameplayEffectAggregator.cpp
@@ -51,9 +51,9 @@ float FAggregator::EvaluateWithBase(float InlineBaseValue, const FAggregatorEval
 		}
 	}
 
-	float Additive = SumMods(Mods[EGameplayModOp::Additive], 0.f, Parameters);
-	float Multiplicitive = SumMods(Mods[EGameplayModOp::Multiplicitive], 1.f, Parameters);
-	float Division = SumMods(Mods[EGameplayModOp::Division], 1.f, Parameters);
+	float Additive = SumMods(Mods[EGameplayModOp::Additive], GameplayEffectUtilities::GetModifierBiasByModifierOp(EGameplayModOp::Additive), Parameters);
+	float Multiplicitive = SumMods(Mods[EGameplayModOp::Multiplicitive], GameplayEffectUtilities::GetModifierBiasByModifierOp(EGameplayModOp::Multiplicitive), Parameters);
+	float Division = SumMods(Mods[EGameplayModOp::Division], GameplayEffectUtilities::GetModifierBiasByModifierOp(EGameplayModOp::Division), Parameters);
 
 	if (FMath::IsNearlyZero(Division))
 	{
@@ -75,9 +75,9 @@ float FAggregator::ReverseEvaluate(float FinalValue, const FAggregatorEvaluatePa
 		}
 	}
 
-	float Additive = SumMods(Mods[EGameplayModOp::Additive], 0.f, Parameters);
-	float Multiplicitive = SumMods(Mods[EGameplayModOp::Multiplicitive], 1.f, Parameters);
-	float Division = SumMods(Mods[EGameplayModOp::Division], 1.f, Parameters);
+	float Additive = SumMods(Mods[EGameplayModOp::Additive], GameplayEffectUtilities::GetModifierBiasByModifierOp(EGameplayModOp::Additive), Parameters);
+	float Multiplicitive = SumMods(Mods[EGameplayModOp::Multiplicitive], GameplayEffectUtilities::GetModifierBiasByModifierOp(EGameplayModOp::Multiplicitive), Parameters);
+	float Division = SumMods(Mods[EGameplayModOp::Division], GameplayEffectUtilities::GetModifierBiasByModifierOp(EGameplayModOp::Division), Parameters);
 
 	if (FMath::IsNearlyZero(Division))
 	{
diff --git a/Engine/Source/Runtime/GameplayAbilities/Private/GameplayEffectTypes.cpp b/Engine/Source/Runtime/GameplayAbilities/Private/GameplayEffectTypes.cpp
--- a/Engine/Source/Runtime/GameplayAbilities/Private/GameplayEffectTypes.cpp
+++ b/Engine/Source/Runtime/GameplayAbilities/Private/GameplayEffectTypes.cpp
@@ -6,6 +6,34 @@
 #include "GameplayEffectTypes.h"
 #include "AbilitySystemComponent.h"
 
+float GameplayEffectUtilities::GetModifierBiasByModifierOp(EGameplayModOp::Type ModOp)
+{
+	static const float ModifierOpBiases[EGameplayModOp::Max] = {0.f, 1.f, 1.f, 0.f};
+	check(ModOp >= 0 && ModOp < EGameplayModOp::Max);
+
+	return ModifierOpBiases[ModOp];
+}
+
+float GameplayEffectUtilities::ComputeStackedModifierMagnitude(float BaseComputedMagnitude, int32 StackCount, EGameplayModOp::Type ModOp)
+{
+	const float OperationBias = GameplayEffectUtilities::GetModifierBiasByModifierOp(ModOp);
+
+	StackCount = FMath::Clamp<int32>(StackCount, 0, StackCount);
+
+	float StackMag = BaseComputedMagnitude;
+	
+	// Override modifiers don't care about stack count at all. All other modifier ops need to subtract out their bias value in order to handle
+	// stacking correctly
+	if (ModOp != EGameplayModOp::Override)
+	{
+		StackMag -= OperationBias;
+		StackMag *= StackCount;
+		StackMag += OperationBias;
+	}
+
+	return StackMag;
+}
+
 bool FGameplayEffectAttributeCaptureDefinition::operator==(const FGameplayEffectAttributeCaptureDefinition& Other) const
 {
 	return ((AttributeToCapture == Other.AttributeToCapture) && (AttributeSource == Other.AttributeSource) && (bSnapshot == Other.bSnapshot));
diff --git a/Engine/Source/Runtime/GameplayAbilities/Public/GameplayEffect.h b/Engine/Source/Runtime/GameplayAbilities/Public/GameplayEffect.h
--- a/Engine/Source/Runtime/GameplayAbilities/Public/GameplayEffect.h
+++ b/Engine/Source/Runtime/GameplayAbilities/Public/GameplayEffect.h
@@ -562,10 +562,6 @@ public:
 	UPROPERTY(EditDefaultsOnly, Category=Application, meta=(GameplayAttribute="True"))
 	FScalableFloat	ChanceToApplyToTarget;
 
-	/** Probability that this gameplay effect will execute on another GE after it has been successfully applied to the target actor (0.0 for never, 1.0 for always) */
-	UPROPERTY(EditDefaultsOnly, Category = Application, meta = (GameplayAttribute = "True"))
-	FScalableFloat	ChanceToExecuteOnGameplayEffect;
-
 	/** other gameplay effects that will be applied to the target of this effect if this effect applies */
 	UPROPERTY(EditDefaultsOnly, BlueprintReadOnly, Category = GameplayEffect,meta = (DisplayName = "Linked Gameplay Effects"))
 	TArray<TSubclassOf<UGameplayEffect>> TargetEffectClasses;
@@ -911,7 +907,6 @@ struct GAMEPLAYABILITIES_API FGameplayEffectSpec
 		, Duration(UGameplayEffect::INSTANT_APPLICATION)
 		, Period(UGameplayEffect::NO_PERIOD)
 		, ChanceToApplyToTarget(1.f)
-		, ChanceToExecuteOnGameplayEffect(1.f)
 		, StackCount(1)
 		, bCompletedSourceAttributeCapture(false)
 		, bCompletedTargetAttributeCapture(false)
@@ -967,8 +962,6 @@ struct GAMEPLAYABILITIES_API FGameplayEffectSpec
 	float GetDuration() const;
 	float GetPeriod() const;
 	float GetChanceToApplyToTarget() const;
-	float GetChanceToExecuteOnGameplayEffect() const;
-	float GetMagnitude(const FGameplayAttribute &Attribute) const;
 
 	/** other effects that need to be applied to the target if this effect is successful */
 	TArray< FGameplayEffectSpecHandle > TargetEffectSpecs;
@@ -984,10 +977,10 @@ struct GAMEPLAYABILITIES_API FGameplayEffectSpec
 	void GetAllGrantedTags(OUT FGameplayTagContainer& Container) const;
 
 	/** Sets the magnitude of a SetByCaller modifier */
-	void SetMagnitude(FName DataName, float Magnitude);
+	void SetSetByCallerMagnitude(FName DataName, float Magnitude);
 
 	/** Returns the magnitude of a SetByCaller modifier. Will return 0.f and Warn if the magnitude has not been set. */
-	float GetMagnitude(FName DataName) const;
+	float GetSetByCallerMagnitude(FName DataName) const;
 
 	// The duration in seconds of this effect
 	// instantaneous effects should have a duration of UGameplayEffect::INSTANT_APPLICATION
@@ -1004,9 +997,6 @@ struct GAMEPLAYABILITIES_API FGameplayEffectSpec
 	UPROPERTY()
 	float ChanceToApplyToTarget;
 
-	UPROPERTY()
-	float ChanceToExecuteOnGameplayEffect;
-
 	// Captured Source Tags on GameplayEffectSpec creation.	
 	UPROPERTY(NotReplicated)
 	FTagContainerAggregator	CapturedSourceTags;
@@ -1037,6 +1027,16 @@ struct GAMEPLAYABILITIES_API FGameplayEffectSpec
 	UPROPERTY(NotReplicated)
 	uint32 bDurationLocked : 1;
 	
+	/**
+	 * Get the computed magnitude of the modifier on the spec with the specified index
+	 * 
+	 * @param ModifierIndx			Modifier to get
+	 * @param bFactorInStackCount	If true, the calculation will include the stack count
+	 * 
+	 * @return Computed magnitude
+	 */
+	float GetModifierMagnitude(int32 ModifierIdx, bool bFactorInStackCount) const;
+
 	void CalculateModifierMagnitudes();
 
 	void SetLevel(float InLevel);
@@ -1414,6 +1414,8 @@ private:
 
 	void OnStackCountChange(FActiveGameplayEffect& ActiveEffect);
 
+	void UpdateAllAggregatorModMagnitudes(FActiveGameplayEffect& ActiveEffect);
+
 	void UpdateAggregatorModMagnitudes(const TSet<FGameplayAttribute>& AttributesToUpdate, FActiveGameplayEffect& ActiveEffect);
 
 	/** Helper function to find the active GE that the specified spec can stack with, if any */
diff --git a/Engine/Source/Runtime/GameplayAbilities/Public/GameplayEffectTypes.h b/Engine/Source/Runtime/GameplayAbilities/Public/GameplayEffectTypes.h
--- a/Engine/Source/Runtime/GameplayAbilities/Public/GameplayEffectTypes.h
+++ b/Engine/Source/Runtime/GameplayAbilities/Public/GameplayEffectTypes.h
@@ -46,6 +46,30 @@ namespace EGameplayModOp
 	};
 }
 
+namespace GameplayEffectUtilities
+{
+	/**
+	 * Helper function to retrieve the modifier bias based upon modifier operation
+	 * 
+	 * @param ModOp	Modifier operation to retrieve the modifier bias for
+	 * 
+	 * @return Modifier bias for the specified operation
+	 */
+	GAMEPLAYABILITIES_API float GetModifierBiasByModifierOp(EGameplayModOp::Type ModOp);
+
+	/**
+	 * Helper function to compute the stacked modifier magnitude from a base magnitude, given a stack count and modifier operation
+	 * 
+	 * @param BaseComputedMagnitude	Base magnitude to compute from
+	 * @param StackCount			Stack count to use for the calculation
+	 * @param ModOp					Modifier operation to use
+	 * 
+	 * @return Computed modifier magnitude with stack count factored in
+	 */
+	GAMEPLAYABILITIES_API float ComputeStackedModifierMagnitude(float BaseComputedMagnitude, int32 StackCount, EGameplayModOp::Type ModOp);
+}
+
+
 /** Enumeration for options of where to capture gameplay attributes from for gameplay effects */
 UENUM()
 enum class EGameplayEffectAttributeCaptureSource : uint8
