FGameplayEffectContextHandle::NetSerialize now assumes struct type based on ability system global alloc function.

diff --git a/Engine/Source/Runtime/GameplayAbilities/Private/GameplayEffectTypes.cpp b/Engine/Source/Runtime/GameplayAbilities/Private/GameplayEffectTypes.cpp
--- a/Engine/Source/Runtime/GameplayAbilities/Private/GameplayEffectTypes.cpp
+++ b/Engine/Source/Runtime/GameplayAbilities/Private/GameplayEffectTypes.cpp
@@ -162,24 +162,25 @@ void FGameplayEffectContext::GetOwnedGameplayTags(OUT FGameplayTagContainer& Act
 
 bool FGameplayEffectContextHandle::NetSerialize(FArchive& Ar, class UPackageMap* Map, bool& bOutSuccess)
 {
-	UScriptStruct* ScriptStruct = Data.IsValid() ? Data->GetScriptStruct() : NULL;
-	Ar << ScriptStruct;
+	bool ValidData = Data.IsValid();
+	Ar << ValidData;
 
-	if (ScriptStruct)
+	if (ValidData)
 	{
 		if (Ar.IsLoading())
 		{
 			// For now, just always reset/reallocate the data when loading.
 			// Longer term if we want to generalize this and use it for property replication, we should support
 			// only reallocating when necessary
-
-			FGameplayEffectContext * NewData = (FGameplayEffectContext*)FMemory::Malloc(ScriptStruct->GetCppStructOps()->GetSize());
-			ScriptStruct->InitializeStruct(NewData);
-
-			Data = TSharedPtr<FGameplayEffectContext>(NewData);
+			
+			if (Data.IsValid() == false)
+			{
+				Data = TSharedPtr<FGameplayEffectContext>(UAbilitySystemGlobals::Get().AllocGameplayEffectContext());
+			}
 		}
 
 		void* ContainerPtr = Data.Get();
+		UScriptStruct* ScriptStruct = Data->GetScriptStruct();
 
 		if (ScriptStruct->StructFlags & STRUCT_NetSerializeNative)
 		{
@@ -192,18 +193,6 @@ bool FGameplayEffectContextHandle::NetSerialize(FArchive& Ar, class UPackageMap*
 			//	2) if there are any UStructProperties in the topmost struct's fields, we will assert in UStructProperty::NetSerializeItem.
 
 			ABILITY_LOG(Fatal, TEXT("FGameplayEffectContextHandle::NetSerialize called on data struct %s without a native NetSerialize"), *ScriptStruct->GetName());
-
-			for (TFieldIterator<UProperty> It(ScriptStruct); It; ++It)
-			{
-				if (It->PropertyFlags & CPF_RepSkip)
-				{
-					continue;
-				}
-
-				void * PropertyData = It->ContainerPtrToValuePtr<void*>(ContainerPtr);
-
-				It->NetSerializeItem(Ar, Map, PropertyData);
-			}
 		}
 	}
 
