UAbilitySystemBlueprintLibrary::GetGameplayCueEndLocationAndNormal helper function

diff --git a/Engine/Source/Runtime/GameplayAbilities/Private/AbilitySystemBlueprintLibrary.cpp b/Engine/Source/Runtime/GameplayAbilities/Private/AbilitySystemBlueprintLibrary.cpp
--- a/Engine/Source/Runtime/GameplayAbilities/Private/AbilitySystemBlueprintLibrary.cpp
+++ b/Engine/Source/Runtime/GameplayAbilities/Private/AbilitySystemBlueprintLibrary.cpp
@@ -402,6 +402,26 @@ FVector UAbilitySystemBlueprintLibrary::GetOrigin(FGameplayCueParameters Paramet
 	return FVector::ZeroVector;
 }
 
+bool UAbilitySystemBlueprintLibrary::GetGameplayCueEndLocationAndNormal(AActor* TargetActor, FGameplayCueParameters Parameters, FVector& Location, FVector& Normal)
+{
+	FGameplayEffectContext* Data = Parameters.EffectContext.Get();
+	if (Data && Data->GetHitResult())
+	{
+		
+		Location = Data->GetHitResult()->Location;
+		Normal = Data->GetHitResult()->Normal;
+		return true;
+	}
+	else if(TargetActor)
+	{
+		Location = TargetActor->GetActorLocation();
+		Normal = TargetActor->GetActorRotation().Vector();
+		return true;
+	}
+
+	return false;
+}
+
 // ---------------------------------------------------------------------------------------
 
 FGameplayEffectSpecHandle UAbilitySystemBlueprintLibrary::AssignSetByCallerMagnitude(FGameplayEffectSpecHandle SpecHandle, FName DataName, float Magnitude)
diff --git a/Engine/Source/Runtime/GameplayAbilities/Public/AbilitySystemBlueprintLibrary.h b/Engine/Source/Runtime/GameplayAbilities/Public/AbilitySystemBlueprintLibrary.h
--- a/Engine/Source/Runtime/GameplayAbilities/Public/AbilitySystemBlueprintLibrary.h
+++ b/Engine/Source/Runtime/GameplayAbilities/Public/AbilitySystemBlueprintLibrary.h
@@ -152,6 +152,10 @@ class GAMEPLAYABILITIES_API UAbilitySystemBlueprintLibrary : public UBlueprintFu
 	UFUNCTION(BlueprintPure, Category = "Ability|GameplayCue")
 	static FVector GetOrigin(FGameplayCueParameters Parameters);
 
+	/** Gets the best end location and normal for this gameplay cue. If there is hit result data, it will return this. Otherwise it will return the target actor's location/rotation. If none of this is available, it will return false. */
+	UFUNCTION(BlueprintPure, Category = "Ability|GameplayCue")
+	static bool GetGameplayCueEndLocationAndNormal(AActor* TargetActor, FGameplayCueParameters Parameters, FVector& Location, FVector& Normal);
+
 
 	// -------------------------------------------------------------------------------
 	//		GameplayEffectSpec
