Properly populate parentIdx on gameplay cue notify data.

diff --git a/Engine/Source/Runtime/GameplayAbilities/Private/GameplayCueManager.cpp b/Engine/Source/Runtime/GameplayAbilities/Private/GameplayCueManager.cpp
--- a/Engine/Source/Runtime/GameplayAbilities/Private/GameplayCueManager.cpp
+++ b/Engine/Source/Runtime/GameplayAbilities/Private/GameplayCueManager.cpp
@@ -313,6 +313,9 @@ void UGameplayCueManager::BuildAccelerationMap_Internal()
 
 	FGameplayTagContainer AllGameplayCueTags = IGameplayTagsModule::Get().GetGameplayTagsManager().RequestGameplayTagChildren(BaseGameplayCueTag());
 
+
+	// Create entries for children.
+	// E.g., if "a.b" notify exists but "a.b.c" does not, point "a.b.c" entry to "a.b"'s notify.
 	for (FGameplayTag ThisGameplayCueTag : AllGameplayCueTags)
 	{
 		if (GameplayCueDataMap.Contains(ThisGameplayCueTag))
@@ -325,7 +328,24 @@ void UGameplayCueManager::BuildAccelerationMap_Internal()
 		GameplayCueDataMap.Add(ThisGameplayCueTag) = GameplayCueDataMap.FindChecked(Parent);
 	}
 
-	// PrintGameplayCueNotifyMap();
+
+	// Build up parentIdx on each item in GameplayCUeData
+	for (FGameplayCueNotifyData& Data : GameplayCueData)
+	{
+		FGameplayTag Parent = Data.GameplayCueTag.RequestDirectParent();
+		while(Parent != BaseGameplayCueTag())
+		{
+			int32* idxPtr = GameplayCueDataMap.Find(Parent);
+			if (idxPtr)
+			{
+				Data.ParentDataIdx = *idxPtr;
+				break;
+			}
+			Parent = Parent.RequestDirectParent();
+		}
+	}
+
+	PrintGameplayCueNotifyMap();
 }
 
 int32 UGameplayCueManager::FinishLoadingGameplayCueNotifies()
