added methods for invoking gameplay cue events with a specific gameplay cue parmeter struct (rather than it always being derived).

diff --git a/Engine/Source/Runtime/GameplayAbilities/Private/AbilitySystemComponent.cpp b/Engine/Source/Runtime/GameplayAbilities/Private/AbilitySystemComponent.cpp
--- a/Engine/Source/Runtime/GameplayAbilities/Private/AbilitySystemComponent.cpp
+++ b/Engine/Source/Runtime/GameplayAbilities/Private/AbilitySystemComponent.cpp
@@ -767,29 +767,25 @@ void UAbilitySystemComponent::InvokeGameplayCueEvent(const FGameplayEffectSpec &
 
 void UAbilitySystemComponent::InvokeGameplayCueEvent(const FGameplayTag GameplayCueTag, EGameplayCueEvent::Type EventType, FGameplayEffectContextHandle EffectContext)
 {
-	AActor* ActorAvatar = AbilityActorInfo->AvatarActor.Get();
-	AActor* ActorOwner = AbilityActorInfo->OwnerActor.Get();
-	IGameplayCueInterface* GameplayCueInterface = Cast<IGameplayCueInterface>(ActorAvatar);
-	if (!GameplayCueInterface)
-	{
-		return;
-	}
-
 	FGameplayCueParameters CueParameters;
 
 	if (EffectContext.IsValid())
 	{
 		CueParameters.EffectContext = EffectContext;
 	}
-	else
-	{
-		CueParameters.EffectContext.AddInstigator(ActorOwner, ActorAvatar); // By default use the owner and avatar as the instigator and causer
-	}
 
 	CueParameters.NormalizedMagnitude = 1.f;
 	CueParameters.RawMagnitude = 0.f;
 
-	UAbilitySystemGlobals::Get().GetGameplayCueManager()->HandleGameplayCues(ActorAvatar, GameplayCueTag, EventType, CueParameters);
+	InvokeGameplayCueEvent(GameplayCueTag, EventType, CueParameters);
+}
+
+void UAbilitySystemComponent::InvokeGameplayCueEvent(const FGameplayTag GameplayCueTag, EGameplayCueEvent::Type EventType, const FGameplayCueParameters& GameplayCueParameters)
+{
+	AActor* ActorAvatar = AbilityActorInfo->AvatarActor.Get();
+	AActor* ActorOwner = AbilityActorInfo->OwnerActor.Get();
+
+	UAbilitySystemGlobals::Get().GetGameplayCueManager()->HandleGameplayCues(ActorAvatar, GameplayCueTag, EventType, GameplayCueParameters);
 }
 
 void UAbilitySystemComponent::ExecuteGameplayCue(const FGameplayTag GameplayCueTag, FGameplayEffectContextHandle EffectContext)
@@ -805,6 +801,19 @@ void UAbilitySystemComponent::ExecuteGameplayCue(const FGameplayTag GameplayCueT
 	}
 }
 
+void UAbilitySystemComponent::ExecuteGameplayCue(const FGameplayTag GameplayCueTag, const FGameplayCueParameters& GameplayCueParameters)
+{
+	if (IsOwnerActorAuthoritative())
+	{
+		ForceReplication();
+		NetMulticast_InvokeGameplayCueExecuted_WithParams(GameplayCueTag, ScopedPredictionKey, GameplayCueParameters);
+	}
+	else if (ScopedPredictionKey.IsValidKey())
+	{
+		InvokeGameplayCueEvent(GameplayCueTag, EGameplayCueEvent::Executed, GameplayCueParameters);
+	}
+}
+
 void UAbilitySystemComponent::AddGameplayCue(const FGameplayTag GameplayCueTag, FGameplayEffectContextHandle EffectContext)
 {
 	if (IsOwnerActorAuthoritative())
@@ -868,6 +877,14 @@ void UAbilitySystemComponent::NetMulticast_InvokeGameplayCueExecuted_Implementat
 	}
 }
 
+void UAbilitySystemComponent::NetMulticast_InvokeGameplayCueExecuted_WithParams_Implementation(const FGameplayTag GameplayCueTag, FPredictionKey PredictionKey, FGameplayCueParameters GameplayCueParameters)
+{
+	if (IsOwnerActorAuthoritative() || PredictionKey.IsValidKey() == false)
+	{
+		InvokeGameplayCueEvent(GameplayCueTag, EGameplayCueEvent::Executed, GameplayCueParameters);
+	}
+}
+
 void UAbilitySystemComponent::NetMulticast_InvokeGameplayCueAdded_Implementation(const FGameplayTag GameplayCueTag, FPredictionKey PredictionKey, FGameplayEffectContextHandle EffectContext)
 {
 	if (IsOwnerActorAuthoritative() || PredictionKey.IsValidKey() == false)
diff --git a/Engine/Source/Runtime/GameplayAbilities/Public/AbilitySystemComponent.h b/Engine/Source/Runtime/GameplayAbilities/Public/AbilitySystemComponent.h
--- a/Engine/Source/Runtime/GameplayAbilities/Public/AbilitySystemComponent.h
+++ b/Engine/Source/Runtime/GameplayAbilities/Public/AbilitySystemComponent.h
@@ -425,30 +425,37 @@ class GAMEPLAYABILITIES_API UAbilitySystemComponent : public UActorComponent, pu
 	UFUNCTION(NetMulticast, unreliable)
 	void NetMulticast_InvokeGameplayCueExecuted_FromSpec(const FGameplayEffectSpec Spec, FPredictionKey PredictionKey);
 
-	void InvokeGameplayCueEvent(const FGameplayEffectSpec &Spec, EGameplayCueEvent::Type EventType);
-
 	// GameplayCues can also come on their own. These take an optional effect context to pass through hit result, etc
-
 	void ExecuteGameplayCue(const FGameplayTag GameplayCueTag, FGameplayEffectContextHandle EffectContext = FGameplayEffectContextHandle());
 
+	// This version allows the caller to set an explicit FGameplayCueParmeters.
+	void ExecuteGameplayCue(const FGameplayTag GameplayCueTag, const FGameplayCueParameters& GameplayCueParameters);
+
 	void AddGameplayCue(const FGameplayTag GameplayCueTag, FGameplayEffectContextHandle EffectContext = FGameplayEffectContextHandle());
 	
 	void RemoveGameplayCue(const FGameplayTag GameplayCueTag);
 
 	/** Removes any GameplayCue added on its own, i.e. not as part of a GameplayEffect. */
 	void RemoveAllGameplayCues();
-
+	
 	UFUNCTION(NetMulticast, unreliable)
 	void NetMulticast_InvokeGameplayCueExecuted(const FGameplayTag GameplayCueTag, FPredictionKey PredictionKey, FGameplayEffectContextHandle EffectContext);
 
+	UFUNCTION(NetMulticast, unreliable)
+	void NetMulticast_InvokeGameplayCueExecuted_WithParams(const FGameplayTag GameplayCueTag, FPredictionKey PredictionKey, FGameplayCueParameters GameplayCueParameters);
+
 	UFUNCTION(NetMulticast, unreliable)
 	void NetMulticast_InvokeGameplayCueAdded(const FGameplayTag GameplayCueTag, FPredictionKey PredictionKey, FGameplayEffectContextHandle EffectContext);
 
 	UFUNCTION(NetMulticast, unreliable)
 	void NetMulticast_InvokeGameplayCueRemoved(const FGameplayTag GameplayCueTag, FPredictionKey PredictionKey);
 
+	void InvokeGameplayCueEvent(const FGameplayEffectSpec &Spec, EGameplayCueEvent::Type EventType);
+
 	void InvokeGameplayCueEvent(const FGameplayTag GameplayCueTag, EGameplayCueEvent::Type EventType, FGameplayEffectContextHandle EffectContext = FGameplayEffectContextHandle());
 
+	void InvokeGameplayCueEvent(const FGameplayTag GameplayCueTag, EGameplayCueEvent::Type EventType, const FGameplayCueParameters& GameplayCueParameters);
+
 	/** Allows polling to see if a GameplayCue is active. We expect most GameplayCue handling to be event based, but some cases we may need to check if a GamepalyCue is active (Animation Blueprint for example) */
 	UFUNCTION(BlueprintCallable, Category="GameplayCue", meta=(GameplayTagFilter="GameplayCue"))
 	bool IsGameplayCueActive(const FGameplayTag GameplayCueTag) const;
