move all ability task delegate declarations outside of uclass. Fixes BP compile problem (introduced by recent uclass changes?)

diff --git a/Engine/Source/Runtime/GameplayAbilities/Private/GameplayCueManager.cpp b/Engine/Source/Runtime/GameplayAbilities/Private/GameplayCueManager.cpp
--- a/Engine/Source/Runtime/GameplayAbilities/Private/GameplayCueManager.cpp
+++ b/Engine/Source/Runtime/GameplayAbilities/Private/GameplayCueManager.cpp
@@ -345,7 +345,7 @@ void UGameplayCueManager::BuildAccelerationMap_Internal()
 		}
 	}
 
-	PrintGameplayCueNotifyMap();
+	// PrintGameplayCueNotifyMap();
 }
 
 int32 UGameplayCueManager::FinishLoadingGameplayCueNotifies()
diff --git a/Engine/Source/Runtime/GameplayAbilities/Public/Abilities/Tasks/AbilityTask_WaitInputPress.h b/Engine/Source/Runtime/GameplayAbilities/Public/Abilities/Tasks/AbilityTask_WaitInputPress.h
--- a/Engine/Source/Runtime/GameplayAbilities/Public/Abilities/Tasks/AbilityTask_WaitInputPress.h
+++ b/Engine/Source/Runtime/GameplayAbilities/Public/Abilities/Tasks/AbilityTask_WaitInputPress.h
@@ -5,17 +5,17 @@
 #include "AbilitySystemComponent.h"
 #include "AbilityTask_WaitInputPress.generated.h"
 
+DECLARE_DYNAMIC_MULTICAST_DELEGATE_OneParam(FInputPressDelegate, float, TimeWaited);
+
 /**
-*	Waits until the input is pressed from activating an ability. This should be true immediately upon starting the ability, since the key was pressed to activate it.
-*	We expect server to execute this task in parrallel and keep its own time. We do not keep track of 
-*/
+ *	Waits until the input is pressed from activating an ability. This should be true immediately upon starting the ability, since the key was pressed to activate it.
+ *	We expect server to execute this task in parrallel and keep its own time. We do not keep track of 
+ */
 UCLASS(MinimalAPI)
 class UAbilityTask_WaitInputPress : public UAbilityTask
 {
 	GENERATED_UCLASS_BODY()
 
-	DECLARE_DYNAMIC_MULTICAST_DELEGATE_OneParam(FInputPressDelegate, float, TimeWaited);
-
 	UPROPERTY(BlueprintAssignable)
 	FInputPressDelegate		OnPress;
 
