FScopedAggregatorOnDirtyBatch fix in cases where the onDirty calls triggers another GE add which causes another scope lock and modifies the OnDirty list as we are processing it.

diff --git a/Engine/Source/Runtime/GameplayAbilities/Private/GameplayEffectAggregator.cpp b/Engine/Source/Runtime/GameplayAbilities/Private/GameplayEffectAggregator.cpp
--- a/Engine/Source/Runtime/GameplayAbilities/Private/GameplayEffectAggregator.cpp
+++ b/Engine/Source/Runtime/GameplayAbilities/Private/GameplayEffectAggregator.cpp
@@ -324,11 +324,12 @@ void FScopedAggregatorOnDirtyBatch::EndLock()
 	GlobalBatchCount--;
 	if (GlobalBatchCount == 0)
 	{
-		for (FAggregator* Agg : DirtyAggregators)
+		TSet<FAggregator*> LocalSet(MoveTemp(DirtyAggregators));
+		for (FAggregator* Agg : LocalSet)
 		{
 			Agg->BroadcastOnDirty();
 		}
-		DirtyAggregators.Empty();
+		LocalSet.Empty();
 	}
 }
 
@@ -349,4 +350,4 @@ void FScopedAggregatorOnDirtyBatch::EndNetReceiveLock()
 		EndLock();
 		GlobalFromNetworkUpdate = false;
 	}
-}
\ No newline at end of file
+}
