Add base value to ability system debug hud

diff --git a/Engine/Source/Runtime/GameplayAbilities/Private/AbilitySystemComponent.cpp b/Engine/Source/Runtime/GameplayAbilities/Private/AbilitySystemComponent.cpp
--- a/Engine/Source/Runtime/GameplayAbilities/Private/AbilitySystemComponent.cpp
+++ b/Engine/Source/Runtime/GameplayAbilities/Private/AbilitySystemComponent.cpp
@@ -1160,8 +1160,17 @@ void UAbilitySystemComponent::DisplayDebug(class UCanvas* Canvas, const class FD
 			{
 				FAggregator& Aggregator = *AggregatorRef.Get();
 
+				float FinalValue = GetNumericAttribute(Attribute);
+				float BaseValue = Aggregator.GetBaseValue();
+
+				FString AttributeString = FString::Printf(TEXT("%s %.2f "), *Attribute.GetName(), GetNumericAttribute(Attribute));
+				if (FMath::Abs<float>(BaseValue - FinalValue) > SMALL_NUMBER)
+				{
+					AttributeString += FString::Printf(TEXT(" (Base: %.2f)"), BaseValue);
+				}
+
 				Canvas->SetDrawColor(FColor::White);
-				YPos += Canvas->DrawText(GEngine->GetTinyFont(), FString::Printf(TEXT("%s %.2f"), *Attribute.GetName(), GetNumericAttribute(Attribute) ), 4.f, YPos);
+				YPos += Canvas->DrawText(GEngine->GetTinyFont(), AttributeString, 4.f, YPos);
 
 				DrawAttributes.Add(Attribute);
 
