added PreAttributeBaseChange to attribute set. Allows game code to clamp base values of attributes.

diff --git a/Engine/Source/Runtime/GameplayAbilities/Private/GameplayEffect.cpp b/Engine/Source/Runtime/GameplayAbilities/Private/GameplayEffect.cpp
--- a/Engine/Source/Runtime/GameplayAbilities/Private/GameplayEffect.cpp
+++ b/Engine/Source/Runtime/GameplayAbilities/Private/GameplayEffect.cpp
@@ -1564,6 +1564,11 @@ void FActiveGameplayEffectsContainer::SetAttributeBaseValue(FGameplayAttribute A
 	{
 		// There is an aggregator for this attribute, so set the base value. The dirty callback chain
 		// will update the actual AttributeSet property value for us.
+		
+		const UAttributeSet* Set = Owner->GetAttributeSubobject( Attribute.GetAttributeSetClass() );
+		check(Set);
+
+		Set->PreAttributeBaseChange(Attribute, NewBaseValue);
 		RefPtr->Get()->SetBaseValue(NewBaseValue);
 	}
 	else
@@ -1640,7 +1645,18 @@ void FActiveGameplayEffectsContainer::ApplyModToAttribute(const FGameplayAttribu
 	if (RefPtr)
 	{
 		ABILITY_LOG(Log, TEXT("Property %s has active mods. Adding to Aggregator."), *Attribute.GetName());
-		RefPtr->Get()->ExecModOnBaseValue(ModifierOp, ModifierMagnitude);
+		FAggregator* Agg = RefPtr->Get();
+		
+		// We must give the attribute set a change to clamp to base value. Otherwise we may haver aggregator base values get 
+		// way out of sync with the final uproperty value.
+		float NewBase = Agg->GetBaseValue();
+		NewBase = FAggregator::StaticExecModOnBaseValue(NewBase, ModifierOp, ModifierMagnitude);
+		
+		const UAttributeSet* Set = Owner->GetAttributeSubobject( Attribute.GetAttributeSetClass() );
+		check(Set);
+
+		Set->PreAttributeBaseChange(Attribute, NewBase);
+		Agg->SetBaseValue(NewBase);
 	}
 	else
 	{
diff --git a/Engine/Source/Runtime/GameplayAbilities/Public/AttributeSet.h b/Engine/Source/Runtime/GameplayAbilities/Public/AttributeSet.h
--- a/Engine/Source/Runtime/GameplayAbilities/Public/AttributeSet.h
+++ b/Engine/Source/Runtime/GameplayAbilities/Public/AttributeSet.h
@@ -113,6 +113,14 @@ public:
 	 */
 	virtual void PreAttributeChange(const FGameplayAttribute& Attribute, float& NewValue) { }
 
+	/**
+	 *	This is called just before any modification happens to an attribute's base value when an attribute aggregator exists.
+	 *	This function should enforce clamping (presuming you wish to clamp the base value along with the final value in PreAttributeChange)
+	 *	This function should NOT invoke gameplay related events or callbacks. Do those in PreAttributeChange() which will be called prior to the
+	 *	final value of the attribute actually changing.
+	 */
+	virtual void PreAttributeBaseChange(const FGameplayAttribute& Attribute, float& NewValue) const { }
+
 	/** This signifies the attribute set can be ID'd by name over the network. */
 	void SetNetAddressable();
 
