Remove ensure on (net) Role in TryActivateAbility. There is a small window where a client can be killed and unpossessed, but not yet cleared from the ASC's Avatar/ActorInfo.

diff --git a/Engine/Source/Runtime/GameplayAbilities/Private/Abilities/GameplayAbility.cpp b/Engine/Source/Runtime/GameplayAbilities/Private/Abilities/GameplayAbility.cpp
--- a/Engine/Source/Runtime/GameplayAbilities/Private/Abilities/GameplayAbility.cpp
+++ b/Engine/Source/Runtime/GameplayAbilities/Private/Abilities/GameplayAbility.cpp
@@ -142,6 +142,12 @@ bool UGameplayAbility::CanActivateAbility(const FGameplayAbilitySpecHandle Handl
 {
 	// Don't set the actor info, CanActivate is called on the CDO
 
+	// A valid AvatarActor is required. Simulated proxy check means only authority or autonomous proxies should be executing abilities.
+	if (ActorInfo == nullptr || ActorInfo->AvatarActor == nullptr || ActorInfo->AvatarActor->Role == ROLE_SimulatedProxy)
+	{
+		return false;
+	}
+
 	//make into a reference for simplicity
 	FGameplayTagContainer DummyContainer;
 	FGameplayTagContainer& OutTags = OptionalRelevantTags ? *OptionalRelevantTags : DummyContainer;
diff --git a/Engine/Source/Runtime/GameplayAbilities/Private/AbilitySystemComponent_Abilities.cpp b/Engine/Source/Runtime/GameplayAbilities/Private/AbilitySystemComponent_Abilities.cpp
--- a/Engine/Source/Runtime/GameplayAbilities/Private/AbilitySystemComponent_Abilities.cpp
+++ b/Engine/Source/Runtime/GameplayAbilities/Private/AbilitySystemComponent_Abilities.cpp
@@ -694,7 +694,10 @@ bool UAbilitySystemComponent::TryActivateAbility(FGameplayAbilitySpecHandle Hand
 
 	// This should only come from button presses/local instigation (AI, etc)
 	ENetRole NetMode = ActorInfo->AvatarActor->Role;
-	ensure(NetMode != ROLE_SimulatedProxy);
+	if (NetMode == ROLE_SimulatedProxy)
+	{
+		return false;
+	}
 
 	bool bIsLocal = AbilityActorInfo->IsLocallyControlled();
 
