Merging CL# 2414775 from UE4-Orion

- MontageStop function now has optional parameter OverrideBlendOutTime (by default it will use the BlendOutTime on the AnimMontage).
- BlendTime on the AnimMontage instance is now replicated so abilities that override the blend time of montages will behave correctly.

diff --git a/Engine/Source/Runtime/Engine/Classes/Animation/AnimInstance.h b/Engine/Source/Runtime/Engine/Classes/Animation/AnimInstance.h
--- a/Engine/Source/Runtime/Engine/Classes/Animation/AnimInstance.h
+++ b/Engine/Source/Runtime/Engine/Classes/Animation/AnimInstance.h
@@ -449,6 +449,10 @@ public:
 	/** return true if Montage is not currently active. (not valid or blending out) */
 	bool Montage_GetIsStopped(UAnimMontage* Montage);
 
+	/** Get the current blend time of the Montage.
+	If Montage reference is NULL, it will return the current blend time on the first active Montage found. */
+	float Montage_GetBlendTime(UAnimMontage* Montage);
+
 	/** Get PlayRate for Montage.
 	If Montage reference is NULL, PlayRate for any Active Montage will be returned.
 	If Montage is not playing, 0 is returned. */
diff --git a/Engine/Source/Runtime/Engine/Private/Animation/AnimInstance.cpp b/Engine/Source/Runtime/Engine/Private/Animation/AnimInstance.cpp
--- a/Engine/Source/Runtime/Engine/Private/Animation/AnimInstance.cpp
+++ b/Engine/Source/Runtime/Engine/Private/Animation/AnimInstance.cpp
@@ -2104,6 +2104,32 @@ bool UAnimInstance::Montage_GetIsStopped(UAnimMontage* Montage)
 	return true;
 }
 
+float UAnimInstance::Montage_GetBlendTime(UAnimMontage* Montage)
+{
+	if (Montage)
+	{
+		FAnimMontageInstance * MontageInstance = GetActiveInstanceForMontage(*Montage);
+		if (MontageInstance)
+		{
+			return MontageInstance->BlendTime;
+		}
+	}
+	else
+	{
+		// If no Montage reference, use first active one found.
+		for (int32 InstanceIndex = 0; InstanceIndex < MontageInstances.Num(); InstanceIndex++)
+		{
+			FAnimMontageInstance * MontageInstance = MontageInstances[InstanceIndex];
+			if (MontageInstance && MontageInstance->IsActive())
+			{
+				return MontageInstance->BlendTime;
+			}
+		}
+	}
+
+	return 0.f;
+}
+
 float UAnimInstance::Montage_GetPlayRate(UAnimMontage* Montage)
 {
 	if (Montage)
diff --git a/Engine/Source/Runtime/GameplayAbilities/Private/Abilities/GameplayAbility.cpp b/Engine/Source/Runtime/GameplayAbilities/Private/Abilities/GameplayAbility.cpp
--- a/Engine/Source/Runtime/GameplayAbilities/Private/Abilities/GameplayAbility.cpp
+++ b/Engine/Source/Runtime/GameplayAbilities/Private/Abilities/GameplayAbility.cpp
@@ -738,7 +738,7 @@ void UGameplayAbility::MontageSetNextSectionName(FName FromSectionName, FName To
 	}
 }
 
-void UGameplayAbility::MontageStop()
+void UGameplayAbility::MontageStop(float OverrideBlendOutTime)
 {
 	check(CurrentActorInfo);
 
@@ -748,7 +748,7 @@ void UGameplayAbility::MontageStop()
 		// We should only stop the current montage if we are the animating ability
 		if (AbilitySystemComponent->IsAnimatingAbility(this))
 		{
-			AbilitySystemComponent->CurrentMontageStop();
+			AbilitySystemComponent->CurrentMontageStop(OverrideBlendOutTime);
 		}
 	}
 }
diff --git a/Engine/Source/Runtime/GameplayAbilities/Private/AbilitySystemComponent_Abilities.cpp b/Engine/Source/Runtime/GameplayAbilities/Private/AbilitySystemComponent_Abilities.cpp
--- a/Engine/Source/Runtime/GameplayAbilities/Private/AbilitySystemComponent_Abilities.cpp
+++ b/Engine/Source/Runtime/GameplayAbilities/Private/AbilitySystemComponent_Abilities.cpp
@@ -1862,6 +1862,7 @@ void UAbilitySystemComponent::AnimMontage_UpdateReplicatedData()
 		RepAnimMontageInfo.AnimMontage = LocalAnimMontageInfo.AnimMontage;
 		RepAnimMontageInfo.PlayRate = AnimInstance->Montage_GetPlayRate(LocalAnimMontageInfo.AnimMontage);
 		RepAnimMontageInfo.Position = AnimInstance->Montage_GetPosition(LocalAnimMontageInfo.AnimMontage);
+		RepAnimMontageInfo.BlendTime = AnimInstance->Montage_GetBlendTime(LocalAnimMontageInfo.AnimMontage);
 
 		// Compressed Flags
 		bool bIsStopped = AnimInstance->Montage_GetIsStopped(LocalAnimMontageInfo.AnimMontage);
@@ -1918,10 +1919,11 @@ void UAbilitySystemComponent::OnRep_ReplicatedAnimMontage()
 		if (DebugMontage)
 		{
 			ABILITY_LOG( Warning, TEXT("\n\nOnRep_ReplicatedAnimMontage, %s"), *GetNameSafe(this));
-			ABILITY_LOG( Warning, TEXT("\tAnimMontage: %s\n\tPlayRate: %f\n\tPosition: %f\n\tNextSectionID: %d\n\tIsStopped: %d\n\tForcePlayBit: %d"),
+			ABILITY_LOG( Warning, TEXT("\tAnimMontage: %s\n\tPlayRate: %f\n\tPosition: %f\n\tBlendTime: %f\n\tNextSectionID: %d\n\tIsStopped: %d\n\tForcePlayBit: %d"),
 				*GetNameSafe(RepAnimMontageInfo.AnimMontage), 
 				RepAnimMontageInfo.PlayRate, 
-				RepAnimMontageInfo.Position, 
+				RepAnimMontageInfo.Position,
+				RepAnimMontageInfo.BlendTime,
 				RepAnimMontageInfo.NextSectionID, 
 				RepAnimMontageInfo.IsStopped, 
 				RepAnimMontageInfo.ForcePlayBit);
@@ -1982,13 +1984,13 @@ void UAbilitySystemComponent::OnRep_ReplicatedAnimMontage()
 			bool ReplicatedIsStopped = bool(RepAnimMontageInfo.IsStopped);
 			if( ReplicatedIsStopped && !bIsStopped )
 			{
-				CurrentMontageStop();
+				CurrentMontageStop(RepAnimMontageInfo.BlendTime);
 			}
 		}
 	}
 }
 
-void UAbilitySystemComponent::CurrentMontageStop()
+void UAbilitySystemComponent::CurrentMontageStop(float OverrideBlendOutTime)
 {
 	UAnimInstance* AnimInstance = AbilityActorInfo->AnimInstance.Get();
 	UAnimMontage* MontageToStop = LocalAnimMontageInfo.AnimMontage;
@@ -1996,7 +1998,9 @@ void UAbilitySystemComponent::CurrentMontageStop()
 
 	if (bShouldStopMontage)
 	{
-		AnimInstance->Montage_Stop(MontageToStop->BlendOutTime);
+		const float BlendOutTime = (OverrideBlendOutTime >= 0.0f ? OverrideBlendOutTime : MontageToStop->BlendOutTime);
+
+		AnimInstance->Montage_Stop(BlendOutTime);
 
 		if (IsOwnerActorAuthoritative())
 		{
diff --git a/Engine/Source/Runtime/GameplayAbilities/Public/Abilities/GameplayAbility.h b/Engine/Source/Runtime/GameplayAbilities/Public/Abilities/GameplayAbility.h
--- a/Engine/Source/Runtime/GameplayAbilities/Public/Abilities/GameplayAbility.h
+++ b/Engine/Source/Runtime/GameplayAbilities/Public/Abilities/GameplayAbility.h
@@ -659,8 +659,13 @@ protected:
 	UFUNCTION(BlueprintCallable, Category = "Ability|Animation")
 	void MontageSetNextSectionName(FName FromSectionName, FName ToSectionName);
 
-	UFUNCTION(BlueprintCallable, Category="Ability|Animation")
-	void MontageStop();
+	/**
+	 * Stops the current animation montage.
+	 *
+	 * @param OverrideBlendTime If < 0, will override the BlendOutTime parameter on the AnimMontage instance
+	 */
+	UFUNCTION(BlueprintCallable, Category="Ability|Animation", Meta = (AdvancedDisplay = "OverrideBlendOutTime"))
+	void MontageStop(float OverrideBlendOutTime = -1.0f);
 
 	// ----------------------------------------------------------------------------------------------------------------
 	//
diff --git a/Engine/Source/Runtime/GameplayAbilities/Public/Abilities/GameplayAbilityTypes.h b/Engine/Source/Runtime/GameplayAbilities/Public/Abilities/GameplayAbilityTypes.h
--- a/Engine/Source/Runtime/GameplayAbilities/Public/Abilities/GameplayAbilityTypes.h
+++ b/Engine/Source/Runtime/GameplayAbilities/Public/Abilities/GameplayAbilityTypes.h
@@ -411,6 +411,10 @@ struct GAMEPLAYABILITIES_API FGameplayAbilityRepAnimMontage
 	UPROPERTY()
 	float Position;
 
+	/** Montage current blend time */
+	UPROPERTY()
+	float BlendTime;
+
 	/** NextSectionID */
 	UPROPERTY()
 	uint8 NextSectionID;
@@ -427,6 +431,7 @@ struct GAMEPLAYABILITIES_API FGameplayAbilityRepAnimMontage
 		: AnimMontage(NULL),
 		PlayRate(0.f),
 		Position(0.f),
+		BlendTime(0.f),
 		NextSectionID(0),
 		IsStopped(true),
 		ForcePlayBit(0)
diff --git a/Engine/Source/Runtime/GameplayAbilities/Public/AbilitySystemComponent.h b/Engine/Source/Runtime/GameplayAbilities/Public/AbilitySystemComponent.h
--- a/Engine/Source/Runtime/GameplayAbilities/Public/AbilitySystemComponent.h
+++ b/Engine/Source/Runtime/GameplayAbilities/Public/AbilitySystemComponent.h
@@ -697,7 +697,7 @@ public:
 	float PlayMontageSimulated(UAnimMontage* Montage, float InPlayRate, FName StartSectionName = NAME_None);
 
 	/** Stops whatever montage is currently playing. Expectation is caller should only be stoping it if they are the current animating ability (or have good reason not to check) */
-	void CurrentMontageStop();
+	void CurrentMontageStop(float OverrideBlendOutTime = -1.0f);
 
 	/** Clear the animating ability that is passed in, if it's still currently animating */
 	void ClearAnimatingAbility(UGameplayAbility* Ability);
