Integration(2416174)

Fixed UAbilityTask_WaitTargetData assert caused by accessing the ability reference after the task has ended.

diff --git a/Engine/Source/Runtime/GameplayAbilities/Private/Abilities/Tasks/AbilityTask_WaitTargetData.cpp b/Engine/Source/Runtime/GameplayAbilities/Private/Abilities/Tasks/AbilityTask_WaitTargetData.cpp
--- a/Engine/Source/Runtime/GameplayAbilities/Private/Abilities/Tasks/AbilityTask_WaitTargetData.cpp
+++ b/Engine/Source/Runtime/GameplayAbilities/Private/Abilities/Tasks/AbilityTask_WaitTargetData.cpp
@@ -40,6 +40,11 @@ void UAbilityTask_WaitTargetData::Activate()
 
 			TargetClass = SpawnedActor->GetClass();
 
+			if (IsPendingKill())
+			{
+				return;
+			}
+
 			if (ShouldSpawnTargetActor())
 			{
 				InitializeTargetActor(SpawnedActor);
