Integration (2416175)

Hold input abilities;

* Added "Activate On Input Held" option to UGameplayAbility.  This flags the ability to activate when the bound input key is held down.  For this to work, you'll need to call UAbilitySystemComponent::TryActivateInputHeldAbilities() somewhere in your game.

diff --git a/Engine/Source/Runtime/GameplayAbilities/Private/AbilitySystemComponent_Abilities.cpp b/Engine/Source/Runtime/GameplayAbilities/Private/AbilitySystemComponent_Abilities.cpp
--- a/Engine/Source/Runtime/GameplayAbilities/Private/AbilitySystemComponent_Abilities.cpp
+++ b/Engine/Source/Runtime/GameplayAbilities/Private/AbilitySystemComponent_Abilities.cpp
@@ -1553,6 +1553,22 @@ void UAbilitySystemComponent::AbilitySpecInputReleased(FGameplayAbilitySpec& Spe
 	}
 }
 
+void UAbilitySystemComponent::TryActivateInputHeldAbilities()
+{
+	for (FGameplayAbilitySpec& Spec : ActivatableAbilities.Items)
+	{
+		if ((Spec.InputID == INDEX_NONE) || !Spec.InputPressed || Spec.IsActive() || !Spec.Ability)
+		{
+			continue;
+		}
+
+		if (Spec.Ability->bActivateOnInputHeld)
+		{
+			TryActivateAbility(Spec.Handle);
+		}
+	}
+}
+
 void UAbilitySystemComponent::InputConfirm()
 {
 	FScopedPredictionWindow ScopedPrediction(this, true);
diff --git a/Engine/Source/Runtime/GameplayAbilities/Public/Abilities/GameplayAbility.h b/Engine/Source/Runtime/GameplayAbilities/Public/Abilities/GameplayAbility.h
--- a/Engine/Source/Runtime/GameplayAbilities/Public/Abilities/GameplayAbility.h
+++ b/Engine/Source/Runtime/GameplayAbilities/Public/Abilities/GameplayAbility.h
@@ -127,6 +127,10 @@ public:
 	/** Input binding stub. */
 	virtual void InputReleased(const FGameplayAbilitySpecHandle Handle, const FGameplayAbilityActorInfo* ActorInfo, const FGameplayAbilityActivationInfo ActivationInfo) {};
 
+	/** If this is set, the ability will try to activate if the input is held. */
+	UPROPERTY(EditDefaultsOnly, Category = Input)
+	bool bActivateOnInputHeld;
+
 	/** Returns true if this ability can be activated right now. Has no side effects */
 	virtual bool CanActivateAbility(const FGameplayAbilitySpecHandle Handle, const FGameplayAbilityActorInfo* ActorInfo, OUT FGameplayTagContainer* OptionalRelevantTags = nullptr) const;
 
diff --git a/Engine/Source/Runtime/GameplayAbilities/Public/AbilitySystemComponent.h b/Engine/Source/Runtime/GameplayAbilities/Public/AbilitySystemComponent.h
--- a/Engine/Source/Runtime/GameplayAbilities/Public/AbilitySystemComponent.h
+++ b/Engine/Source/Runtime/GameplayAbilities/Public/AbilitySystemComponent.h
@@ -635,6 +635,8 @@ public:
 
 	void AbilitySpecInputReleased(FGameplayAbilitySpec& Spec);
 
+	void TryActivateInputHeldAbilities();
+
 	/** Sent by abilities to tell server when activation input is pressed. (Sent by default in order not to short-circuit WaitInputPress/WaitInputRelease tasks) */
 	UFUNCTION(Server, reliable, WithValidation)
 	void ServerInputPress(FGameplayAbilitySpecHandle AbilityToActivate, FPredictionKey ScopedPedictionKey);
