Integration(2424977): Held ability activation fix;

* Checking for input held abilities will now only happen if the owning player controller has a valid player input.  This fixes cases of the ability activating twice on the server.

=== Bugs Fixed ===
https://jira.ol.epicgames.net/browse/OR-879

diff --git a/Engine/Source/Runtime/GameplayAbilities/Private/AbilitySystemComponent_Abilities.cpp b/Engine/Source/Runtime/GameplayAbilities/Private/AbilitySystemComponent_Abilities.cpp
--- a/Engine/Source/Runtime/GameplayAbilities/Private/AbilitySystemComponent_Abilities.cpp
+++ b/Engine/Source/Runtime/GameplayAbilities/Private/AbilitySystemComponent_Abilities.cpp
@@ -1555,16 +1555,21 @@ void UAbilitySystemComponent::AbilitySpecInputReleased(FGameplayAbilitySpec& Spe
 
 void UAbilitySystemComponent::TryActivateInputHeldAbilities()
 {
-	for (FGameplayAbilitySpec& Spec : ActivatableAbilities.Items)
+	APlayerController* PC = AbilityActorInfo.IsValid() ? AbilityActorInfo->PlayerController.Get() : nullptr;
+
+	if (PC && PC->PlayerInput)
 	{
-		if ((Spec.InputID == INDEX_NONE) || !Spec.InputPressed || Spec.IsActive() || !Spec.Ability)
+		for (FGameplayAbilitySpec& Spec : ActivatableAbilities.Items)
 		{
-			continue;
-		}
+			if ((Spec.InputID == INDEX_NONE) || !Spec.InputPressed || Spec.IsActive() || !Spec.Ability)
+			{
+				continue;
+			}
 
-		if (Spec.Ability->bActivateOnInputHeld)
-		{
-			TryActivateAbility(Spec.Handle);
+			if (Spec.Ability->bActivateOnInputHeld)
+			{
+				TryActivateAbility(Spec.Handle);
+			}
 		}
 	}
 }
