Merge AbilitySystem changes from //depot/UE4-Orion

* Support for pending linked list for new gameplay effects added during scoped lock
* Added FActiveGameplayEffectIterator

diff --git a/Engine/Source/Runtime/GameplayAbilities/Private/Abilities/Tasks/AbilityTask_WaitTargetData.cpp b/Engine/Source/Runtime/GameplayAbilities/Private/Abilities/Tasks/AbilityTask_WaitTargetData.cpp
--- a/Engine/Source/Runtime/GameplayAbilities/Private/Abilities/Tasks/AbilityTask_WaitTargetData.cpp
+++ b/Engine/Source/Runtime/GameplayAbilities/Private/Abilities/Tasks/AbilityTask_WaitTargetData.cpp
@@ -37,9 +37,11 @@ void UAbilityTask_WaitTargetData::Activate()
 		if (TargetActor.IsValid())
 		{
 			AGameplayAbilityTargetActor* SpawnedActor = TargetActor.Get();
-
 			TargetClass = SpawnedActor->GetClass();
 
+			RegisterTargetDataCallbacks();
+
+
 			if (IsPendingKill())
 			{
 				return;
@@ -49,6 +51,8 @@ void UAbilityTask_WaitTargetData::Activate()
 			{
 				InitializeTargetActor(SpawnedActor);
 				FinalizeTargetActor(SpawnedActor);
+
+				// Note that the call to FinalizeTargetActor, this task could finish and our owning ability may be ended.
 			}
 			else
 			{
@@ -58,8 +62,6 @@ void UAbilityTask_WaitTargetData::Activate()
 				SpawnedActor->Destroy();
 				SpawnedActor = nullptr;
 			}
-
-			RegisterTargetDataCallbacks();
 		}
 		else
 		{
@@ -184,6 +186,11 @@ void UAbilityTask_WaitTargetData::FinalizeTargetActor(AGameplayAbilityTargetActo
 
 void UAbilityTask_WaitTargetData::RegisterTargetDataCallbacks()
 {
+	if (!ensure(IsPendingKill() == false))
+	{
+		return;
+	}
+
 	check(TargetClass);
 	check(Ability.IsValid());
 
diff --git a/Engine/Source/Runtime/GameplayAbilities/Private/AbilitySystemComponent.cpp b/Engine/Source/Runtime/GameplayAbilities/Private/AbilitySystemComponent.cpp
--- a/Engine/Source/Runtime/GameplayAbilities/Private/AbilitySystemComponent.cpp
+++ b/Engine/Source/Runtime/GameplayAbilities/Private/AbilitySystemComponent.cpp
@@ -325,11 +325,8 @@ int32 UAbilitySystemComponent::GetNumActiveGameplayEffects() const
 }
 
 void UAbilitySystemComponent::GetAllActiveGameplayEffectSpecs(TArray<FGameplayEffectSpec>& OutSpecCopies)
-{
-	for (const auto& ActiveEffect : ActiveGameplayEffects.GameplayEffects)
-	{
-		OutSpecCopies.Add(ActiveEffect.Spec);
-	}
+{	
+	ActiveGameplayEffects.GetAllActiveGameplayEffectSpecs(OutSpecCopies);
 }
 
 const FGameplayTagContainer* UAbilitySystemComponent::GetGameplayEffectSourceTagsFromHandle(FActiveGameplayEffectHandle Handle) const
@@ -1208,8 +1205,9 @@ void UAbilitySystemComponent::DisplayDebug(class UCanvas* Canvas, const class FD
 
 	if (bShowGameplayEffects)
 	{
-		for (FActiveGameplayEffect& ActiveGE : ActiveGameplayEffects.GameplayEffects)
+		for (FActiveGameplayEffect& ActiveGE : &ActiveGameplayEffects)
 		{
+			
 			Canvas->SetDrawColor(FColor::White);
 
 			FString DurationStr = TEXT("Infinite Duration ");
diff --git a/Engine/Source/Runtime/GameplayAbilities/Private/AbilitySystemDebugHUD.cpp b/Engine/Source/Runtime/GameplayAbilities/Private/AbilitySystemDebugHUD.cpp
--- a/Engine/Source/Runtime/GameplayAbilities/Private/AbilitySystemDebugHUD.cpp
+++ b/Engine/Source/Runtime/GameplayAbilities/Private/AbilitySystemDebugHUD.cpp
@@ -113,7 +113,7 @@ void AAbilitySystemDebugHUD::DrawDebugAbilitySystemComponent(UAbilitySystemCompo
 
 		Y+= 25;
 		// Draw Active GameplayEffect
-		for (FActiveGameplayEffect &Effect : Component->ActiveGameplayEffects.GameplayEffects)
+		for (FActiveGameplayEffect& Effect : &Component->ActiveGameplayEffects)
 		{
 			String = FString::Printf(TEXT("%s. [%d, %d] %.2f"), *Effect.Spec.ToSimpleString(), Effect.PredictionKey.Current, Effect.PredictionKey.Base, Effect.GetTimeRemaining(GameWorldTime));
 			DrawWithBackground(Font, String, Color, EAlignHorizontal::Left, X, EAlignVertical::Top, Y);	
diff --git a/Engine/Source/Runtime/GameplayAbilities/Private/AbilitySystemGlobals.cpp b/Engine/Source/Runtime/GameplayAbilities/Private/AbilitySystemGlobals.cpp
--- a/Engine/Source/Runtime/GameplayAbilities/Private/AbilitySystemGlobals.cpp
+++ b/Engine/Source/Runtime/GameplayAbilities/Private/AbilitySystemGlobals.cpp
@@ -125,14 +125,14 @@ UAbilitySystemGlobals& UAbilitySystemGlobals::Get()
 }
 
 /** Helping function to avoid having to manually cast */
-UAbilitySystemComponent* UAbilitySystemGlobals::GetAbilitySystemComponentFromActor(AActor* Actor, bool LookForComponent)
+UAbilitySystemComponent* UAbilitySystemGlobals::GetAbilitySystemComponentFromActor(const AActor* Actor, bool LookForComponent)
 {
 	if (Actor == nullptr)
 	{
 		return nullptr;
 	}
 
-	IAbilitySystemInterface* ASI = Cast<IAbilitySystemInterface>(Actor);
+	const IAbilitySystemInterface* ASI = Cast<IAbilitySystemInterface>(Actor);
 	if (ASI)
 	{
 		return ASI->GetAbilitySystemComponent();
diff --git a/Engine/Source/Runtime/GameplayAbilities/Private/GameplayCueNotify_Actor.cpp b/Engine/Source/Runtime/GameplayAbilities/Private/GameplayCueNotify_Actor.cpp
--- a/Engine/Source/Runtime/GameplayAbilities/Private/GameplayCueNotify_Actor.cpp
+++ b/Engine/Source/Runtime/GameplayAbilities/Private/GameplayCueNotify_Actor.cpp
@@ -11,6 +11,7 @@ AGameplayCueNotify_Actor::AGameplayCueNotify_Actor(const FObjectInitializer& Obj
 {
 	IsOverride = true;
 	PrimaryActorTick.bCanEverTick = true;
+	bAutoDestroyOnRemove = false;
 }
 
 #if WITH_EDITOR
@@ -81,6 +82,11 @@ void AGameplayCueNotify_Actor::HandleGameplayCue(AActor* MyTarget, EGameplayCueE
 
 		case EGameplayCueEvent::Removed:
 			OnRemove(MyTarget, Parameters);
+
+			if (bAutoDestroyOnRemove)
+			{
+				Destroy();
+			}
 			break;
 		};
 	}
diff --git a/Engine/Source/Runtime/GameplayAbilities/Private/GameplayEffect.cpp b/Engine/Source/Runtime/GameplayAbilities/Private/GameplayEffect.cpp
--- a/Engine/Source/Runtime/GameplayAbilities/Private/GameplayEffect.cpp
+++ b/Engine/Source/Runtime/GameplayAbilities/Private/GameplayEffect.cpp
@@ -382,6 +382,20 @@ void FGameplayEffectExecutionDefinition::GetAttributeCaptureDefinitions(OUT TArr
 //
 // --------------------------------------------------------------------------------------------------------------------------------------------------------
 
+FGameplayEffectSpec::FGameplayEffectSpec()
+	: Def(nullptr)
+	, Duration(UGameplayEffect::INSTANT_APPLICATION)
+	, Period(UGameplayEffect::NO_PERIOD)
+	, ChanceToApplyToTarget(1.f)
+	, StackCount(1)
+	, bCompletedSourceAttributeCapture(false)
+	, bCompletedTargetAttributeCapture(false)
+	, bDurationLocked(false)
+	, Level(UGameplayEffect::INVALID_LEVEL)
+{
+
+}
+
 FGameplayEffectSpec::FGameplayEffectSpec(const UGameplayEffect* InDef, const FGameplayEffectContextHandle& InEffectContext, float Level)
 	: Def(InDef)
 	, StackCount(1)
@@ -415,6 +429,80 @@ FGameplayEffectSpec::FGameplayEffectSpec(const UGameplayEffect* InDef, const FGa
 	CaptureDataFromSource();
 }
 
+FGameplayEffectSpec::FGameplayEffectSpec(const FGameplayEffectSpec& Other)
+{
+	*this = Other;
+}
+
+FGameplayEffectSpec::FGameplayEffectSpec(FGameplayEffectSpec&& Other)
+	: Def(Other.Def)
+	, ModifiedAttributes(MoveTemp(Other.ModifiedAttributes))
+	, CapturedRelevantAttributes(MoveTemp(Other.CapturedRelevantAttributes))
+	, TargetEffectSpecs(MoveTemp(Other.TargetEffectSpecs))
+	, Duration(Other.Duration)
+	, Period(Other.Period)
+	, ChanceToApplyToTarget(Other.ChanceToApplyToTarget)
+	, CapturedSourceTags(MoveTemp(Other.CapturedSourceTags))
+	, CapturedTargetTags(MoveTemp(Other.CapturedTargetTags))
+	, DynamicGrantedTags(MoveTemp(Other.DynamicGrantedTags))
+	, Modifiers(MoveTemp(Other.Modifiers))
+	, StackCount(Other.StackCount)
+	, bCompletedSourceAttributeCapture(Other.bCompletedSourceAttributeCapture)
+	, bCompletedTargetAttributeCapture(Other.bCompletedTargetAttributeCapture)
+	, bDurationLocked(Other.bDurationLocked)
+	, SetByCallerMagnitudes(MoveTemp(Other.SetByCallerMagnitudes))
+	, EffectContext(Other.EffectContext)
+	, Level(Other.Level)
+{
+
+}
+
+FGameplayEffectSpec& FGameplayEffectSpec::operator=(FGameplayEffectSpec&& Other)
+{
+	Def = Other.Def;
+	ModifiedAttributes = MoveTemp(Other.ModifiedAttributes);
+	CapturedRelevantAttributes = MoveTemp(Other.CapturedRelevantAttributes);
+	TargetEffectSpecs = MoveTemp(Other.TargetEffectSpecs);
+	Duration = Other.Duration;
+	Period = Other.Period;
+	ChanceToApplyToTarget = Other.ChanceToApplyToTarget;
+	CapturedSourceTags = MoveTemp(Other.CapturedSourceTags);
+	CapturedTargetTags = MoveTemp(Other.CapturedTargetTags);
+	DynamicGrantedTags = MoveTemp(Other.DynamicGrantedTags);
+	Modifiers = MoveTemp(Other.Modifiers);
+	StackCount = Other.StackCount;
+	bCompletedSourceAttributeCapture = Other.bCompletedSourceAttributeCapture;
+	bCompletedTargetAttributeCapture = Other.bCompletedTargetAttributeCapture;
+	bDurationLocked = Other.bDurationLocked;
+	SetByCallerMagnitudes = MoveTemp(Other.SetByCallerMagnitudes);
+	EffectContext = Other.EffectContext;
+	Level = Other.Level;
+	return *this;
+}
+
+FGameplayEffectSpec& FGameplayEffectSpec::operator=(const FGameplayEffectSpec& Other)
+{
+	Def = Other.Def;
+	ModifiedAttributes = Other.ModifiedAttributes;
+	CapturedRelevantAttributes = Other.CapturedRelevantAttributes;
+	TargetEffectSpecs = Other.TargetEffectSpecs;
+	Duration = Other.Duration;
+	Period = Other.Period;
+	ChanceToApplyToTarget = Other.ChanceToApplyToTarget;
+	CapturedSourceTags = Other.CapturedSourceTags;
+	CapturedTargetTags = Other.CapturedTargetTags;
+	DynamicGrantedTags = Other.DynamicGrantedTags;
+	Modifiers = Other.Modifiers;
+	StackCount = Other.StackCount;
+	bCompletedSourceAttributeCapture = Other.bCompletedSourceAttributeCapture;
+	bCompletedTargetAttributeCapture = Other.bCompletedTargetAttributeCapture;
+	bDurationLocked = Other.bDurationLocked;
+	SetByCallerMagnitudes = Other.SetByCallerMagnitudes;
+	EffectContext = Other.EffectContext;
+	Level = Other.Level;
+	return *this;
+}
+
 void FGameplayEffectSpec::SetupAttributeCaptureDefinitions()
 {
 	// Add duration if required
@@ -819,6 +907,43 @@ const FGameplayEffectAttributeCaptureDefinition& FGameplayEffectAttributeCapture
 	return BackingDefinition;
 }
 
+FGameplayEffectAttributeCaptureSpecContainer::FGameplayEffectAttributeCaptureSpecContainer()
+	: bHasNonSnapshottedAttributes(false)
+{
+
+}
+
+FGameplayEffectAttributeCaptureSpecContainer::FGameplayEffectAttributeCaptureSpecContainer(FGameplayEffectAttributeCaptureSpecContainer&& Other)
+	: SourceAttributes(MoveTemp(Other.SourceAttributes))
+	, TargetAttributes(MoveTemp(Other.TargetAttributes))
+	, bHasNonSnapshottedAttributes(Other.bHasNonSnapshottedAttributes)
+{
+}
+
+FGameplayEffectAttributeCaptureSpecContainer::FGameplayEffectAttributeCaptureSpecContainer(const FGameplayEffectAttributeCaptureSpecContainer& Other)
+	: SourceAttributes(Other.SourceAttributes)
+	, TargetAttributes(Other.TargetAttributes)
+	, bHasNonSnapshottedAttributes(Other.bHasNonSnapshottedAttributes)
+{
+}
+
+
+FGameplayEffectAttributeCaptureSpecContainer& FGameplayEffectAttributeCaptureSpecContainer::operator=(FGameplayEffectAttributeCaptureSpecContainer&& Other)
+{
+	SourceAttributes = MoveTemp(Other.SourceAttributes);
+	TargetAttributes = MoveTemp(Other.TargetAttributes);
+	bHasNonSnapshottedAttributes = Other.bHasNonSnapshottedAttributes;
+	return *this;
+}
+
+FGameplayEffectAttributeCaptureSpecContainer& FGameplayEffectAttributeCaptureSpecContainer::operator=(const FGameplayEffectAttributeCaptureSpecContainer& Other)
+{
+	SourceAttributes = Other.SourceAttributes;
+	TargetAttributes = Other.TargetAttributes;
+	bHasNonSnapshottedAttributes = Other.bHasNonSnapshottedAttributes;
+	return *this;
+}
+
 void FGameplayEffectAttributeCaptureSpecContainer::AddCaptureDefinition(const FGameplayEffectAttributeCaptureDefinition& InCaptureDefinition)
 {
 	const bool bSourceAttribute = (InCaptureDefinition.AttributeSource == EGameplayEffectAttributeCaptureSource::Source);
@@ -929,6 +1054,84 @@ void FGameplayEffectAttributeCaptureSpecContainer::UnregisterLinkedAggregatorCal
 //
 // --------------------------------------------------------------------------------------------------------------------------------------------------------
 
+FActiveGameplayEffect::FActiveGameplayEffect()
+	: StartGameStateTime(0)
+	, CachedStartGameStateTime(0)
+	, StartWorldTime(0.f)
+	, bIsInhibited(true)
+	, IsPendingRemove(false)
+	, PendingNext(nullptr)
+{
+}
+
+FActiveGameplayEffect::FActiveGameplayEffect(const FActiveGameplayEffect& Other)
+{
+	*this = Other;
+}
+
+FActiveGameplayEffect::FActiveGameplayEffect(FActiveGameplayEffectHandle InHandle, const FGameplayEffectSpec &InSpec, float CurrentWorldTime, int32 InStartGameStateTime, FPredictionKey InPredictionKey)
+	: Handle(InHandle)
+	, Spec(InSpec)
+	, PredictionKey(InPredictionKey)
+	, StartGameStateTime(InStartGameStateTime)
+	, CachedStartGameStateTime(InStartGameStateTime)
+	, StartWorldTime(CurrentWorldTime)
+	, bIsInhibited(true)
+	, IsPendingRemove(false)
+	, PendingNext(nullptr)
+{
+}
+
+FActiveGameplayEffect::FActiveGameplayEffect(FActiveGameplayEffect&& Other)
+	:Handle(Other.Handle)
+	,Spec(MoveTemp(Other.Spec))
+	,PredictionKey(Other.PredictionKey)
+	,StartGameStateTime(Other.StartGameStateTime)
+	,CachedStartGameStateTime(Other.CachedStartGameStateTime)
+	,StartWorldTime(Other.StartWorldTime)
+	,bIsInhibited(Other.bIsInhibited)
+	,IsPendingRemove(Other.IsPendingRemove)
+	,OnRemovedDelegate(Other.OnRemovedDelegate)
+	,PeriodHandle(Other.PeriodHandle)
+	,DurationHandle(Other.DurationHandle)
+{
+	// Note: purposefully not copying PendingNext pointer.
+}
+
+FActiveGameplayEffect& FActiveGameplayEffect::operator=(FActiveGameplayEffect&& Other)
+{
+	Handle = Other.Handle;
+	Spec = MoveTemp(Other.Spec);
+	PredictionKey = Other.PredictionKey;
+	StartGameStateTime = Other.StartGameStateTime;
+	CachedStartGameStateTime = Other.CachedStartGameStateTime;
+	StartWorldTime = Other.StartWorldTime;
+	bIsInhibited = Other.bIsInhibited;
+	IsPendingRemove = Other.IsPendingRemove;
+	OnRemovedDelegate = Other.OnRemovedDelegate;
+	PeriodHandle = Other.PeriodHandle;
+	DurationHandle = Other.DurationHandle;
+	// Note: purposefully not copying PendingNext pointer.
+	return *this;
+}
+
+FActiveGameplayEffect& FActiveGameplayEffect::operator=(const FActiveGameplayEffect& Other)
+{
+	Handle = Other.Handle;
+	Spec = Other.Spec;
+	PredictionKey = Other.PredictionKey;
+	StartGameStateTime = Other.StartGameStateTime;
+	CachedStartGameStateTime = Other.CachedStartGameStateTime;
+	StartWorldTime = Other.StartWorldTime;
+	bIsInhibited = Other.bIsInhibited;
+	IsPendingRemove = Other.IsPendingRemove;
+	OnRemovedDelegate = Other.OnRemovedDelegate;
+	PeriodHandle = Other.PeriodHandle;
+	DurationHandle = Other.DurationHandle;
+	PendingNext = Other.PendingNext;
+	return *this;
+}
+
 /** This is the core function that turns the ActiveGE 'on' or 'off */
 void FActiveGameplayEffect::CheckOngoingTagRequirements(const FGameplayTagContainer& OwnerTags, FActiveGameplayEffectsContainer& OwningContainer)
 {
@@ -1038,12 +1241,26 @@ void FActiveGameplayEffect::PostReplicatedChange(const struct FActiveGameplayEff
 //	FActiveGameplayEffectsContainer
 //
 // --------------------------------------------------------------------------------------------------------------------------------------------------------
-// 
+
 FActiveGameplayEffectsContainer::FActiveGameplayEffectsContainer()
 	: Owner(nullptr)
 	, ScopedLockCount(0)
 	, PendingRemoves(0)
+	, PendingAdds(0)
+	, PendingGameplayEffectHead(nullptr)
+{
+	PendingGameplayEffectNext = &PendingGameplayEffectHead;
+}
+
+FActiveGameplayEffectsContainer::~FActiveGameplayEffectsContainer()
 {
+	FActiveGameplayEffect* PendingGameplayEffect = PendingGameplayEffectHead;
+	if (PendingGameplayEffectHead)
+	{
+		FActiveGameplayEffect* Next = PendingGameplayEffectHead->PendingNext;
+		delete PendingGameplayEffectHead;
+		PendingGameplayEffectHead =	Next;
+	}
 }
 
 void FActiveGameplayEffectsContainer::RegisterWithOwner(UAbilitySystemComponent* InOwner)
@@ -1093,6 +1310,7 @@ void FActiveGameplayEffectsContainer::ExecuteActiveEffectsFrom(FGameplayEffectSp
 	// ------------------------------------------------------
 	
 	TArray< FGameplayEffectSpecHandle > ConditionalEffectSpecs;
+	bool ShouldHandleGameplayCues = true;
 
 	for (const FGameplayEffectExecutionDefinition& CurExecDef : SpecToUse.Def->Executions)
 	{
@@ -1134,6 +1352,12 @@ void FActiveGameplayEffectsContainer::ExecuteActiveEffectsFrom(FGameplayEffectSp
 				}
 				InvokeGameplayCueExecute |= InternalExecuteMod(SpecToUse, CurExecMod);
 			}
+
+			// If execution handled GameplayCues, we dont have to.
+			if (ExecutionOutput.AreGameplayCuesHandledManually())
+			{
+				InvokeGameplayCueExecute = false;
+			}
 		}
 	}
 
@@ -1166,30 +1390,39 @@ void FActiveGameplayEffectsContainer::ExecuteActiveEffectsFrom(FGameplayEffectSp
 
 void FActiveGameplayEffectsContainer::ExecutePeriodicGameplayEffect(FActiveGameplayEffectHandle Handle)
 {
-	GAMEPLAYEFFECT_SCOPE_LOCK();
-	FActiveGameplayEffect* ActiveEffect = GetActiveGameplayEffect(Handle);
-	if (ActiveEffect)
+	for (FActiveGameplayEffect& ActiveEffect : this)
 	{
 		if (UE_LOG_ACTIVE(VLogAbilitySystem, Log))
 		{
-			ABILITY_VLOG(Owner->OwnerActor, Log, TEXT("Executed Periodic Effect %s"), *ActiveEffect->Spec.Def->GetFName().ToString());
-			for (FGameplayModifierInfo Modifier : ActiveEffect->Spec.Def->Modifiers)
+			ABILITY_VLOG(Owner->OwnerActor, Log, TEXT("Executed Periodic Effect %s"), *ActiveEffect.Spec.Def->GetFName().ToString());
+			for (FGameplayModifierInfo Modifier : ActiveEffect.Spec.Def->Modifiers)
 			{
 				float Magnitude = 0.f;
-				Modifier.ModifierMagnitude.AttemptCalculateMagnitude(ActiveEffect->Spec, Magnitude);
+				Modifier.ModifierMagnitude.AttemptCalculateMagnitude(ActiveEffect.Spec, Magnitude);
 				ABILITY_VLOG(Owner->OwnerActor, Log, TEXT("         %s: %s %f"), *Modifier.Attribute.GetName(), *EGameplayModOpToString(Modifier.ModifierOp), Magnitude);
 			}
 		}
 
 		// Execute
-		ExecuteActiveEffectsFrom(ActiveEffect->Spec);
+		ExecuteActiveEffectsFrom(ActiveEffect.Spec);
 	}
 }
 
 FActiveGameplayEffect* FActiveGameplayEffectsContainer::GetActiveGameplayEffect(const FActiveGameplayEffectHandle Handle)
 {
-	// Could make this a map for quicker lookup
-	for (FActiveGameplayEffect& Effect : GameplayEffects)
+	for (FActiveGameplayEffect& Effect : this)
+	{
+		if (Effect.Handle == Handle)
+		{
+			return &Effect;
+		}
+	}
+	return nullptr;
+}
+
+const FActiveGameplayEffect* FActiveGameplayEffectsContainer::GetActiveGameplayEffect(const FActiveGameplayEffectHandle Handle) const
+{
+	for (const FActiveGameplayEffect& Effect : this)
 	{
 		if (Effect.Handle == Handle)
 		{
@@ -1375,7 +1608,7 @@ FActiveGameplayEffect* FActiveGameplayEffectsContainer::FindStackableActiveGamep
 		// do a linear search through GameplayEffects to find the actual FActiveGameplayEffect (due to unstable nature of the GameplayEffects array).
 		// If this becomes a slow point in the profiler, the map may still be useful as an early out to avoid an unnecessary sweep.
 
-		for (FActiveGameplayEffect& ActiveEffect : GameplayEffects)
+		for (FActiveGameplayEffect& ActiveEffect: this)
 		{
 			// Aggregate by source stacking additionally requires the source ability component to match
 			if (ActiveEffect.Spec.Def == Spec.Def && ((StackingType == EGameplayEffectStackingType::AggregateByTarget) || (SourceASC == ActiveEffect.Spec.GetContext().GetInstigatorAbilitySystemComponent())))
@@ -1466,13 +1699,21 @@ void FActiveGameplayEffectsContainer::SetBaseAttributeValueFromReplication(FGame
 	}
 }
 
+void FActiveGameplayEffectsContainer::GetAllActiveGameplayEffectSpecs(TArray<FGameplayEffectSpec>& OutSpecCopies)
+{
+	for (const FActiveGameplayEffect& ActiveEffect : this)	
+	{
+		OutSpecCopies.Add(ActiveEffect.Spec);
+	}
+}
+
 float FActiveGameplayEffectsContainer::GetGameplayEffectDuration(FActiveGameplayEffectHandle Handle) const
 {
-	for (int32 idx = 0; idx < GameplayEffects.Num(); ++idx)
+	for (const FActiveGameplayEffect& ActiveEffect : this)
 	{
-		if (GameplayEffects[idx].Handle == Handle)
+		if (ActiveEffect.Handle == Handle)
 		{
-			return GameplayEffects[idx].GetDuration();
+			return ActiveEffect.GetDuration();
 		}
 	}
 	
@@ -1482,7 +1723,7 @@ float FActiveGameplayEffectsContainer::GetGameplayEffectDuration(FActiveGameplay
 
 float FActiveGameplayEffectsContainer::GetGameplayEffectMagnitude(FActiveGameplayEffectHandle Handle, FGameplayAttribute Attribute) const
 {
-	for (FActiveGameplayEffect Effect : GameplayEffects)
+	for (const FActiveGameplayEffect& Effect : this)
 	{
 		if (Effect.Handle == Handle)
 		{
@@ -1506,7 +1747,7 @@ float FActiveGameplayEffectsContainer::GetGameplayEffectMagnitude(FActiveGamepla
 const FGameplayTagContainer* FActiveGameplayEffectsContainer::GetGameplayEffectSourceTagsFromHandle(FActiveGameplayEffectHandle Handle) const
 {
 	// @todo: Need to consider this with tag changes
-	for (const FActiveGameplayEffect& Effect : GameplayEffects)
+	for (const FActiveGameplayEffect& Effect : this)
 	{
 		if (Effect.Handle == Handle)
 		{
@@ -1520,12 +1761,10 @@ const FGameplayTagContainer* FActiveGameplayEffectsContainer::GetGameplayEffectS
 const FGameplayTagContainer* FActiveGameplayEffectsContainer::GetGameplayEffectTargetTagsFromHandle(FActiveGameplayEffectHandle Handle) const
 {
 	// @todo: Need to consider this with tag changes
-	for (const FActiveGameplayEffect& Effect : GameplayEffects)
+	const FActiveGameplayEffect* Effect = GetActiveGameplayEffect(Handle);
+	if (Effect)
 	{
-		if (Effect.Handle == Handle)
-		{
-			return Effect.Spec.CapturedTargetTags.GetAggregatedTags();
-		}
+		return Effect->Spec.CapturedTargetTags.GetAggregatedTags();
 	}
 
 	return nullptr;
@@ -1747,19 +1986,45 @@ FActiveGameplayEffect* FActiveGameplayEffectsContainer::ApplyGameplayEffectSpec(
 	}
 	else
 	{
-		if (ScopedLockCount > 0 && GameplayEffects.GetSlack() <= 0)
+		FActiveGameplayEffectHandle NewHandle = FActiveGameplayEffectHandle::GenerateNewHandle(Owner);
+
+		if (ScopedLockCount > 0 && GameplayEffects_Internal.GetSlack() <= 0)
 		{
-			// TODO: Instead of failing we could add to a pending add list. For that to work though all other places that look at GameplayEffects list, they will need to look at
-			// the pending list as well.
-			
-			ABILITY_LOG(Error, TEXT("******************************************************"));
-			ABILITY_LOG(Error, TEXT("FActiveGameplayEffectsContainer::ApplyGameplayEffectSpec is attemping to add GameplayEffect %s while locked and having no more slack!"));
-			ABILITY_LOG(Error, TEXT("******************************************************"));
-			ensure(false);
+			/**
+			 *	If we have no more slack and we are scope locked, we need to put this addition on our pending GE list, which will be moved
+			 *	onto the real active GE list once the scope lock is over.
+			 *	
+			 *	To avoid extra heap allocations, each active gameplayeffect container keeps a linked list of pending GEs. This list is allocated
+			 *	on demand and resused in subsequent pending adds. The code below will either 1) Alloc a new pending GE 2) reuse an existing pending GE.
+			 *	The move operator is used to move stuff to and from these pending GEs to avoid deep copies.
+			 */
+
+			check(PendingGameplayEffectNext);
+			if (*PendingGameplayEffectNext == nullptr)
+			{
+				// We have no memory allocated to put our next pending GE, so make a new one.
+				// [#1] If you change this, please change #1-3!!!
+				AppliedActiveGE = new FActiveGameplayEffect(NewHandle, Spec, GetWorldTime(), GetGameStateTime(), InPredictionKey);
+				*PendingGameplayEffectNext = AppliedActiveGE;
+			}
+			else
+			{
+				// We already had memory allocated to put a pending GE, move in.
+				// [#2] If you change this, please change #1-3!!!
+				**PendingGameplayEffectNext = MoveTemp(FActiveGameplayEffect(NewHandle, Spec, GetWorldTime(), GetGameStateTime(), InPredictionKey));
+				AppliedActiveGE = *PendingGameplayEffectNext;
+			}
+
+			// The next pending GameplayEffect goes to where our PendingNext points
+			PendingGameplayEffectNext = &AppliedActiveGE->PendingNext;
+			PendingAdds++;
 		}
+		else
+		{
 
-		FActiveGameplayEffectHandle NewHandle = FActiveGameplayEffectHandle::GenerateNewHandle(Owner);
-		AppliedActiveGE = new(GameplayEffects) FActiveGameplayEffect(NewHandle, Spec, GetWorldTime(), GetGameStateTime(), InPredictionKey);
+			// [#3] If you change this, please change #1-3!!!
+			AppliedActiveGE = new(GameplayEffects_Internal) FActiveGameplayEffect(NewHandle, Spec, GetWorldTime(), GetGameStateTime(), InPredictionKey);
+		}
 	}
 
 	FGameplayEffectSpec& AppliedEffectSpec = AppliedActiveGE->Spec;
@@ -1928,23 +2193,24 @@ void FActiveGameplayEffectsContainer::AddActiveGameplayEffectGrantedTagsAndModif
 /** Called on server to remove a GameplayEffect */
 bool FActiveGameplayEffectsContainer::RemoveActiveGameplayEffect(FActiveGameplayEffectHandle Handle, int32 StacksToRemove)
 {
-	// Could make this a map for quicker lookup
-	for (int32 Idx = 0; Idx < GameplayEffects.Num(); ++Idx)
+	// Iterating through manually since this is a removal operation and we need to pass the index into InternalRemoveActiveGameplayEffect
+	for (int32 ActiveGEIdx = 0; ActiveGEIdx < GetNumGameplayEffects(); ++ActiveGEIdx)
 	{
-		if (GameplayEffects[Idx].Handle == Handle)
+		FActiveGameplayEffect& Effect = *GetActiveGameplayEffect(ActiveGEIdx);
+		if (Effect.Handle == Handle && Effect.IsPendingRemove == false)
 		{
 			if (UE_LOG_ACTIVE(VLogAbilitySystem, Log))
 			{
-				ABILITY_VLOG(Owner->OwnerActor, Log, TEXT("Removed %s"), *GameplayEffects[Idx].Spec.Def->GetFName().ToString());
-				for (FGameplayModifierInfo Modifier : GameplayEffects[Idx].Spec.Def->Modifiers)
+				ABILITY_VLOG(Owner->OwnerActor, Log, TEXT("Removed %s"), *Effect.Spec.Def->GetFName().ToString());
+				for (FGameplayModifierInfo Modifier : Effect.Spec.Def->Modifiers)
 				{
 					float Magnitude = 0.f;
-					Modifier.ModifierMagnitude.AttemptCalculateMagnitude(GameplayEffects[Idx].Spec, Magnitude);
+					Modifier.ModifierMagnitude.AttemptCalculateMagnitude(Effect.Spec, Magnitude);
 					ABILITY_VLOG(Owner->OwnerActor, Log, TEXT("         %s: %s %f"), *Modifier.Attribute.GetName(), *EGameplayModOpToString(Modifier.ModifierOp), Magnitude);
 				}
 			}
 
-			InternalRemoveActiveGameplayEffect(Idx, StacksToRemove, true);
+			InternalRemoveActiveGameplayEffect(ActiveGEIdx, StacksToRemove, true);
 			return true;
 		}
 	}
@@ -1957,9 +2223,9 @@ bool FActiveGameplayEffectsContainer::InternalRemoveActiveGameplayEffect(int32 I
 {
 	SCOPE_CYCLE_COUNTER(STAT_RemoveActiveGameplayEffect);
 	
-	if (ensure(Idx < GameplayEffects.Num()))
+	if (ensure(Idx < GetNumGameplayEffects()))
 	{
-		FActiveGameplayEffect& Effect = GameplayEffects[Idx];
+		FActiveGameplayEffect& Effect = *GetActiveGameplayEffect(Idx);
 		ensure(!Effect.IsPendingRemove);
 
 		if (StacksToRemove > 0 && Effect.Spec.StackCount > StacksToRemove)
@@ -1986,18 +2252,18 @@ bool FActiveGameplayEffectsContainer::InternalRemoveActiveGameplayEffect(int32 I
 
 		if (ShouldInvokeGameplayCueEvent)
 		{
-			Owner->InvokeGameplayCueEvent(GameplayEffects[Idx].Spec, EGameplayCueEvent::Removed);
+			Owner->InvokeGameplayCueEvent(Effect.Spec, EGameplayCueEvent::Removed);
 		}
 
 		InternalOnActiveGameplayEffectRemoved(Effect);
 
-		if (GameplayEffects[Idx].DurationHandle.IsValid())
+		if (Effect.DurationHandle.IsValid())
 		{
-			Owner->GetWorld()->GetTimerManager().ClearTimer(GameplayEffects[Idx].DurationHandle);
+			Owner->GetWorld()->GetTimerManager().ClearTimer(Effect.DurationHandle);
 		}
-		if (GameplayEffects[Idx].PeriodHandle.IsValid())
+		if (Effect.PeriodHandle.IsValid())
 		{
-			Owner->GetWorld()->GetTimerManager().ClearTimer(GameplayEffects[Idx].PeriodHandle);
+			Owner->GetWorld()->GetTimerManager().ClearTimer(Effect.PeriodHandle);
 		}
 
 		if (bIsNetAuthority && Owner->OwnerActor)
@@ -2022,11 +2288,18 @@ bool FActiveGameplayEffectsContainer::InternalRemoveActiveGameplayEffect(int32 I
 		else
 		{
 			// Not locked, so do the removal right away.
-			GameplayEffects.RemoveAtSwap(Idx);
-			MarkArrayDirty();
+
+			// If we are not scope locked, then there is no way this idx should be referring to something on the pending add list.
+			// It is possible to remove a GE that is pedning add, but it would happen while the scope lock is still in effect, resulting
+			// in a pending remove being set.
+			check(Idx < GameplayEffects_Internal.Num());
+
+			GameplayEffects_Internal.RemoveAtSwap(Idx);
 			ModifiedArray = true;
 		}
 
+		MarkArrayDirty();
+
 		// Hack: force netupdate on owner. This isn't really necessary in real gameplay but is nice
 		// during debugging where breakpoints or pausing can mess up network update times. Open issue
 		// with network team.
@@ -2187,7 +2460,7 @@ bool FActiveGameplayEffectsContainer::HasApplicationImmunityToSpec(const FGamepl
 		return false;
 	}
 
-	for (const FActiveGameplayEffect& Effect : GameplayEffects)
+	for (const FActiveGameplayEffect& Effect :  this)
 	{
 		if (Effect.Spec.Def->GrantedApplicationImmunityTags.IsEmpty() == false && Effect.Spec.Def->GrantedApplicationImmunityTags.RequirementsMet( *AggregatedSourceTags ))
 		{
@@ -2230,10 +2503,12 @@ float FActiveGameplayEffectsContainer::GetWorldTime() const
 void FActiveGameplayEffectsContainer::CheckDuration(FActiveGameplayEffectHandle Handle)
 {
 	GAMEPLAYEFFECT_SCOPE_LOCK();
-	for (int32 Idx = 0; Idx < GameplayEffects.Num(); ++Idx)
+	// Intentionally iterating through only the internal list since we need to pass the index for removal
+	// and pending effecets will never need to be checked for duration expiration (They will be added to the real list first)
+	for (int32 ActiveGEIdx = 0; ActiveGEIdx < GameplayEffects_Internal.Num(); ++ActiveGEIdx)
 	{
-		FActiveGameplayEffect& Effect = GameplayEffects[Idx];
-		if (Effect.Handle == Handle)
+		FActiveGameplayEffect& Effect = GameplayEffects_Internal[ActiveGEIdx];
+		if (Effect.Handle == Handle && Effect.IsPendingRemove == false)
 		{
 			FTimerManager& TimerManager = Owner->GetWorld()->GetTimerManager();
 
@@ -2257,7 +2532,7 @@ void FActiveGameplayEffectsContainer::CheckDuration(FActiveGameplayEffectHandle
 					TimerManager.ClearTimer(Effect.PeriodHandle);
 				}
 
-				InternalRemoveActiveGameplayEffect(Idx, -1, false);
+				InternalRemoveActiveGameplayEffect(ActiveGEIdx, -1, false);
 			}
 			else
 			{
@@ -2317,7 +2592,7 @@ TArray<float> FActiveGameplayEffectsContainer::GetActiveEffectsTimeRemaining(con
 
 	TArray<float>	ReturnList;
 
-	for (const FActiveGameplayEffect& Effect : GameplayEffects)
+	for (const FActiveGameplayEffect& Effect : this)
 	{
 		if (!Query.Matches(Effect))
 		{
@@ -2340,7 +2615,7 @@ TArray<float> FActiveGameplayEffectsContainer::GetActiveEffectsDuration(const FA
 
 	TArray<float>	ReturnList;
 
-	for (const FActiveGameplayEffect &Effect : GameplayEffects)
+	for (const FActiveGameplayEffect& Effect : this)
 	{
 		if (!Query.Matches(Effect))
 		{
@@ -2355,11 +2630,17 @@ TArray<float> FActiveGameplayEffectsContainer::GetActiveEffectsDuration(const FA
 }
 
 void FActiveGameplayEffectsContainer::RemoveActiveEffects(const FActiveGameplayEffectQuery Query, int32 StacksToRemove)
-{	
-	for (int32 idx=GameplayEffects.Num()-1; idx >= 0; --idx)
+{
+	int32 Count=0;
+	for (const FActiveGameplayEffect& Effect : this)
+	{
+		Count++;
+	}
+	// Manually iterating through in reverse because this is a removal operation
+	for (int32 idx=GetNumGameplayEffects()-1; idx >= 0; --idx)
 	{
-		const FActiveGameplayEffect& Effect = GameplayEffects[idx];
-		if (Query.Matches(Effect))
+		const FActiveGameplayEffect& Effect = *GetActiveGameplayEffect(idx);
+		if (Effect.IsPendingRemove==false && Query.Matches(Effect))
 		{
 			InternalRemoveActiveGameplayEffect(idx, StacksToRemove, true);			
 		}
@@ -2373,7 +2654,7 @@ FOnGameplayAttributeChange& FActiveGameplayEffectsContainer::RegisterGameplayAtt
 
 bool FActiveGameplayEffectsContainer::HasReceivedEffectWithPredictedKey(FPredictionKey PredictionKey) const
 {
-	for (const FActiveGameplayEffect& Effect : GameplayEffects)
+	for (const FActiveGameplayEffect& Effect : this)
 	{
 		if (Effect.PredictionKey == PredictionKey && Effect.PredictionKey.WasReceived() == true)
 		{
@@ -2386,7 +2667,7 @@ bool FActiveGameplayEffectsContainer::HasReceivedEffectWithPredictedKey(FPredict
 
 bool FActiveGameplayEffectsContainer::HasPredictedEffectWithPredictedKey(FPredictionKey PredictionKey) const
 {
-	for (const FActiveGameplayEffect& Effect : GameplayEffects)
+	for (const FActiveGameplayEffect& Effect : this)
 	{
 		if (Effect.PredictionKey == PredictionKey && Effect.PredictionKey.WasReceived() == false)
 		{
@@ -2406,7 +2687,7 @@ void FActiveGameplayEffectsContainer::GrabDebugSnapshot(FVisualLogEntry* Snapsho
 	TMultiMap<FGameplayAttribute, FActiveGameplayEffectsContainer::DebugExecutedGameplayEffectData> EffectMap;
 
 	// Add all of the active gameplay effects
-	for (const FActiveGameplayEffect& Effect : GameplayEffects)
+	for (const FActiveGameplayEffect& Effect : this)
 	{
 		ensure(Effect.Spec.Modifiers.Num() == Effect.Spec.Def->Modifiers.Num());
 		for (int32 Idx = 0; Idx < Effect.Spec.Modifiers.Num(); ++Idx)
@@ -2640,42 +2921,57 @@ void FInheritedTagContainer::RemoveTag(FGameplayTag TagToRemove)
 
 void FActiveGameplayEffectsContainer::IncrementLock()
 {
-	if (ScopedLockCount++ == 0)
-	{
-		// If we transitioned from unlocked to locked, make sure we have enough slack for additions.
-		// 4 is arbitrary determined. We should add instrumentation to get a better ideal/default number.
-		static const int32 GAMEPLAYEFFECT_MIN_SLACK = 4;
-
-		if (GameplayEffects.GetSlack() < GAMEPLAYEFFECT_MIN_SLACK)
-		{
-			GameplayEffects.Reserve( GameplayEffects.Num() + GAMEPLAYEFFECT_MIN_SLACK );
-		}
-	}
+	ScopedLockCount++;
 }
 
 void FActiveGameplayEffectsContainer::DecrementLock()
 {
-	if (--ScopedLockCount == 0 && PendingRemoves > 0)
+	if (--ScopedLockCount == 0)
 	{
-		// If we are transitioned from locked to unlocked, we now need to actually delete and pending kill gameplay effects
-		for (int32 idx=GameplayEffects.Num()-1; idx >= 0 && PendingRemoves > 0; --idx)
+		// ------------------------------------------
+		// Move any pending effects onto the real list
+		// ------------------------------------------
+		if (PendingAdds > 0)
+		{
+			FActiveGameplayEffect* PendingGameplayEffect = PendingGameplayEffectHead;
+			FActiveGameplayEffect* Stop = *PendingGameplayEffectNext;
+			while (PendingGameplayEffect != Stop)
+			{
+				GameplayEffects_Internal.Add(MoveTemp(*PendingGameplayEffect));
+				PendingGameplayEffect = PendingGameplayEffect->PendingNext;
+				PendingAdds--;
+			}
+
+			
+			if (!ensure(PendingAdds == 0))
+			{
+				ABILITY_LOG(Error, TEXT("~FScopedActiveGameplayEffectLock has %d pending adds after a scope lock removal"), PendingAdds);
+				PendingAdds = 0;
+			}
+
+			// Reset our pending GameplayEffect linked list
+			PendingGameplayEffectNext = &PendingGameplayEffectHead;
+		}
+
+		// -----------------------------------------
+		// Delete any pending remove effects
+		// -----------------------------------------
+		for (int32 idx=GameplayEffects_Internal.Num()-1; idx >= 0 && PendingRemoves > 0; --idx)
 		{
-			FActiveGameplayEffect& Effect = GameplayEffects[idx];
+			FActiveGameplayEffect& Effect = GameplayEffects_Internal[idx];
 
 			if (Effect.IsPendingRemove)
 			{
-				GameplayEffects.RemoveAtSwap(idx, 1, false);
+				GameplayEffects_Internal.RemoveAtSwap(idx, 1, false);
 				PendingRemoves--;
 			}
 		}
 
 		if (!ensure(PendingRemoves == 0))
 		{
-			ABILITY_LOG(Warning, TEXT("~FScopedActiveGameplayEffectLock has %d pending removes after a scope lock removal"), PendingRemoves);
+			ABILITY_LOG(Error, TEXT("~FScopedActiveGameplayEffectLock has %d pending removes after a scope lock removal"), PendingRemoves);
 			PendingRemoves = 0;
 		}
-
-		MarkArrayDirty();
 	}
 }
 
diff --git a/Engine/Source/Runtime/GameplayAbilities/Private/GameplayEffectExecutionCalculation.cpp b/Engine/Source/Runtime/GameplayAbilities/Private/GameplayEffectExecutionCalculation.cpp
--- a/Engine/Source/Runtime/GameplayAbilities/Private/GameplayEffectExecutionCalculation.cpp
+++ b/Engine/Source/Runtime/GameplayAbilities/Private/GameplayEffectExecutionCalculation.cpp
@@ -174,6 +174,7 @@ bool FGameplayEffectCustomExecutionParameters::AttemptGetCapturedAttributeAggreg
 FGameplayEffectCustomExecutionOutput::FGameplayEffectCustomExecutionOutput()
 	: bTriggerConditionalGameplayEffects(false)
 	, bHandledStackCountManually(false)
+	, bHandledGameplayCuesManually(false)
 {
 }
 
@@ -187,11 +188,21 @@ bool FGameplayEffectCustomExecutionOutput::IsStackCountHandledManually() const
 	return bHandledStackCountManually;
 }
 
+bool FGameplayEffectCustomExecutionOutput::AreGameplayCuesHandledManually() const
+{
+	return bHandledGameplayCuesManually;
+}
+
 void FGameplayEffectCustomExecutionOutput::MarkConditionalGameplayEffectsToTrigger()
 {
 	bTriggerConditionalGameplayEffects = true;
 }
 
+void FGameplayEffectCustomExecutionOutput::MarkGameplayCuesHandledManually()
+{
+	bHandledGameplayCuesManually = true;
+}
+
 bool FGameplayEffectCustomExecutionOutput::ShouldTriggerConditionalGameplayEffects() const
 {
 	return bTriggerConditionalGameplayEffects;
diff --git a/Engine/Source/Runtime/GameplayAbilities/Private/GameplayEffectTypes.cpp b/Engine/Source/Runtime/GameplayAbilities/Private/GameplayEffectTypes.cpp
--- a/Engine/Source/Runtime/GameplayAbilities/Private/GameplayEffectTypes.cpp
+++ b/Engine/Source/Runtime/GameplayAbilities/Private/GameplayEffectTypes.cpp
@@ -359,8 +359,9 @@ FString FGameplayTagRequirements::ToString() const
 
 void FActiveGameplayEffectsContainer::PrintAllGameplayEffects() const
 {
-	ABILITY_LOG_SCOPE(TEXT("ActiveGameplayEffects. Num: %d"), GameplayEffects.Num());
-	for (const FActiveGameplayEffect& Effect : GameplayEffects)
+	ABILITY_LOG_SCOPE(TEXT("ActiveGameplayEffects. Num: %d"), GetNumGameplayEffects());
+
+	for (const FActiveGameplayEffect& Effect : this)
 	{
 		Effect.PrintAll();
 	}
diff --git a/Engine/Source/Runtime/GameplayAbilities/Public/AbilitySystemGlobals.h b/Engine/Source/Runtime/GameplayAbilities/Public/AbilitySystemGlobals.h
--- a/Engine/Source/Runtime/GameplayAbilities/Public/AbilitySystemGlobals.h
+++ b/Engine/Source/Runtime/GameplayAbilities/Public/AbilitySystemGlobals.h
@@ -79,7 +79,7 @@ class GAMEPLAYABILITIES_API UAbilitySystemGlobals : public UObject
 
 	static UAbilitySystemGlobals& Get();
 
-	static UAbilitySystemComponent* GetAbilitySystemComponentFromActor(AActor* Actor, bool LookForComponent=false);
+	static UAbilitySystemComponent* GetAbilitySystemComponentFromActor(const AActor* Actor, bool LookForComponent=false);
 
 	/** Should allocate a project specific AbilityActorInfo struct. Caller is responsible for deallocation */
 	virtual FGameplayAbilityActorInfo * AllocAbilityActorInfo() const;
diff --git a/Engine/Source/Runtime/GameplayAbilities/Public/ActiveGameplayEffectIterator.h b/Engine/Source/Runtime/GameplayAbilities/Public/ActiveGameplayEffectIterator.h
new file mode 100644
--- /dev/null
+++ b/Engine/Source/Runtime/GameplayAbilities/Public/ActiveGameplayEffectIterator.h
@@ -0,0 +1,122 @@
+// Copyright 1998-2015 Epic Games, Inc. All Rights Reserved.
+
+/*================================================================================
+ActiveGameplayEffectIterator.h: implementation of ActiveGameplayEffect Iterator.
+
+The is an iterator for FActiveGameplayEffectContainer. The main points are:
+-This iterates over both the active gameplay effect array and the pending list of new GEs
+-This skips any gameplay effects which are pending remove
+-This locks the array while the iterator is in scope, meaning any adds or removes 
+ to the active gameplay effect array will have their memory operations deffered.
+ (That is, the effects will be added/removed, but not removed from memory until the scope lock
+ has been removed).
+	
+================================================================================*/
+
+
+template< typename ElementType, typename ContainerType >
+class FActiveGameplayEffectIterator
+{
+public:
+
+	FActiveGameplayEffectIterator(const ContainerType& InContainer, int32 StartIdx=0)
+		: index(StartIdx)
+		, Current(nullptr)
+		, Pending(nullptr)
+		, Container(const_cast<ContainerType&>(InContainer))
+	{
+		Container.IncrementLock();
+
+		if (index >= 0)
+		{
+			UpdateCurrent();
+		}
+	}
+
+	~FActiveGameplayEffectIterator()
+	{
+		Container.DecrementLock();
+	}
+
+
+	void operator++()
+	{
+		Next();
+	}
+
+	ElementType& operator*()
+	{
+		check(Current);
+		return *Current;
+	}
+
+	ElementType& operator->()
+	{
+		check(Current);
+		return *Current;
+	}
+
+	operator bool()
+	{
+		return (Current != nullptr);
+	}
+
+private:
+
+	FORCEINLINE ElementType* AdvancePending(ElementType** Next)
+	{
+		// Advance pending next only if this isnt the tail of the pending gameplay effect list
+		return (Next != const_cast<ElementType**>(Container.PendingGameplayEffectNext)) ? *Next : nullptr;
+	}
+
+	void Next()
+	{
+		if (index >= 0 )
+		{
+			// While iterating through the array, just increment the index
+			++index;
+		}
+
+		else if (Pending)
+		{
+			// While iteratin through the linked list, jump to the next one
+			Pending = AdvancePending(const_cast<ElementType**>(&Pending->PendingNext));
+		}
+	
+		UpdateCurrent();
+	}
+
+	void UpdateCurrent()
+	{
+		if (index < 0 )
+		{
+			// We are already iterating the linked list, current element should equal the pending list element.
+			Current = Pending;
+		}
+
+		else if (index >= Container.GameplayEffects_Internal.Num())
+		{
+			// Once we get to end of the array, we start iterating on the linked list. 
+			Pending = AdvancePending(const_cast<ElementType**>(&Container.PendingGameplayEffectHead));
+			Current = Pending;
+			index = INDEX_NONE;
+		}
+
+		else
+		{
+			// We are still iterating the array, current is equal to the indexed element.
+			Current = &Container.GameplayEffects_Internal[index];
+		}
+		
+		// If the current element is pending remove, go to the next one.
+		if (Current && Current->IsPendingRemove)
+		{
+			Next();
+		}
+	}
+
+	int32	index;
+	ElementType*	Current;
+	ElementType*	Pending;
+	ContainerType& Container;
+};
\ No newline at end of file
diff --git a/Engine/Source/Runtime/GameplayAbilities/Public/GameplayCueManager.h b/Engine/Source/Runtime/GameplayAbilities/Public/GameplayCueManager.h
--- a/Engine/Source/Runtime/GameplayAbilities/Public/GameplayCueManager.h
+++ b/Engine/Source/Runtime/GameplayAbilities/Public/GameplayCueManager.h
@@ -135,14 +135,14 @@ class GAMEPLAYABILITIES_API UGameplayCueManager : public UDataAsset
 	bool bAccelerationMapOutdated;
 #endif
 
-private:
+protected:
 
 	virtual void HandleGameplayCueNotify_Internal(AActor* TargetActor, int32 DataIdx, EGameplayCueEvent::Type EventType, FGameplayCueParameters Parameters);
 
-	#if WITH_EDITOR
-		//This handles the case where GameplayCueNotifications have changed between sessions, which is possible in editor.
-		void ReloadObjectLibrary(UWorld* World, const UWorld::InitializationValues IVS);
-	#endif
+#if WITH_EDITOR
+	//This handles the case where GameplayCueNotifications have changed between sessions, which is possible in editor.
+	void ReloadObjectLibrary(UWorld* World, const UWorld::InitializationValues IVS);
+#endif
 
 	void LoadObjectLibrary_Internal();
 
diff --git a/Engine/Source/Runtime/GameplayAbilities/Public/GameplayCueNotify_Actor.h b/Engine/Source/Runtime/GameplayAbilities/Public/GameplayCueNotify_Actor.h
--- a/Engine/Source/Runtime/GameplayAbilities/Public/GameplayCueNotify_Actor.h
+++ b/Engine/Source/Runtime/GameplayAbilities/Public/GameplayCueNotify_Actor.h
@@ -62,6 +62,10 @@ class GAMEPLAYABILITIES_API AGameplayCueNotify_Actor : public AActor
 	UPROPERTY(AssetRegistrySearchable)
 	FName GameplayCueName;
 
+	/** We will auto destroy this GameplayCueActor when the OnRemove event fires (after OnRemove is called). */
+	UPROPERTY(EditDefaultsOnly, Category = GameplayCue)
+	bool	bAutoDestroyOnRemove;
+
 	/** Does this Cue override other cues, or is it called in addition to them? E.g., If this is Damage.Physical.Slash, we wont call Damage.Physical afer we run this cue. */
 	UPROPERTY(EditDefaultsOnly, Category = GameplayCue)
 	bool IsOverride;
diff --git a/Engine/Source/Runtime/GameplayAbilities/Public/GameplayEffect.h b/Engine/Source/Runtime/GameplayAbilities/Public/GameplayEffect.h
--- a/Engine/Source/Runtime/GameplayAbilities/Public/GameplayEffect.h
+++ b/Engine/Source/Runtime/GameplayAbilities/Public/GameplayEffect.h
@@ -9,6 +9,7 @@
 #include "GameplayEffectTypes.h"
 #include "GameplayEffectAggregator.h"
 #include "GameplayEffectCalculation.h"
+#include "ActiveGameplayEffectIterator.h"
 #include "GameplayEffect.generated.h"
 
 struct FActiveGameplayEffect;
@@ -829,6 +830,16 @@
 
 public:
 
+	FGameplayEffectAttributeCaptureSpecContainer();
+
+	FGameplayEffectAttributeCaptureSpecContainer(FGameplayEffectAttributeCaptureSpecContainer&& Other);
+
+	FGameplayEffectAttributeCaptureSpecContainer(const FGameplayEffectAttributeCaptureSpecContainer& Other);
+
+	FGameplayEffectAttributeCaptureSpecContainer& operator=(FGameplayEffectAttributeCaptureSpecContainer&& Other);
+
+	FGameplayEffectAttributeCaptureSpecContainer& operator=(const FGameplayEffectAttributeCaptureSpecContainer& Other);
+
 	/**
 	 * Add a definition to be captured by the owner of the container. Will not add the definition if its exact
 	 * match already exists within the container.
@@ -902,31 +913,22 @@
 {
 	GENERATED_USTRUCT_BODY()
 
-	FGameplayEffectSpec()
-		: Def(nullptr)
-		, Duration(UGameplayEffect::INSTANT_APPLICATION)
-		, Period(UGameplayEffect::NO_PERIOD)
-		, ChanceToApplyToTarget(1.f)
-		, StackCount(1)
-		, bCompletedSourceAttributeCapture(false)
-		, bCompletedTargetAttributeCapture(false)
-		, bDurationLocked(false)
-		, Level(UGameplayEffect::INVALID_LEVEL)
-	{
+	// --------------------------------------------------------------------------------------------------------------------------
+	//	IMPORTANT: Any state added to FGameplayEffectSpec must be handled in the move/copy constructor/operator!
+	//	(When VS2012/2013 support is dropped, we can use compiler generated operators, but until then these need to be maintained manually!)
+	// --------------------------------------------------------------------------------------------------------------------------
 
-	}
+	FGameplayEffectSpec();
 
 	FGameplayEffectSpec(const UGameplayEffect* InDef, const FGameplayEffectContextHandle& InEffectContext, float Level = UGameplayEffect::INVALID_LEVEL);
-	
-	UPROPERTY()
-	const UGameplayEffect* Def;
-	
-	UPROPERTY()
-	TArray<FGameplayEffectModifiedAttribute> ModifiedAttributes;
-	
-	/** Attributes captured by the spec that are relevant to custom calculations, potentially in owned modifiers, etc.; NOT replicated to clients */
-	UPROPERTY(NotReplicated)
-	FGameplayEffectAttributeCaptureSpecContainer CapturedRelevantAttributes;
 
+	FGameplayEffectSpec(const FGameplayEffectSpec& Other);
+
+	FGameplayEffectSpec(FGameplayEffectSpec&& Other);
+
+	FGameplayEffectSpec& operator=(FGameplayEffectSpec&& Other);
+
+	FGameplayEffectSpec& operator=(const FGameplayEffectSpec& Other);
+
 	/**
 	 * Determines if the spec has capture specs with valid captures for all of the specified definitions.
@@ -963,9 +965,6 @@
 	float GetPeriod() const;
 	float GetChanceToApplyToTarget() const;
 
-	/** other effects that need to be applied to the target if this effect is successful */
-	TArray< FGameplayEffectSpecHandle > TargetEffectSpecs;
-
 	/** Set the context info: who and where this spec came from. */
 	void SetContext(FGameplayEffectContextHandle NewEffectContext);
 
@@ -982,6 +981,61 @@
 	/** Returns the magnitude of a SetByCaller modifier. Will return 0.f and Warn if the magnitude has not been set. */
 	float GetSetByCallerMagnitude(FName DataName) const;
 
+	void SetLevel(float InLevel);
+
+	float GetLevel() const;
+
+	void PrintAll() const;
+
+	FString ToSimpleString() const
+	{
+		return FString::Printf(TEXT("%s"), *Def->GetName());
+	}
+
+	const FGameplayEffectContextHandle& GetEffectContext() const
+	{
+		return EffectContext;
+	}
+
+	void CaptureAttributeDataFromTarget(UAbilitySystemComponent* TargetAbilitySystemComponent);
+
+	/**
+	 * Get the computed magnitude of the modifier on the spec with the specified index
+	 * 
+	 * @param ModifierIndx			Modifier to get
+	 * @param bFactorInStackCount	If true, the calculation will include the stack count
+	 * 
+	 * @return Computed magnitude
+	 */
+	float GetModifierMagnitude(int32 ModifierIdx, bool bFactorInStackCount) const;
+
+	void CalculateModifierMagnitudes();
+
+private:
+
+	void CaptureDataFromSource();
+
+	/** Helper function to initialize all of the capture definitions required by the spec */
+	void SetupAttributeCaptureDefinitions();
+
+public:
+
+	// -----------------------------------------------------------------------
+
+	/** GameplayEfect definition. The static data that this spec points to. */
+	UPROPERTY()
+	const UGameplayEffect* Def;
+	
+	UPROPERTY()
+	TArray<FGameplayEffectModifiedAttribute> ModifiedAttributes;
+	
+	/** Attributes captured by the spec that are relevant to custom calculations, potentially in owned modifiers, etc.; NOT replicated to clients */
+	UPROPERTY(NotReplicated)
+	FGameplayEffectAttributeCaptureSpecContainer CapturedRelevantAttributes;
+
+	/** other effects that need to be applied to the target if this effect is successful */
+	TArray< FGameplayEffectSpecHandle > TargetEffectSpecs;
+
 	// The duration in seconds of this effect
 	// instantaneous effects should have a duration of UGameplayEffect::INSTANT_APPLICATION
 	// effects that last forever should have a duration of UGameplayEffect::INFINITE_DURATION
@@ -1026,46 +1080,11 @@
 	/** Whether the duration of the spec is locked or not; If it is, attempts to set it will fail */
 	UPROPERTY(NotReplicated)
 	uint32 bDurationLocked : 1;
-	
-	/**
-	 * Get the computed magnitude of the modifier on the spec with the specified index
-	 * 
-	 * @param ModifierIndx			Modifier to get
-	 * @param bFactorInStackCount	If true, the calculation will include the stack count
-	 * 
-	 * @return Computed magnitude
-	 */
-	float GetModifierMagnitude(int32 ModifierIdx, bool bFactorInStackCount) const;
 
-	void CalculateModifierMagnitudes();
-
-	void SetLevel(float InLevel);
-
-	float GetLevel() const;
-
-	void PrintAll() const;
-
-	FString ToSimpleString() const
-	{
-		return FString::Printf(TEXT("%s"), *Def->GetName());
-	}
-
-	const FGameplayEffectContextHandle& GetEffectContext() const
-	{
-		return EffectContext;
-	}
-
-	void CaptureAttributeDataFromTarget(UAbilitySystemComponent* TargetAbilitySystemComponent);
-
 private:
 
 	/** Map of set by caller magnitudes */
 	TMap<FName, float>	SetByCallerMagnitudes;
-
-	void CaptureDataFromSource();
-
-	/** Helper function to initialize all of the capture definitions required by the spec */
-	void SetupAttributeCaptureDefinitions();
 	
 	UPROPERTY()
 	FGameplayEffectContextHandle EffectContext; // This tells us how we got here (who / what applied us)
@@ -1083,31 +1102,59 @@
  *
  */
 USTRUCT()
-struct FActiveGameplayEffect : public FFastArraySerializerItem
+struct GAMEPLAYABILITIES_API FActiveGameplayEffect : public FFastArraySerializerItem
 {
 	GENERATED_USTRUCT_BODY()
 
-	FActiveGameplayEffect()
-		: StartGameStateTime(0)
-		, CachedStartGameStateTime(0)
-		, StartWorldTime(0.f)
-		, bIsInhibited(true)
-		, IsPendingRemove(false)
+	// ---------------------------------------------------------------------------------------------------------------------------------
+	//  IMPORTANT: Any new state added to FActiveGameplayEffect must be handled in the copy/move constructor/operator
+	//	(When VS2012/2013 support is dropped, we can use compiler generated operators, but until then these need to be maintained manually)
+	// ---------------------------------------------------------------------------------------------------------------------------------
+
+	FActiveGameplayEffect();
+
+	FActiveGameplayEffect(const FActiveGameplayEffect& Other);
+
+	FActiveGameplayEffect(FActiveGameplayEffectHandle InHandle, const FGameplayEffectSpec &InSpec, float CurrentWorldTime, int32 InStartGameStateTime, FPredictionKey InPredictionKey);
+	
+	FActiveGameplayEffect(FActiveGameplayEffect&& Other);
+
+	FActiveGameplayEffect& operator=(FActiveGameplayEffect&& other);
+
+	FActiveGameplayEffect& operator=(const FActiveGameplayEffect& other);
+
+	float GetTimeRemaining(float WorldTime)
 	{
+		float Duration = GetDuration();		
+		return (Duration == UGameplayEffect::INFINITE_DURATION ? -1.f : Duration - (WorldTime - StartWorldTime));
+	}
+	
+	float GetDuration() const
+	{
+		return Spec.GetDuration();
 	}
 
-	FActiveGameplayEffect(FActiveGameplayEffectHandle InHandle, const FGameplayEffectSpec &InSpec, float CurrentWorldTime, int32 InStartGameStateTime, FPredictionKey InPredictionKey)
-		: Handle(InHandle)
-		, Spec(InSpec)
-		, PredictionKey(InPredictionKey)
-		, StartGameStateTime(InStartGameStateTime)
-		, CachedStartGameStateTime(InStartGameStateTime)
-		, StartWorldTime(CurrentWorldTime)
-		, bIsInhibited(true)
-		, IsPendingRemove(false)
+	float GetPeriod() const
 	{
+		return Spec.GetPeriod();
 	}
 
+	void CheckOngoingTagRequirements(const FGameplayTagContainer& OwnerTags, struct FActiveGameplayEffectsContainer& OwningContainer);
+
+	void PrintAll() const;
+
+	void PreReplicatedRemove(const struct FActiveGameplayEffectsContainer &InArray);
+	void PostReplicatedAdd(const struct FActiveGameplayEffectsContainer &InArray);
+	void PostReplicatedChange(const struct FActiveGameplayEffectsContainer &InArray);
+
+	bool operator==(const FActiveGameplayEffect& Other)
+	{
+		return Handle == Other.Handle;
+	}
+
+	// ---------------------------------------------------------------------------------------------------------------------------------
+
+	/** Globally unique ID for identify this active gameplay effect. Can be used to look up owner. Not networked. */
 	FActiveGameplayEffectHandle Handle;
 
 	UPROPERTY()
@@ -1139,34 +1186,7 @@
 
 	FTimerHandle DurationHandle;
 
-	float GetTimeRemaining(float WorldTime)
-	{
-		float Duration = GetDuration();		
-		return (Duration == UGameplayEffect::INFINITE_DURATION ? -1.f : Duration - (WorldTime - StartWorldTime));
-	}
-	
-	float GetDuration() const
-	{
-		return Spec.GetDuration();
-	}
-
-	float GetPeriod() const
-	{
-		return Spec.GetPeriod();
-	}
-
-	void CheckOngoingTagRequirements(const FGameplayTagContainer& OwnerTags, struct FActiveGameplayEffectsContainer& OwningContainer);
-
-	void PrintAll() const;
-
-	void PreReplicatedRemove(const struct FActiveGameplayEffectsContainer &InArray);
-	void PostReplicatedAdd(const struct FActiveGameplayEffectsContainer &InArray);
-	void PostReplicatedChange(const struct FActiveGameplayEffectsContainer &InArray);
-
-	bool operator==(const FActiveGameplayEffect& Other)
-	{
-		return Handle == Other.Handle;
-	}
+	FActiveGameplayEffect* PendingNext;
 };
 
 /**
@@ -1235,6 +1255,18 @@
  *  
  * This should only be used by UAbilitySystemComponent. All of this could just live in UAbilitySystemComponent except that we need a distinct USTRUCT to implement FFastArraySerializer.
  *
+ * The preferred way to iterate through the ActiveGameplayEffectContainer is with CreateConstITerator/CreateIterator or stl style range iteration:
+ * 
+ * 
+ *	for (const FActiveGameplayEffect& Effect : this)
+ *	{
+ *	}
+ *
+ *	for (auto It = CreateConstIterator(); It; ++It) 
+ *	{
+ *	}
+ *
+ *
  */
 USTRUCT()
 struct FActiveGameplayEffectsContainer : public FFastArraySerializer
@@ -1245,8 +1277,14 @@
 	friend class UAbilitySystemComponent;
 	friend struct FScopedActiveGameplayEffectLock;
 	friend class AAbilitySystemDebugHUD;
+	friend class FActiveGameplayEffectIterator<const FActiveGameplayEffect, FActiveGameplayEffectsContainer>;
+	friend class FActiveGameplayEffectIterator<FActiveGameplayEffect, FActiveGameplayEffectsContainer>;
 
+	typedef FActiveGameplayEffectIterator<const FActiveGameplayEffect, FActiveGameplayEffectsContainer> ConstIterator;
+	typedef FActiveGameplayEffectIterator<FActiveGameplayEffect, FActiveGameplayEffectsContainer> Iterator;
+
 	FActiveGameplayEffectsContainer();
+	~FActiveGameplayEffectsContainer();
 
 	UAbilitySystemComponent* Owner;
 
@@ -1270,6 +1308,7 @@
 	FActiveGameplayEffect* ApplyGameplayEffectSpec(const FGameplayEffectSpec& Spec, FPredictionKey InPredictionKey);
 
 	FActiveGameplayEffect* GetActiveGameplayEffect(const FActiveGameplayEffectHandle Handle);
+	const FActiveGameplayEffect* GetActiveGameplayEffect(const FActiveGameplayEffectHandle Handle) const;
 		
 	void ExecuteActiveEffectsFrom(FGameplayEffectSpec &Spec, FPredictionKey PredictionKey = FPredictionKey() );
 	
@@ -1313,16 +1352,21 @@
 
 	void PrintAllGameplayEffects() const;
 
-	int32 GetNumGameplayEffects() const
+	/**
+	 *	Returns the total number of gameplay effects.
+	 *	NOTE this does include GameplayEffects that pending removal.
+	 *	Any pending remove gameplay effects are deleted at the end of their scope lock
+	 */
+	FORCEINLINE int32 GetNumGameplayEffects() const
 	{
-		return GameplayEffects.Num();
+		return GameplayEffects_Internal.Num() + PendingAdds;
 	}
 
 	void CheckDuration(FActiveGameplayEffectHandle Handle);
 
 	bool NetDeltaSerialize(FNetDeltaSerializeInfo & DeltaParms)
 	{
-		return FastArrayDeltaSerialize<FActiveGameplayEffect>(GameplayEffects, DeltaParms, *this);
+		return FastArrayDeltaSerialize<FActiveGameplayEffect>(GameplayEffects_Internal, DeltaParms, *this);
 	}
 
 	void PreDestroy();	
@@ -1347,6 +1391,8 @@
 		
 	void SetBaseAttributeValueFromReplication(FGameplayAttribute Attribute, float BaseBalue);
 
+	void GetAllActiveGameplayEffectSpecs(TArray<FGameplayEffectSpec>& OutSpecCopies);
+
 	// -------------------------------------------------------------------------------------------
 
 	FOnGameplayAttributeChange& RegisterGameplayAttributeEvent(FGameplayAttribute Attribute);
@@ -1357,11 +1403,43 @@
 
 	void IncrementLock();
 	void DecrementLock();
+	
+	FORCEINLINE ConstIterator CreateConstIterator() const { return ConstIterator(*this);	}
+	FORCEINLINE Iterator CreateIterator() { return Iterator(*this);	}
 
 private:
 
+	/**
+	 *	Accessors for internal functions to get GameplayEffects directly by index.
+	 *	Note this will return GameplayEffects that are pending removal!
+	 *	
+	 *	To iteratre over all 'valid' gameplay effects, use the CreateConstIterator/CreateIterator or the stl style range iterator
+	 */
+	FORCEINLINE const FActiveGameplayEffect* GetActiveGameplayEffect(int32 idx) const
+	{
+		return const_cast<FActiveGameplayEffectsContainer*>(this)->GetActiveGameplayEffect(idx);
+	}
+
+	FORCEINLINE FActiveGameplayEffect* GetActiveGameplayEffect(int32 idx)
+	{
+		if (idx < GameplayEffects_Internal.Num())
+		{
+			return &GameplayEffects_Internal[idx];
+		}
+
+		idx -= GameplayEffects_Internal.Num();
+		FActiveGameplayEffect* Ptr = PendingGameplayEffectHead;
+		while(idx-- > 0 && Ptr)
+		{
+			Ptr = Ptr->PendingNext;
+		}
+
+		return idx <= 0 ? Ptr : nullptr;
+	}
+
+	/** Our active list of Effects. Do not access this directly (Even from internal functions!) Use GetNumGameplayEffect() / GetGameplayEffect() ! */
 	UPROPERTY()
-	TArray<FActiveGameplayEffect>	GameplayEffects;
+	TArray<FActiveGameplayEffect>	GameplayEffects_Internal;
 
 	void InternalUpdateNumericalAttribute(FGameplayAttribute Attribute, float NewValue, const FGameplayEffectModCallbackData* ModData);
 	
@@ -1432,6 +1510,21 @@
 
 	mutable int32 ScopedLockCount;
 	int32 PendingRemoves;
+	int32 PendingAdds;
+
+	FActiveGameplayEffect*	PendingGameplayEffectHead;	// Head of pending GE linked list
+	FActiveGameplayEffect** PendingGameplayEffectNext;	// Points to the where to store the next pending GE (starts pointing at head, as more are added, points further down the list).
+
+	/**
+	 * DO NOT USE DIRECTLY
+	 * STL-like iterators to enable range-based for loop support.
+	 */
+
+	FORCEINLINE friend Iterator begin(FActiveGameplayEffectsContainer* Container) { return Container->CreateIterator(); }
+	FORCEINLINE friend Iterator end(FActiveGameplayEffectsContainer* Container) { return Iterator(*Container, -1); }
+
+	FORCEINLINE friend ConstIterator begin(const FActiveGameplayEffectsContainer* Container) { return Container->CreateConstIterator(); }
+	FORCEINLINE friend ConstIterator end(const FActiveGameplayEffectsContainer* Container) { return ConstIterator(*Container, -1); }
 };
 
 template<>
@@ -1443,7 +1536,15 @@
 	};
 };
 
-
+/**
+ *	FScopedActiveGameplayEffectLock
+ *	Provides a mechanism for locking the active gameplay effect list while possibly invoking callbacks into gamecode.
+ *	For example, if some internal code in FActiveGameplayEffectsContainer is iterating through the active GE list
+ *	or holding onto a pointer to something in that list, any changes to that list could cause memory the move out from underneath.
+ *	
+ *	This scope lock will queue deletions and additions until after the scope is over. The additions and deletions will actually 
+ *	go through, but we will defer the memory operations to the active gameplay effect list.
+ */
 struct GAMEPLAYABILITIES_API FScopedActiveGameplayEffectLock
 {
 	FScopedActiveGameplayEffectLock(FActiveGameplayEffectsContainer& InContainer);
@@ -1453,5 +1554,4 @@
 	FActiveGameplayEffectsContainer& Container;
 };
 
-
 #define GAMEPLAYEFFECT_SCOPE_LOCK()	FScopedActiveGameplayEffectLock ActiveScopeLock(*this);
diff --git a/Engine/Source/Runtime/GameplayAbilities/Public/GameplayEffectExecutionCalculation.h b/Engine/Source/Runtime/GameplayAbilities/Public/GameplayEffectExecutionCalculation.h
--- a/Engine/Source/Runtime/GameplayAbilities/Public/GameplayEffectExecutionCalculation.h
+++ b/Engine/Source/Runtime/GameplayAbilities/Public/GameplayEffectExecutionCalculation.h
@@ -119,9 +119,15 @@ public:
 	/** Simple accessor for determining whether the execution has manually handled the stack count or not */
 	bool IsStackCountHandledManually() const;
 
+	/** Accessor for determining if GameplayCue events have already been handled */
+	bool AreGameplayCuesHandledManually() const;
+
 	/** Mark that the execution wants conditional gameplay effects to trigger */
 	void MarkConditionalGameplayEffectsToTrigger();
 
+	/** Mark that the execution wants conditional gameplay effects to trigger */
+	void MarkGameplayCuesHandledManually();
+
 	/** Simple accessor for determining whether the execution wants conditional gameplay effects to trigger or not */
 	bool ShouldTriggerConditionalGameplayEffects() const;
 
@@ -147,6 +153,10 @@ private:
 	/** If true, the execution itself has manually handled the stack count of the effect and the GE system doesn't have to automatically handle it */
 	UPROPERTY()
 	uint32 bHandledStackCountManually : 1;
+
+	/** If true, the execution itself has manually invoked all gameplay cues and the GE system doesn't have to automatically handle them. */
+	UPROPERTY()
+	uint32 bHandledGameplayCuesManually : 1;
 };
 
 UCLASS(BlueprintType, Blueprintable, Abstract)
diff --git a/Engine/Source/Runtime/GameplayAbilities/Public/GameplayEffectTypes.h b/Engine/Source/Runtime/GameplayAbilities/Public/GameplayEffectTypes.h
--- a/Engine/Source/Runtime/GameplayAbilities/Public/GameplayEffectTypes.h
+++ b/Engine/Source/Runtime/GameplayAbilities/Public/GameplayEffectTypes.h
@@ -859,14 +859,50 @@ struct GAMEPLAYABILITIES_API FTagContainerAggregator
 
 	FTagContainerAggregator() : CacheIsValid(false) {}
 
+	FTagContainerAggregator(FTagContainerAggregator&& Other)
+		: CapturedActorTags(MoveTemp(Other.CapturedActorTags))
+		, CapturedSpecTags(MoveTemp(Other.CapturedSpecTags))
+		, ScopedTags(MoveTemp(Other.ScopedTags))
+		, CachedAggregator(MoveTemp(Other.CachedAggregator))
+		, CacheIsValid(Other.CacheIsValid)
+	{
+	}
+
+	FTagContainerAggregator(const FTagContainerAggregator& Other)
+		: CapturedActorTags(Other.CapturedActorTags)
+		, CapturedSpecTags(Other.CapturedSpecTags)
+		, ScopedTags(Other.ScopedTags)
+		, CachedAggregator(Other.CachedAggregator)
+		, CacheIsValid(Other.CacheIsValid)
+	{
+	}
+
+	FTagContainerAggregator& operator=(FTagContainerAggregator&& Other)
+	{
+		CapturedActorTags = MoveTemp(Other.CapturedActorTags);
+		CapturedSpecTags = MoveTemp(Other.CapturedSpecTags);
+		ScopedTags = MoveTemp(Other.ScopedTags);
+		CachedAggregator = MoveTemp(Other.CachedAggregator);
+		CacheIsValid = Other.CacheIsValid;
+		return *this;
+	}
+
+	FTagContainerAggregator& operator=(const FTagContainerAggregator& Other)
+	{
+		CapturedActorTags = Other.CapturedActorTags;
+		CapturedSpecTags = Other.CapturedSpecTags;
+		ScopedTags = Other.ScopedTags;
+		CachedAggregator = Other.CachedAggregator;
+		CacheIsValid = Other.CacheIsValid;
+		return *this;
+	}
+
 	FGameplayTagContainer& GetActorTags();
 	const FGameplayTagContainer& GetActorTags() const;
 
 	FGameplayTagContainer& GetSpecTags();
 	const FGameplayTagContainer& GetSpecTags() const;
 
-
-
 	const FGameplayTagContainer* GetAggregatedTags() const;
 
 private:
diff --git a/Engine/Source/Runtime/GameplayTags/Classes/GameplayTagContainer.h b/Engine/Source/Runtime/GameplayTags/Classes/GameplayTagContainer.h
--- a/Engine/Source/Runtime/GameplayTags/Classes/GameplayTagContainer.h
+++ b/Engine/Source/Runtime/GameplayTags/Classes/GameplayTagContainer.h
@@ -108,10 +108,12 @@ struct GAMEPLAYTAGS_API FGameplayTagContainer
 	FGameplayTagContainer();
 	FGameplayTagContainer(FGameplayTagContainer const& Other);
 	FGameplayTagContainer(const FGameplayTag& Tag);
+	FGameplayTagContainer(FGameplayTagContainer&& Other);
 	virtual ~FGameplayTagContainer() {}
 
 	/** Assignment/Equality operators */
 	FGameplayTagContainer& operator=(FGameplayTagContainer const& Other);
+	FGameplayTagContainer & operator=(FGameplayTagContainer&& Other);
 	bool operator==(FGameplayTagContainer const& Other) const;
 	bool operator!=(FGameplayTagContainer const& Other) const;
 
diff --git a/Engine/Source/Runtime/GameplayTags/Private/GameplayTagContainer.cpp b/Engine/Source/Runtime/GameplayTags/Private/GameplayTagContainer.cpp
--- a/Engine/Source/Runtime/GameplayTags/Private/GameplayTagContainer.cpp
+++ b/Engine/Source/Runtime/GameplayTags/Private/GameplayTagContainer.cpp
@@ -17,6 +17,12 @@ FGameplayTagContainer::FGameplayTagContainer(const FGameplayTag& Tag)
 	AddTag(Tag);
 }
 
+FGameplayTagContainer::FGameplayTagContainer(FGameplayTagContainer&& Other)
+	: GameplayTags(MoveTemp(Other.GameplayTags))
+{
+	
+}
+
 FGameplayTagContainer& FGameplayTagContainer::operator=(FGameplayTagContainer const& Other)
 {
 	// Guard against self-assignment
@@ -30,6 +36,12 @@ FGameplayTagContainer& FGameplayTagContainer::operator=(FGameplayTagContainer co
 	return *this;
 }
 
+FGameplayTagContainer& FGameplayTagContainer::operator=(FGameplayTagContainer&& Other)
+{
+	GameplayTags = MoveTemp(Other.GameplayTags);
+	return *this;
+}
+
 bool FGameplayTagContainer::operator==(FGameplayTagContainer const& Other) const
 {
 	if (GameplayTags.Num() != Other.GameplayTags.Num())
