Merging

//depot/UE4-Orion/...

to //depot/UE4/...

- GameplayCueActor and GameplayCueView have been deprecated and removed in favor of GameplayCueNotify_Actor

diff --git a/Engine/Source/Editor/GameplayAbilitiesEditor/Classes/K2Node_GameplayCueEvent.h b/Engine/Source/Editor/GameplayAbilitiesEditor/Classes/K2Node_GameplayCueEvent.h
--- a/Engine/Source/Editor/GameplayAbilitiesEditor/Classes/K2Node_GameplayCueEvent.h
+++ b/Engine/Source/Editor/GameplayAbilitiesEditor/Classes/K2Node_GameplayCueEvent.h
@@ -1,7 +1,6 @@
 // Copyright 1998-2015 Epic Games, Inc. All Rights Reserved.
 
 #pragma once
-#include "GameplayCueView.h"
 #include "K2Node_GameplayCueEvent.generated.h"
 
 
diff --git a/Engine/Source/Runtime/GameplayAbilities/Private/AbilitySystemTestPawn.cpp b/Engine/Source/Runtime/GameplayAbilities/Private/AbilitySystemTestPawn.cpp
--- a/Engine/Source/Runtime/GameplayAbilities/Private/AbilitySystemTestPawn.cpp
+++ b/Engine/Source/Runtime/GameplayAbilities/Private/AbilitySystemTestPawn.cpp
@@ -21,7 +21,6 @@ void AAbilitySystemTestPawn::PostInitializeComponents()
 	static UProperty *DamageProperty = FindFieldChecked<UProperty>(UAbilitySystemTestAttributeSet::StaticClass(), GET_MEMBER_NAME_CHECKED(UAbilitySystemTestAttributeSet, Damage));
 
 	Super::PostInitializeComponents();
-	GameplayCueHandler.Owner = this;
 	AbilitySystemComponent->InitStats(UAbilitySystemTestAttributeSet::StaticClass(), NULL);
 
 	/*
diff --git a/Engine/Source/Runtime/GameplayAbilities/Private/GameplayCueActor.cpp b/Engine/Source/Runtime/GameplayAbilities/Private/GameplayCueActor.cpp
deleted file mode 100644
diff --git a/Engine/Source/Runtime/GameplayAbilities/Private/GameplayCueView.cpp b/Engine/Source/Runtime/GameplayAbilities/Private/GameplayCueView.cpp
deleted file mode 100644
diff --git a/Engine/Source/Runtime/GameplayAbilities/Public/AbilitySystemBlueprintLibrary.h b/Engine/Source/Runtime/GameplayAbilities/Public/AbilitySystemBlueprintLibrary.h
--- a/Engine/Source/Runtime/GameplayAbilities/Public/AbilitySystemBlueprintLibrary.h
+++ b/Engine/Source/Runtime/GameplayAbilities/Public/AbilitySystemBlueprintLibrary.h
@@ -4,7 +4,6 @@
 
 #include "Abilities/GameplayAbility.h"
 #include "Abilities/GameplayAbilityTargetDataFilter.h"
-#include "GameplayCueView.h"
 #include "GameplayCueInterface.h"
 #include "AbilitySystemBlueprintLibrary.generated.h"
 
diff --git a/Engine/Source/Runtime/GameplayAbilities/Public/AbilitySystemTestPawn.h b/Engine/Source/Runtime/GameplayAbilities/Public/AbilitySystemTestPawn.h
--- a/Engine/Source/Runtime/GameplayAbilities/Public/AbilitySystemTestPawn.h
+++ b/Engine/Source/Runtime/GameplayAbilities/Public/AbilitySystemTestPawn.h
@@ -3,7 +3,6 @@
 #pragma once
 #include "GameplayCueInterface.h"
 #include "AbilitySystemInterface.h"
-#include "GameplayCueView.h"
 #include "AbilitySystemTestPawn.generated.h"
 
 UCLASS(Blueprintable, BlueprintType, notplaceable)
@@ -22,9 +21,6 @@ private_subobject:
 	class UAbilitySystemComponent* AbilitySystemComponent;
 public:
 
-	UPROPERTY(EditDefaultsOnly, Category=GameplayEffects)
-	FGameplayCueHandler	GameplayCueHandler;
-
 	//UPROPERTY(EditDefaultsOnly, Category=GameplayEffects)
 	//UGameplayAbilitySet * DefaultAbilitySet;
 
diff --git a/Engine/Source/Runtime/GameplayAbilities/Public/GameplayCueActor.h b/Engine/Source/Runtime/GameplayAbilities/Public/GameplayCueActor.h
deleted file mode 100644
diff --git a/Engine/Source/Runtime/GameplayAbilities/Public/GameplayCueView.h b/Engine/Source/Runtime/GameplayAbilities/Public/GameplayCueView.h
deleted file mode 100644
