Adds blueprint-callable function for getting a float attribute value from an actor.

Merging

//depot/UE4-Orion/Engine/Source/Runtime/GameplayAbilities/...

to //depot/UE4/Engine/Source/Runtime/GameplayAbilities/...

diff --git a/Engine/Source/Runtime/GameplayAbilities/Private/AbilitySystemBlueprintLibrary.cpp b/Engine/Source/Runtime/GameplayAbilities/Private/AbilitySystemBlueprintLibrary.cpp
--- a/Engine/Source/Runtime/GameplayAbilities/Private/AbilitySystemBlueprintLibrary.cpp
+++ b/Engine/Source/Runtime/GameplayAbilities/Private/AbilitySystemBlueprintLibrary.cpp
@@ -35,6 +35,18 @@ void UAbilitySystemBlueprintLibrary::SendGameplayEventToActor(AActor* Actor, FGa
 	}
 }
 
+float UAbilitySystemBlueprintLibrary::GetFloatAttribute(const class AActor* Actor, FGameplayAttribute Attribute)
+{
+	const UAbilitySystemComponent* const AbilitySystem = UAbilitySystemGlobals::GetAbilitySystemComponentFromActor(Actor);
+	if (!AbilitySystem)
+	{
+		return 0.f;
+	}
+
+	const float Result = AbilitySystem->GetNumericAttribute(Attribute);
+	return Result;
+}
+
 FGameplayAbilityTargetDataHandle UAbilitySystemBlueprintLibrary::AppendTargetDataHandle(FGameplayAbilityTargetDataHandle TargetHandle, FGameplayAbilityTargetDataHandle HandleToAdd)
 {
 	TargetHandle.Append(&HandleToAdd);
diff --git a/Engine/Source/Runtime/GameplayAbilities/Private/AbilitySystemComponent.cpp b/Engine/Source/Runtime/GameplayAbilities/Private/AbilitySystemComponent.cpp
--- a/Engine/Source/Runtime/GameplayAbilities/Private/AbilitySystemComponent.cpp
+++ b/Engine/Source/Runtime/GameplayAbilities/Private/AbilitySystemComponent.cpp
@@ -159,7 +159,23 @@ void UAbilitySystemComponent::SetNumericAttribute_Internal(const FGameplayAttrib
 	Attribute.SetNumericValueChecked(NewFloatValue, const_cast<UAttributeSet*>(AttributeSet));
 }
 
-float UAbilitySystemComponent::GetNumericAttribute(const FGameplayAttribute &Attribute)
+float UAbilitySystemComponent::GetNumericAttribute(const FGameplayAttribute &Attribute) const
+{
+	if (Attribute.IsSystemAttribute())
+	{
+		return 0.f;
+	}
+
+	const UAttributeSet* const AttributeSetOrNull = GetAttributeSubobject(Attribute.GetAttributeSetClass());
+	if (AttributeSetOrNull == nullptr)
+	{
+		return 0.f;
+	}
+
+	return Attribute.GetNumericValue(AttributeSetOrNull);
+}
+
+float UAbilitySystemComponent::GetNumericAttributeChecked(const FGameplayAttribute &Attribute) const
 {
 	if(Attribute.IsSystemAttribute())
 	{
diff --git a/Engine/Source/Runtime/GameplayAbilities/Private/AttributeSet.cpp b/Engine/Source/Runtime/GameplayAbilities/Private/AttributeSet.cpp
--- a/Engine/Source/Runtime/GameplayAbilities/Private/AttributeSet.cpp
+++ b/Engine/Source/Runtime/GameplayAbilities/Private/AttributeSet.cpp
@@ -31,6 +31,18 @@ void FGameplayAttribute::SetNumericValueChecked(float NewValue, class UAttribute
 	}
 }
 
+float FGameplayAttribute::GetNumericValue(const UAttributeSet* Src) const
+{
+	const UNumericProperty* const NumericPropertyOrNull = Cast<UNumericProperty>(Attribute);
+	if (!NumericPropertyOrNull)
+	{
+		return 0.f;
+	}
+
+	const void* ValuePtr = NumericPropertyOrNull->ContainerPtrToValuePtr<void>(Src);
+	return NumericPropertyOrNull->GetFloatingPointPropertyValue(ValuePtr);
+}
+
 float FGameplayAttribute::GetNumericValueChecked(const UAttributeSet* Src) const
 {
 	UNumericProperty* NumericProperty = CastChecked<UNumericProperty>(Attribute);
diff --git a/Engine/Source/Runtime/GameplayAbilities/Public/AbilitySystemBlueprintLibrary.h b/Engine/Source/Runtime/GameplayAbilities/Public/AbilitySystemBlueprintLibrary.h
--- a/Engine/Source/Runtime/GameplayAbilities/Public/AbilitySystemBlueprintLibrary.h
+++ b/Engine/Source/Runtime/GameplayAbilities/Public/AbilitySystemBlueprintLibrary.h
@@ -27,6 +27,13 @@ class GAMEPLAYABILITIES_API UAbilitySystemBlueprintLibrary : public UBlueprintFu
 	// found, the event will not be sent.
 	UFUNCTION(BlueprintCallable, Category = Ability, Meta = (Tooltip = "This function can be used to trigger an ability on the actor in question with useful payload data."))
 	static void SendGameplayEventToActor(AActor* Actor, FGameplayTag EventTag, FGameplayEventData Payload);
+	
+	// -------------------------------------------------------------------------------
+	//		Attribute
+	// -------------------------------------------------------------------------------
+
+	UFUNCTION(BlueprintPure, Category = "Ability|Attribute")
+	static float GetFloatAttribute(const class AActor* Actor, FGameplayAttribute Attribute);
 
 	// -------------------------------------------------------------------------------
 	//		TargetData
diff --git a/Engine/Source/Runtime/GameplayAbilities/Public/AbilitySystemComponent.h b/Engine/Source/Runtime/GameplayAbilities/Public/AbilitySystemComponent.h
--- a/Engine/Source/Runtime/GameplayAbilities/Public/AbilitySystemComponent.h
+++ b/Engine/Source/Runtime/GameplayAbilities/Public/AbilitySystemComponent.h
@@ -141,7 +141,8 @@ class GAMEPLAYABILITIES_API UAbilitySystemComponent : public UActorComponent, pu
 	void ApplyModToAttributeUnsafe(const FGameplayAttribute &Attribute, TEnumAsByte<EGameplayModOp::Type> ModifierOp, float ModifierMagnitude);
 
 	/** Returns current (final) value of an attribute */
-	float GetNumericAttribute(const FGameplayAttribute &Attribute);
+	float GetNumericAttribute(const FGameplayAttribute &Attribute) const;
+	float GetNumericAttributeChecked(const FGameplayAttribute &Attribute) const;
 
 	virtual void DisplayDebug(class UCanvas* Canvas, const class FDebugDisplayInfo& DebugDisplay, float& YL, float& YPos);
 
@@ -904,9 +905,7 @@ protected:
 	// ---------------------------------------------
 
 	/** Array of currently active UAbilityTasks that require ticking */
-	TArray<TWeakObjectPtr<UAbilityTask> >	TickingTasks;	
-
-protected:
+	TArray<TWeakObjectPtr<UAbilityTask> >	TickingTasks;
 
 	virtual void OnRegister() override;
 
diff --git a/Engine/Source/Runtime/GameplayAbilities/Public/AttributeSet.h b/Engine/Source/Runtime/GameplayAbilities/Public/AttributeSet.h
--- a/Engine/Source/Runtime/GameplayAbilities/Public/AttributeSet.h
+++ b/Engine/Source/Runtime/GameplayAbilities/Public/AttributeSet.h
@@ -46,6 +46,7 @@ struct GAMEPLAYABILITIES_API FGameplayAttribute
 
 	void SetNumericValueChecked(float NewValue, class UAttributeSet* Dest) const;
 
+	float GetNumericValue(const UAttributeSet* Src) const;
 	float GetNumericValueChecked(const UAttributeSet* Src) const;
 	
 	/** Equality/Inequality operators */
