Merging

//depot/UE4-Orion/Engine/Source/Runtime/GameplayAbilities/...

to //depot/UE4/Engine/Source/Runtime/GameplayAbilities/...

diff --git a/Engine/Source/Runtime/GameplayAbilities/Private/AbilitySystemComponent.cpp b/Engine/Source/Runtime/GameplayAbilities/Private/AbilitySystemComponent.cpp
--- a/Engine/Source/Runtime/GameplayAbilities/Private/AbilitySystemComponent.cpp
+++ b/Engine/Source/Runtime/GameplayAbilities/Private/AbilitySystemComponent.cpp
@@ -416,6 +416,7 @@ void UAbilitySystemComponent::GetOwnedGameplayTags(FGameplayTagContainer& TagCon
 
 bool UAbilitySystemComponent::HasMatchingGameplayTag(FGameplayTag TagToCheck) const
 {
+	SCOPE_CYCLE_COUNTER(STAT_HasMatchingGameplayTag);
 	return GameplayTagCountContainer.HasMatchingGameplayTag(TagToCheck);
 }
 
diff --git a/Engine/Source/Runtime/GameplayAbilities/Private/AbilitySystemStats.cpp b/Engine/Source/Runtime/GameplayAbilities/Private/AbilitySystemStats.cpp
--- a/Engine/Source/Runtime/GameplayAbilities/Private/AbilitySystemStats.cpp
+++ b/Engine/Source/Runtime/GameplayAbilities/Private/AbilitySystemStats.cpp
@@ -20,5 +20,6 @@ DEFINE_STAT(STAT_TickAbilityTasks);
 DEFINE_STAT(STAT_FindAbilitySpecFromHandle);
 DEFINE_STAT(STAT_AggregatorEvaluate);
 DEFINE_STAT(STAT_HasApplicationImmunityToSpec);
+DEFINE_STAT(STAT_HasMatchingGameplayTag);
 
 
diff --git a/Engine/Source/Runtime/GameplayAbilities/Public/AbilitySystemStats.h b/Engine/Source/Runtime/GameplayAbilities/Public/AbilitySystemStats.h
--- a/Engine/Source/Runtime/GameplayAbilities/Public/AbilitySystemStats.h
+++ b/Engine/Source/Runtime/GameplayAbilities/Public/AbilitySystemStats.h
@@ -23,4 +23,5 @@ DECLARE_CYCLE_STAT_EXTERN(TEXT("InitAttributeSetDefaults"), STAT_InitAttributeSe
 DECLARE_CYCLE_STAT_EXTERN(TEXT("TickAbilityTasks"), STAT_TickAbilityTasks, STATGROUP_AbilitySystem, );
 DECLARE_CYCLE_STAT_EXTERN(TEXT("FindAbilitySpecFromHandle"), STAT_FindAbilitySpecFromHandle, STATGROUP_AbilitySystem, );
 DECLARE_CYCLE_STAT_EXTERN(TEXT("Aggregator Evaluate"), STAT_AggregatorEvaluate, STATGROUP_AbilitySystem, );
-DECLARE_CYCLE_STAT_EXTERN(TEXT("Has Application Immunity To Spec"), STAT_HasApplicationImmunityToSpec, STATGROUP_AbilitySystem, );
\ No newline at end of file
+DECLARE_CYCLE_STAT_EXTERN(TEXT("Has Application Immunity To Spec"), STAT_HasApplicationImmunityToSpec, STATGROUP_AbilitySystem, );
+DECLARE_CYCLE_STAT_EXTERN(TEXT("Has Matching GameplayTag"), STAT_HasMatchingGameplayTag, STATGROUP_AbilitySystem, );
\ No newline at end of file
