- Fixed issue removing an active gameplay effect inside of a gameplay tag event delegate.

Merging

//depot/UE4-Orion/Engine/Source/Runtime/GameplayAbilities/Private/GameplayEffect.cpp

to //depot/UE4/Engine/Source/Runtime/GameplayAbilities/Private/GameplayEffect.cpp

diff --git a/Engine/Source/Runtime/GameplayAbilities/Private/GameplayEffect.cpp b/Engine/Source/Runtime/GameplayAbilities/Private/GameplayEffect.cpp
--- a/Engine/Source/Runtime/GameplayAbilities/Private/GameplayEffect.cpp
+++ b/Engine/Source/Runtime/GameplayAbilities/Private/GameplayEffect.cpp
@@ -2194,7 +2194,8 @@ void FActiveGameplayEffectsContainer::AddActiveGameplayEffectGrantedTagsAndModif
 bool FActiveGameplayEffectsContainer::RemoveActiveGameplayEffect(FActiveGameplayEffectHandle Handle, int32 StacksToRemove)
 {
 	// Iterating through manually since this is a removal operation and we need to pass the index into InternalRemoveActiveGameplayEffect
-	for (int32 ActiveGEIdx = 0; ActiveGEIdx < GetNumGameplayEffects(); ++ActiveGEIdx)
+	int32 NumGameplayEffects = GetNumGameplayEffects();
+	for (int32 ActiveGEIdx = 0; ActiveGEIdx < NumGameplayEffects; ++ActiveGEIdx)
 	{
 		FActiveGameplayEffect& Effect = *GetActiveGameplayEffect(ActiveGEIdx);
 		if (Effect.Handle == Handle && Effect.IsPendingRemove == false)
@@ -2222,6 +2223,9 @@ bool FActiveGameplayEffectsContainer::RemoveActiveGameplayEffect(FActiveGameplay
 bool FActiveGameplayEffectsContainer::InternalRemoveActiveGameplayEffect(int32 Idx, int32 StacksToRemove, bool bPrematureRemoval)
 {
 	SCOPE_CYCLE_COUNTER(STAT_RemoveActiveGameplayEffect);
+
+	bool IsLocked = (ScopedLockCount > 0);	// Cache off whether we were previously locked
+	GAMEPLAYEFFECT_SCOPE_LOCK();			// Apply lock so no one else can change the AGE list (we may still change it if IsLocked is false)
 	
 	if (ensure(Idx < GetNumGameplayEffects()))
 	{
@@ -2280,7 +2284,7 @@ bool FActiveGameplayEffectsContainer::InternalRemoveActiveGameplayEffect(int32 I
 		bool ModifiedArray = false;
 
 		// Finally remove the ActiveGameplayEffect
-		if (ScopedLockCount > 0)
+		if (IsLocked)
 		{
 			// We are locked, so this removal is now pending.
 			PendingRemoves++;
