Merging

//depot/UE4-Orion/Engine/Source/...

to //depot/UE4/Engine/Source/...

diff --git a/Engine/Source/Runtime/AIModule/Classes/Perception/AISenseConfig_Damage.h b/Engine/Source/Runtime/AIModule/Classes/Perception/AISenseConfig_Damage.h
--- a/Engine/Source/Runtime/AIModule/Classes/Perception/AISenseConfig_Damage.h
+++ b/Engine/Source/Runtime/AIModule/Classes/Perception/AISenseConfig_Damage.h
@@ -3,6 +3,7 @@
 #pragma once
 
 #include "AISenseConfig.h"
+#include "AISense_Damage.h"
 #include "AISenseConfig_Damage.generated.h"
 
 class UAISense_Damage;
diff --git a/Engine/Source/Runtime/Engine/Classes/Kismet/KismetMathLibrary.h b/Engine/Source/Runtime/Engine/Classes/Kismet/KismetMathLibrary.h
--- a/Engine/Source/Runtime/Engine/Classes/Kismet/KismetMathLibrary.h
+++ b/Engine/Source/Runtime/Engine/Classes/Kismet/KismetMathLibrary.h
@@ -482,6 +482,18 @@ class UKismetMathLibrary : public UBlueprintFunctionLibrary
 	UFUNCTION(BlueprintPure, Category = "Math|Float")
 	static float FInterpEaseInOut(float A, float B, float Alpha, float Exponent);
 
+	/**
+	* Simple function to create a pulsating scalar value
+	*
+	* @param  InCurrentTime  Current absolute time
+	* @param  InPulsesPerSecond  How many full pulses per second?
+	* @param  InPhase  Optional phase amount, between 0.0 and 1.0 (to synchronize pulses)
+	*
+	* @return  Pulsating value (0.0-1.0)
+	*/
+	UFUNCTION(BlueprintPure, Category = "Math|Float")
+	static float MakePulsatingValue(float InCurrentTime, float InPulsesPerSecond = 1.0f, float InPhase = 0.0f);
+
 
 	//
 	// Vector functions.
diff --git a/Engine/Source/Runtime/Engine/Private/KismetMathLibrary.cpp b/Engine/Source/Runtime/Engine/Private/KismetMathLibrary.cpp
--- a/Engine/Source/Runtime/Engine/Private/KismetMathLibrary.cpp
+++ b/Engine/Source/Runtime/Engine/Private/KismetMathLibrary.cpp
@@ -548,6 +548,11 @@ float UKismetMathLibrary::FInterpEaseInOut(float A, float B, float Alpha, float
 	return FMath::InterpEaseInOut<float>(A, B, Alpha, Exponent);
 }
 
+float UKismetMathLibrary::MakePulsatingValue(float InCurrentTime, float InPulsesPerSecond, float InPhase)
+{
+	return FMath::MakePulsatingValue((double)InCurrentTime, InPulsesPerSecond, InPhase);
+}
+
 float UKismetMathLibrary::RandomFloat()
 {
 	return FMath::FRand();
@@ -2030,3 +2035,4 @@ FRandomStream UKismetMathLibrary::MakeRandomStream(int32 InitialSeed)
 {
 	return FRandomStream(InitialSeed);
 }
+
diff --git a/Engine/Source/Runtime/GameplayAbilities/Private/AbilitySystemBlueprintLibrary.cpp b/Engine/Source/Runtime/GameplayAbilities/Private/AbilitySystemBlueprintLibrary.cpp
--- a/Engine/Source/Runtime/GameplayAbilities/Private/AbilitySystemBlueprintLibrary.cpp
+++ b/Engine/Source/Runtime/GameplayAbilities/Private/AbilitySystemBlueprintLibrary.cpp
@@ -434,6 +434,38 @@ bool UAbilitySystemBlueprintLibrary::GetGameplayCueEndLocationAndNormal(AActor*
 	return false;
 }
 
+bool UAbilitySystemBlueprintLibrary::GetGameplayCueDirection(AActor* TargetActor, FGameplayCueParameters Parameters, FVector& Direction)
+{
+	if (FGameplayEffectContext* Ctx = Parameters.EffectContext.Get())
+	{
+		if (Ctx->GetHitResult())
+		{
+			// Most projectiles and melee attacks will use this
+			Direction = (-1.f * Ctx->GetHitResult()->Normal);
+			return true;
+		}
+		else if (TargetActor && Ctx->HasOrigin())
+		{
+			// Fallback to trying to use the target location and the origin of the effect
+			FVector NewVec = (TargetActor->GetActorLocation() - Ctx->GetOrigin());
+			NewVec.Normalize();
+			Direction = NewVec;
+			return true;
+		}
+		else if (TargetActor && Ctx->GetEffectCauser())
+		{
+			// Finally, try to use the direction between the causer of the effect and the target of the effect
+			FVector NewVec = (TargetActor->GetActorLocation() - Ctx->GetEffectCauser()->GetActorLocation());
+			NewVec.Normalize();
+			Direction = NewVec;
+			return true;
+		}
+	}
+
+	Direction = FVector::ZeroVector;
+	return false;
+}
+
 // ---------------------------------------------------------------------------------------
 
 FGameplayEffectSpecHandle UAbilitySystemBlueprintLibrary::AssignSetByCallerMagnitude(FGameplayEffectSpecHandle SpecHandle, FName DataName, float Magnitude)
diff --git a/Engine/Source/Runtime/GameplayAbilities/Private/AbilitySystemComponent.cpp b/Engine/Source/Runtime/GameplayAbilities/Private/AbilitySystemComponent.cpp
--- a/Engine/Source/Runtime/GameplayAbilities/Private/AbilitySystemComponent.cpp
+++ b/Engine/Source/Runtime/GameplayAbilities/Private/AbilitySystemComponent.cpp
@@ -1158,6 +1158,12 @@ void UAbilitySystemComponent::DisplayDebug(class UCanvas* Canvas, const class FD
 	YL = Canvas->DrawText(GEngine->GetTinyFont(), FString::Printf(TEXT("Owned Tags: %s"), *OwnerTags.ToStringSimple()), 4.f, YPos);
 	YPos += YL;
 
+	if (BlockedAbilityTags.GetExplicitGameplayTags().Num() > 0)
+	{
+		YPos += Canvas->DrawText(GEngine->GetTinyFont(), FString::Printf(TEXT("BlockedAbilityTags: %s"), *BlockedAbilityTags.GetExplicitGameplayTags().ToStringSimple()), 4.f, YPos);
+		YPos += YL;
+	}
+
 	TSet<FGameplayAttribute> DrawAttributes;
 
 	const float MaxCharHeight = GEngine->GetTinyFont()->GetMaxCharHeight();
@@ -1327,8 +1333,31 @@ void UAbilitySystemComponent::DisplayDebug(class UCanvas* Canvas, const class FD
 			if (AbilitySpec.Ability == nullptr)
 				continue;
 
-			Canvas->SetDrawColor(AbilitySpec.IsActive() ? FColor::Yellow : FColor(128, 128, 128));
-			YL = Canvas->DrawText(GEngine->GetTinyFont(), FString::Printf(TEXT("%s"), *ASC_CleanupName(GetNameSafe(AbilitySpec.Ability))), 4.f, YPos);
+			FString StatusText;
+			FColor AbilityTextColor = FColor(128, 128, 128);
+			if (AbilitySpec.IsActive())
+			{
+				StatusText = TEXT(" (Active)");
+				AbilityTextColor = FColor::Yellow;
+			}
+			else if (BlockedAbilityBindings.IsValidIndex(AbilitySpec.InputID) && BlockedAbilityBindings[AbilitySpec.InputID])
+			{
+				StatusText = TEXT(" (InputBlocked)");
+				AbilityTextColor = FColor::Red;
+			}
+			else if (AbilitySpec.Ability->AbilityTags.MatchesAny(BlockedAbilityTags.GetExplicitGameplayTags(), false))
+			{
+				StatusText = TEXT(" (TagBlocked)");
+				AbilityTextColor = FColor::Red;
+			}
+			else if (AbilitySpec.Ability->CanActivateAbility(AbilitySpec.Handle, AbilityActorInfo.Get()) == false)
+			{
+				StatusText = TEXT(" (CantActivate)");
+				AbilityTextColor = FColor::Red;
+			}
+			
+			Canvas->SetDrawColor(AbilityTextColor);
+			YL = Canvas->DrawText(GEngine->GetTinyFont(), FString::Printf(TEXT("%s %s"), *ASC_CleanupName(GetNameSafe(AbilitySpec.Ability)), *StatusText), 4.f, YPos);
 			YPos += YL;
 	
 			if (AbilitySpec.IsActive())
diff --git a/Engine/Source/Runtime/GameplayAbilities/Public/AbilitySystemBlueprintLibrary.h b/Engine/Source/Runtime/GameplayAbilities/Public/AbilitySystemBlueprintLibrary.h
--- a/Engine/Source/Runtime/GameplayAbilities/Public/AbilitySystemBlueprintLibrary.h
+++ b/Engine/Source/Runtime/GameplayAbilities/Public/AbilitySystemBlueprintLibrary.h
@@ -162,6 +162,10 @@ class GAMEPLAYABILITIES_API UAbilitySystemBlueprintLibrary : public UBlueprintFu
 	UFUNCTION(BlueprintPure, Category = "Ability|GameplayCue")
 	static bool GetGameplayCueEndLocationAndNormal(AActor* TargetActor, FGameplayCueParameters Parameters, FVector& Location, FVector& Normal);
 
+	/** Gets the best normalized effect direction for this gameplay cue. This is useful for effects that require the direction of an enemy attack. Returns true if a valid direction could be calculated. */
+	UFUNCTION(BlueprintPure, Category = "Ability|GameplayCue")
+	static bool GetGameplayCueDirection(AActor* TargetActor, FGameplayCueParameters Parameters, FVector& Direction);
+
 
 	// -------------------------------------------------------------------------------
 	//		GameplayEffectSpec
