Merging chances from dev branch.

Minor:
-Additional debug info in UAbilitySystemComponent::DisplayDebug
-Ensure AbilitySpec::ACtiveCount does not fall below 0 and end up rolling back to 255.
-Fix case where we were copying AGE by value when we shouldnt be
-Add BroadcastCommitEvent cool to Commitcooldown/cost nodes.

AbilitySystem client->server rpc mini refactor
-TargetData and other Client Events are now tied to a specific ability/predictionkey, rather than being shared acrossed the ASC.
-Reworked the way input is replicated: we are now much more conservative on when we send input state to the server.

StartAbilityState (Ability Task)
-Helper node to make cleanup/handling interrupts easier in blueprints

diff --git a/Engine/Source/Runtime/GameplayAbilities/Private/Abilities/GameplayAbility.cpp b/Engine/Source/Runtime/GameplayAbilities/Private/Abilities/GameplayAbility.cpp
--- a/Engine/Source/Runtime/GameplayAbilities/Private/Abilities/GameplayAbility.cpp
+++ b/Engine/Source/Runtime/GameplayAbilities/Private/Abilities/GameplayAbility.cpp
@@ -52,6 +52,8 @@ UGameplayAbility::UGameplayAbility(const FObjectInitializer& ObjectInitializer)
 #endif
 
 	bServerRespectsRemoteAbilityCancellation = true;
+	bActivateOnInputHeld = false;
+	bReplicateInputDirectly = false;
 }
 
 int32 UGameplayAbility::GetFunctionCallspace(UFunction* Function, void* Parameters, FFrame* Stack)
@@ -119,7 +121,13 @@ bool UGameplayAbility::IsActive() const
 		return bIsActive;
 	}
 
-	// Non-instanced and Instanced-Per-Execution abilities are by definition active unless they are pending kill
+	// this should not be called on NonInstanced warn about it, Should call IsActive on the ability spec instead
+	if (GetInstancingPolicy() == EGameplayAbilityInstancingPolicy::NonInstanced)
+	{
+		ABILITY_LOG(Warning, TEXT("UGameplayAbility::IsActive() called on %s NonInstanced ability, call IsActive on the Ability Spec instead"), *GetName());
+	}
+
+	// NonInstanced and Instanced-Per-Execution abilities are by definition active unless they are pending kill
 	return !IsPendingKill();
 }
 
@@ -320,13 +328,22 @@ void UGameplayAbility::CancelAbility(const FGameplayAbilitySpecHandle Handle, co
 {
 	if (CanBeCanceled())
 	{
+		// Gives the Ability BP a chance to perform custom logic/cleanup when any active ability states are active
+		UAbilitySystemComponent* Comp = ActorInfo->AbilitySystemComponent.Get();
+		if (Comp && Comp->OnAbilityStateInterrupted.IsBound())
+		{
+			Comp->OnAbilityStateInterrupted.Broadcast();
+		}
+
 		EndAbility(Handle, ActorInfo, ActivationInfo, true);
 	}
 }
 
 void UGameplayAbility::EndAbility(const FGameplayAbilitySpecHandle Handle, const FGameplayAbilityActorInfo* ActorInfo, const FGameplayAbilityActivationInfo ActivationInfo, bool bReplicateEndAbility)
-{
-	if (IsActive())
+{ 
+	// check to see if this is an NonInstanced or if the ability is active.
+	FGameplayAbilitySpec* Spec = ActorInfo->AbilitySystemComponent->FindAbilitySpecFromHandle(Handle);
+	if ((Spec != nullptr) ? Spec->IsActive() : IsActive())
 	{
 		// Give blueprint a chance to react
 		K2_OnEndAbility();
@@ -344,8 +361,9 @@ void UGameplayAbility::EndAbility(const FGameplayAbilitySpecHandle Handle, const
 		OnGameplayAbilityEnded.Unbind();
 
 		// Tell all our tasks that we are finished and they should cleanup
-		for (TWeakObjectPtr<UAbilityTask> Task : ActiveTasks)
+		for (int32 TaskIdx = ActiveTasks.Num() - 1; TaskIdx >= 0; --TaskIdx)
 		{
+			TWeakObjectPtr<UAbilityTask> Task = ActiveTasks[TaskIdx];
 			if (Task.IsValid())
 			{
 				Task.Get()->AbilityEnded();
@@ -685,15 +703,23 @@ bool UGameplayAbility::K2_CommitAbility()
 	return CommitAbility(CurrentSpecHandle, CurrentActorInfo, CurrentActivationInfo);
 }
 
-bool UGameplayAbility::K2_CommitAbilityCooldown()
+bool UGameplayAbility::K2_CommitAbilityCooldown(bool BroadcastCommitEvent)
 {
 	check(CurrentActorInfo);
+	if (BroadcastCommitEvent)
+	{
+		CurrentActorInfo->AbilitySystemComponent->NotifyAbilityCommit(this);
+	}
 	return CommitAbilityCooldown(CurrentSpecHandle, CurrentActorInfo, CurrentActivationInfo);
 }
 
-bool UGameplayAbility::K2_CommitAbilityCost()
+bool UGameplayAbility::K2_CommitAbilityCost(bool BroadcastCommitEvent)
 {
 	check(CurrentActorInfo);
+	if (BroadcastCommitEvent)
+	{
+		CurrentActorInfo->AbilitySystemComponent->NotifyAbilityCommit(this);
+	}
 	return CommitAbilityCost(CurrentSpecHandle, CurrentActorInfo, CurrentActivationInfo);
 }
 
@@ -852,6 +878,17 @@ void UGameplayAbility::CancelTaskByInstanceName(FName InstanceName)
 	GetWorld()->GetTimerManager().SetTimerForNextTick(this, &UGameplayAbility::EndOrCancelTasksByInstanceName);
 }
 
+void UGameplayAbility::EndAbilityState(FName OptionalStateNameToEnd)
+{
+	check(CurrentActorInfo);
+
+	UAbilitySystemComponent* ASC = CurrentActorInfo->AbilitySystemComponent.Get();
+	if (ASC && ASC->OnAbilityStateEnded.IsBound())
+	{
+		ASC->OnAbilityStateEnded.Broadcast(OptionalStateNameToEnd);
+	}
+}
+
 void UGameplayAbility::TaskEnded(UAbilityTask* Task)
 {
 	ABILITY_VLOG(CastChecked<AActor>(GetOuter()), Log, TEXT("Task Ended %s"), *Task->GetName());
diff --git a/Engine/Source/Runtime/GameplayAbilities/Private/Abilities/GameplayAbilityTargetActor.cpp b/Engine/Source/Runtime/GameplayAbilities/Private/Abilities/GameplayAbilityTargetActor.cpp
--- a/Engine/Source/Runtime/GameplayAbilities/Private/Abilities/GameplayAbilityTargetActor.cpp
+++ b/Engine/Source/Runtime/GameplayAbilities/Private/Abilities/GameplayAbilityTargetActor.cpp
@@ -46,14 +46,13 @@ void AGameplayAbilityTargetActor::ConfirmTargetingAndContinue()
 	{
 		TargetDataReadyDelegate.Broadcast(FGameplayAbilityTargetDataHandle());
 	}
-	else
-	{
-		NotifyPlayerControllerOfRejectedConfirmation();
-	}
 }
 
 void AGameplayAbilityTargetActor::ConfirmTargeting()
 {
+	UAbilitySystemComponent* ASC = OwningAbility->GetCurrentActorInfo()->AbilitySystemComponent.Get();
+	ASC->AbilityReplicatedEventDelegate(EAbilityReplicatedClientEvent::GenericConfirm, OwningAbility->GetCurrentAbilitySpecHandle(), OwningAbility->GetCurrentActivationInfo().GetActivationPredictionKey() ).Remove(GenericConfirmHandle);
+
 	if (IsConfirmTargetingAllowed())
 	{
 		ConfirmTargetingAndContinue();
@@ -62,29 +61,14 @@ void AGameplayAbilityTargetActor::ConfirmTargeting()
 			Destroy();
 		}
 	}
-	else
-	{
-		NotifyPlayerControllerOfRejectedConfirmation();
-	}
-}
-
-void AGameplayAbilityTargetActor::NotifyPlayerControllerOfRejectedConfirmation()
-{
-	if (OwningAbility && OwningAbility->IsInstantiated() && OwningAbility->GetCurrentActorInfo())
-	{
-		if (UAbilitySystemComponent* ASC = OwningAbility->GetCurrentActorInfo()->AbilitySystemComponent.Get())
-		{
-			if (FGameplayAbilitySpec* AbilitySpec = OwningAbility->GetCurrentAbilitySpec())
-			{
-				ASC->ClientAbilityNotifyRejected(AbilitySpec->InputID);
-			}
-		}
-	}
 }
 
 /** Outside code is saying 'stop everything and just forget about it' */
 void AGameplayAbilityTargetActor::CancelTargeting()
 {
+	UAbilitySystemComponent* ASC = OwningAbility->GetCurrentActorInfo()->AbilitySystemComponent.Get();
+	ASC->AbilityReplicatedEventDelegate(EAbilityReplicatedClientEvent::GenericCancel, OwningAbility->GetCurrentAbilitySpecHandle(), OwningAbility->GetCurrentActivationInfo().GetActivationPredictionKey() ).Remove(GenericCancelHandle);
+
 	CanceledDelegate.Broadcast(FGameplayAbilityTargetDataHandle());
 	Destroy();
 }
@@ -131,7 +115,31 @@ void AGameplayAbilityTargetActor::BindToConfirmCancelInputs()
 	UAbilitySystemComponent* ASC = OwningAbility->GetCurrentActorInfo()->AbilitySystemComponent.Get();
 	if (ASC)
 	{
-		ASC->ConfirmCallbacks.AddDynamic(this, &AGameplayAbilityTargetActor::ConfirmTargeting);
-		ASC->CancelCallbacks.AddDynamic(this, &AGameplayAbilityTargetActor::CancelTargeting);
+		const FGameplayAbilityActorInfo* Info = OwningAbility->GetCurrentActorInfo();
+
+		if (Info->IsLocallyControlled())
+		{
+			// We have to wait for the callback from the AbilitySystemComponent. Which will always be instigated locally
+			ASC->GenericLocalConfirmCallbacks.AddDynamic(this, &AGameplayAbilityTargetActor::ConfirmTargeting);	// Tell me if the confirm input is pressed
+			ASC->GenericLocalCancelCallbacks.AddDynamic(this, &AGameplayAbilityTargetActor::CancelTargeting);	// Tell me if the cancel input is pressed
+		}
+		else
+		{	
+			FGameplayAbilitySpecHandle Handle = OwningAbility->GetCurrentAbilitySpecHandle();
+			FPredictionKey PredKey = OwningAbility->GetCurrentActivationInfo().GetActivationPredictionKey();
+
+			GenericConfirmHandle = ASC->AbilityReplicatedEventDelegate(EAbilityReplicatedClientEvent::GenericConfirm, Handle, PredKey ).AddUObject(this, &AGameplayAbilityTargetActor::ConfirmTargeting);
+			GenericCancelHandle = ASC->AbilityReplicatedEventDelegate(EAbilityReplicatedClientEvent::GenericCancel, Handle, PredKey ).AddUObject(this, &AGameplayAbilityTargetActor::CancelTargeting);
+			
+			if (ASC->CallReplicatedEventDelegateIfSet(EAbilityReplicatedClientEvent::GenericConfirm, Handle, PredKey))
+			{
+				return;
+			}
+			
+			if (ASC->CallReplicatedEventDelegateIfSet(EAbilityReplicatedClientEvent::GenericCancel, Handle, PredKey))
+			{
+				return;
+			}
+		}
 	}
 }
\ No newline at end of file
diff --git a/Engine/Source/Runtime/GameplayAbilities/Private/Abilities/Tasks/AbilityTask.cpp b/Engine/Source/Runtime/GameplayAbilities/Private/Abilities/Tasks/AbilityTask.cpp
--- a/Engine/Source/Runtime/GameplayAbilities/Private/Abilities/Tasks/AbilityTask.cpp
+++ b/Engine/Source/Runtime/GameplayAbilities/Private/Abilities/Tasks/AbilityTask.cpp
@@ -135,3 +135,51 @@ FString UAbilityTask::GetDebugString() const
 {
 	return FString::Printf(TEXT("Generic %s"), *GetName());
 }
+
+FGameplayAbilitySpecHandle UAbilityTask::GetAbilitySpecHandle() const
+{
+	UGameplayAbility* MyAbility = Ability.Get();
+	return MyAbility ? MyAbility->GetCurrentAbilitySpecHandle() : FGameplayAbilitySpecHandle();
+}
+
+FPredictionKey UAbilityTask::GetActivationPredictionKey() const
+{
+	UGameplayAbility* MyAbility = Ability.Get();
+	return MyAbility ? MyAbility->GetCurrentActivationInfo().GetActivationPredictionKey() : FPredictionKey();
+}
+
+bool UAbilityTask::IsPredictingClient() const
+{
+	UGameplayAbility* MyAbility = Ability.Get();
+	if (MyAbility && MyAbility->GetCurrentActorInfo()->OwnerActor.IsValid())
+	{
+		bool bIsLocallyControlled = MyAbility->GetCurrentActorInfo()->IsLocallyControlled();
+		bool bIsAuthority = MyAbility->GetCurrentActorInfo()->IsNetAuthority();
+
+		if (!bIsAuthority && bIsLocallyControlled && MyAbility->GetNetExecutionPolicy() == EGameplayAbilityNetExecutionPolicy::LocalPredicted)
+		{
+			// If this is a predicted ability, forward up
+			return true;
+		}
+	}
+
+	return false;
+}
+
+bool UAbilityTask::IsForRemoteClient() const
+{
+	UGameplayAbility* MyAbility = Ability.Get();
+	if (MyAbility && MyAbility->GetCurrentActorInfo()->OwnerActor.IsValid())
+	{
+		bool bIsLocallyControlled = MyAbility->GetCurrentActorInfo()->IsLocallyControlled();
+		bool bIsAuthority = MyAbility->GetCurrentActorInfo()->IsNetAuthority();
+
+		if (bIsAuthority && !bIsLocallyControlled)
+		{
+			// If this is a predicted ability, forward up
+			return true;
+		}
+	}
+
+	return false;
+}
\ No newline at end of file
diff --git a/Engine/Source/Runtime/GameplayAbilities/Private/Abilities/Tasks/AbilityTask_StartAbilityState.cpp b/Engine/Source/Runtime/GameplayAbilities/Private/Abilities/Tasks/AbilityTask_StartAbilityState.cpp
new file mode 100644
--- /dev/null
+++ b/Engine/Source/Runtime/GameplayAbilities/Private/Abilities/Tasks/AbilityTask_StartAbilityState.cpp
@@ -0,0 +1,82 @@
+// Copyright 1998-2015 Epic Games, Inc. All Rights Reserved.
+
+#include "AbilitySystemPrivatePCH.h"
+#include "Abilities/Tasks/AbilityTask_StartAbilityState.h"
+#include "AbilitySystemComponent.h"
+
+UAbilityTask_StartAbilityState::UAbilityTask_StartAbilityState(const FObjectInitializer& ObjectInitializer)
+: Super(ObjectInitializer)
+{
+	bEndCurrentState = true;
+}
+
+UAbilityTask_StartAbilityState* UAbilityTask_StartAbilityState::StartAbilityState(UObject* WorldContextObject, FName StateName, bool bEndCurrentState)
+{
+	auto Task = NewTask<UAbilityTask_StartAbilityState>(WorldContextObject, StateName);
+	Task->bEndCurrentState = bEndCurrentState;
+	return Task;
+}
+
+void UAbilityTask_StartAbilityState::Activate()
+{
+	UAbilitySystemComponent* ASC = AbilitySystemComponent.Get();
+
+	if (ASC)
+	{
+		if (bEndCurrentState && ASC->OnAbilityStateEnded.IsBound())
+		{
+			ASC->OnAbilityStateEnded.Broadcast(NAME_None);
+		}
+
+		EndStateHandle = ASC->OnAbilityStateEnded.AddUObject(this, &UAbilityTask_StartAbilityState::OnEndState);
+		InterruptStateHandle = ASC->OnAbilityStateInterrupted.AddUObject(this, &UAbilityTask_StartAbilityState::OnInterruptState);
+	}
+}
+
+void UAbilityTask_StartAbilityState::OnDestroy(bool AbilityEnded)
+{	
+	// If the ability was ended, make sure to first notify the BP via 'OnStateEnded'
+	if (AbilityEnded && OnStateEnded.IsBound())
+	{
+		OnStateEnded.Broadcast();
+	}
+
+	UAbilitySystemComponent* ASC = AbilitySystemComponent.Get();
+
+	if (ASC)
+	{
+		ASC->OnAbilityStateInterrupted.Remove(InterruptStateHandle);
+		ASC->OnAbilityStateEnded.Remove(EndStateHandle);
+	}
+
+	Super::OnDestroy(AbilityEnded);
+}
+
+void UAbilityTask_StartAbilityState::OnEndState(FName StateNameToEnd)
+{
+	// All states end if 'Name_NONE' is passed to this delegate
+	if (StateNameToEnd.IsNone() || StateNameToEnd == InstanceName)
+	{
+		EndTask();
+
+		if (OnStateEnded.IsBound())
+		{
+			OnStateEnded.Broadcast();
+		}
+	}
+}
+
+void UAbilityTask_StartAbilityState::OnInterruptState()
+{
+	EndTask();
+
+	if (OnStateInterrupted.IsBound())
+	{
+		OnStateInterrupted.Broadcast();
+	}
+}
+
+FString UAbilityTask_StartAbilityState::GetDebugString() const
+{
+	return FString::Printf(TEXT("%s (AbilityState)"), *InstanceName.ToString());
+}
diff --git a/Engine/Source/Runtime/GameplayAbilities/Private/Abilities/Tasks/AbilityTask_WaitConfirmCancel.cpp b/Engine/Source/Runtime/GameplayAbilities/Private/Abilities/Tasks/AbilityTask_WaitConfirmCancel.cpp
--- a/Engine/Source/Runtime/GameplayAbilities/Private/Abilities/Tasks/AbilityTask_WaitConfirmCancel.cpp
+++ b/Engine/Source/Runtime/GameplayAbilities/Private/Abilities/Tasks/AbilityTask_WaitConfirmCancel.cpp
@@ -19,7 +19,6 @@ void UAbilityTask_WaitConfirmCancel::OnConfirmCallback()
 	OnConfirm.Broadcast();
 	if (AbilitySystemComponent.IsValid())
 	{
-		AbilitySystemComponent->ConsumeAbilityConfirmCancel();
 		EndTask();
 	}
 }
@@ -29,11 +28,28 @@ void UAbilityTask_WaitConfirmCancel::OnCancelCallback()
 	OnCancel.Broadcast();
 	if (AbilitySystemComponent.IsValid())
 	{
-		AbilitySystemComponent->ConsumeAbilityConfirmCancel();
 		EndTask();
 	}
 }
 
+void UAbilityTask_WaitConfirmCancel::OnLocalConfirmCallback()
+{
+	if (AbilitySystemComponent.IsValid() && IsPredictingClient())
+	{
+		AbilitySystemComponent->ServerSetReplicatedClientEvent(EAbilityReplicatedClientEvent::GenericConfirm, GetAbilitySpecHandle(), GetActivationPredictionKey() ,AbilitySystemComponent->ScopedPredictionKey);
+	}
+	OnConfirmCallback();
+}
+
+void UAbilityTask_WaitConfirmCancel::OnLocalCancelCallback()
+{
+	if (AbilitySystemComponent.IsValid() && IsPredictingClient())
+	{
+		AbilitySystemComponent->ServerSetReplicatedClientEvent(EAbilityReplicatedClientEvent::GenericCancel, GetAbilitySpecHandle(), GetActivationPredictionKey() ,AbilitySystemComponent->ScopedPredictionKey);
+	}
+	OnCancelCallback();
+}
+
 UAbilityTask_WaitConfirmCancel* UAbilityTask_WaitConfirmCancel::WaitConfirmCancel(UObject* WorldContextObject)
 {
 	return NewTask<UAbilityTask_WaitConfirmCancel>(WorldContextObject);
@@ -45,31 +61,28 @@ void UAbilityTask_WaitConfirmCancel::Activate()
 	{
 		const FGameplayAbilityActorInfo* Info = Ability.Get()->GetCurrentActorInfo();
 
-		// If not locally controlled (E.g., we are the server and this is an ability running on the client), check to see if they already sent us
-		// the Confirm/Cancel. If so, immediately end this task.
-		if (!Info->IsLocallyControlled())
+		
+		if (Info->IsLocallyControlled())
+		{
+			// We have to wait for the callback from the AbilitySystemComponent.
+			AbilitySystemComponent->GenericLocalConfirmCallbacks.AddDynamic(this, &UAbilityTask_WaitConfirmCancel::OnLocalConfirmCallback);	// Tell me if the confirm input is pressed
+			AbilitySystemComponent->GenericLocalCancelCallbacks.AddDynamic(this, &UAbilityTask_WaitConfirmCancel::OnLocalCancelCallback);	// Tell me if the cancel input is pressed
+
+			RegisteredCallbacks = true;
+		}
+		else
 		{
-			if (AbilitySystemComponent->ReplicatedConfirmAbility)
+			if (CallOrAddReplicatedDelegate(EAbilityReplicatedClientEvent::GenericConfirm,  FSimpleMulticastDelegate::FDelegate::CreateUObject(this, &UAbilityTask_WaitConfirmCancel::OnConfirmCallback)) )
 			{
-				OnConfirmCallback();
+				// GenericConfirm was already received from the client and we just called OnConfirmCallback. The task is done.
 				return;
 			}
-			else if (AbilitySystemComponent->ReplicatedCancelAbility)
+			if (CallOrAddReplicatedDelegate(EAbilityReplicatedClientEvent::GenericCancel,  FSimpleMulticastDelegate::FDelegate::CreateUObject(this, &UAbilityTask_WaitConfirmCancel::OnCancelCallback)) )
 			{
-				OnCancelCallback();
+				// GenericCancel was already received from the client and we just called OnCancelCallback. The task is done.
 				return;
 			}
 		}
-		
-		// We have to wait for the callback from the AbilitySystemComponent. Which may be instigated locally or through network replication
-
-		AbilitySystemComponent->ConfirmCallbacks.AddDynamic(this, &UAbilityTask_WaitConfirmCancel::OnConfirmCallback);	// Tell me if the confirm input is pressed
-		AbilitySystemComponent->CancelCallbacks.AddDynamic(this, &UAbilityTask_WaitConfirmCancel::OnCancelCallback);	// Tell me if the cancel input is pressed
-
-		RegisteredCallbacks = true;
-		
-		// Tell me if something else tells me to 'force' my target data (an ability may allow targeting during animation/limited time)
-		// Tell me if something else tells me to cancel my targeting
 	}
 }
 
@@ -77,8 +90,8 @@ void UAbilityTask_WaitConfirmCancel::OnDestroy(bool AbilityEnding)
 {
 	if (RegisteredCallbacks && AbilitySystemComponent.IsValid())
 	{
-		AbilitySystemComponent->ConfirmCallbacks.RemoveDynamic(this, &UAbilityTask_WaitConfirmCancel::OnConfirmCallback);
-		AbilitySystemComponent->CancelCallbacks.RemoveDynamic(this, &UAbilityTask_WaitConfirmCancel::OnCancelCallback);
+		AbilitySystemComponent->GenericLocalConfirmCallbacks.RemoveDynamic(this, &UAbilityTask_WaitConfirmCancel::OnConfirmCallback);
+		AbilitySystemComponent->GenericLocalCancelCallbacks.RemoveDynamic(this, &UAbilityTask_WaitConfirmCancel::OnCancelCallback);
 	}
 
 	Super::OnDestroy(AbilityEnding);
diff --git a/Engine/Source/Runtime/GameplayAbilities/Private/Abilities/Tasks/AbilityTask_WaitInputPress.cpp b/Engine/Source/Runtime/GameplayAbilities/Private/Abilities/Tasks/AbilityTask_WaitInputPress.cpp
--- a/Engine/Source/Runtime/GameplayAbilities/Private/Abilities/Tasks/AbilityTask_WaitInputPress.cpp
+++ b/Engine/Source/Runtime/GameplayAbilities/Private/Abilities/Tasks/AbilityTask_WaitInputPress.cpp
@@ -8,23 +8,23 @@ UAbilityTask_WaitInputPress::UAbilityTask_WaitInputPress(const FObjectInitialize
 	: Super(ObjectInitializer)
 {
 	StartTime = 0.f;
-	RegisteredCallback= false;
 }
 
-void UAbilityTask_WaitInputPress::OnPressCallback(int32 InputID)
+void UAbilityTask_WaitInputPress::OnPressCallback()
 {
 	float ElapsedTime = GetWorld()->GetTimeSeconds() - StartTime;
 
-	if (!Ability.IsValid())
+	if (!Ability.IsValid() || !AbilitySystemComponent.IsValid())
 	{
 		return;
 	}
-	UGameplayAbility* MyAbility = Ability.Get();
-	check(MyAbility->IsInstantiated());
-	if (MyAbility->GetCurrentAbilitySpec()->InputID != InputID)
+
+	AbilitySystemComponent->AbilityReplicatedEventDelegate(EAbilityReplicatedClientEvent::InputPressed, GetAbilitySpecHandle(), GetActivationPredictionKey()).Remove(DelegateHandle);
+
+	if (IsPredictingClient())
 	{
-		//This key is for some other ability
-		return;
+		// Tell the server about this
+		AbilitySystemComponent->ServerSetReplicatedClientEvent(EAbilityReplicatedClientEvent::InputPressed, GetAbilitySpecHandle(), GetActivationPredictionKey(), AbilitySystemComponent->ScopedPredictionKey);
 	}
 
 	// We are done. Kill us so we don't keep getting broadcast messages
@@ -39,33 +39,18 @@ UAbilityTask_WaitInputPress* UAbilityTask_WaitInputPress::WaitInputPress(class U
 
 void UAbilityTask_WaitInputPress::Activate()
 {
+	StartTime = GetWorld()->GetTimeSeconds();
 	if (Ability.IsValid())
 	{
-		FGameplayAbilitySpec* Spec = Ability->GetCurrentAbilitySpec();
-		if (Spec)
+		DelegateHandle = AbilitySystemComponent->AbilityReplicatedEventDelegate(EAbilityReplicatedClientEvent::InputPressed, GetAbilitySpecHandle(), GetActivationPredictionKey()).AddUObject(this, &UAbilityTask_WaitInputPress::OnPressCallback);
+		if (IsForRemoteClient())
 		{
-			if (Spec->InputPressed)
-			{
-				// Input was already pressed before we got here, we are done.
-				OnPress.Broadcast(0.f);
-				EndTask();
-			}
-			else
-			{
-				StartTime = GetWorld()->GetTimeSeconds();
-				AbilitySystemComponent->AbilityKeyPressCallbacks.AddDynamic(this, &UAbilityTask_WaitInputPress::OnPressCallback);
-				RegisteredCallback = true;
-			}
+			AbilitySystemComponent->CallReplicatedEventDelegateIfSet(EAbilityReplicatedClientEvent::InputPressed, GetAbilitySpecHandle(), GetActivationPredictionKey());
 		}
 	}
 }
 
 void UAbilityTask_WaitInputPress::OnDestroy(bool AbilityEnded)
 {
-	if (RegisteredCallback && Ability.IsValid())
-	{
-		AbilitySystemComponent->AbilityKeyPressCallbacks.RemoveDynamic(this, &UAbilityTask_WaitInputPress::OnPressCallback);
-	}
-
 	Super::OnDestroy(AbilityEnded);
 }
diff --git a/Engine/Source/Runtime/GameplayAbilities/Private/Abilities/Tasks/AbilityTask_WaitInputRelease.cpp b/Engine/Source/Runtime/GameplayAbilities/Private/Abilities/Tasks/AbilityTask_WaitInputRelease.cpp
--- a/Engine/Source/Runtime/GameplayAbilities/Private/Abilities/Tasks/AbilityTask_WaitInputRelease.cpp
+++ b/Engine/Source/Runtime/GameplayAbilities/Private/Abilities/Tasks/AbilityTask_WaitInputRelease.cpp
@@ -8,23 +8,23 @@ UAbilityTask_WaitInputRelease::UAbilityTask_WaitInputRelease(const FObjectInitia
 	: Super(ObjectInitializer)
 {
 	StartTime = 0.f;
-	RegisteredCallback= false;
 }
 
-void UAbilityTask_WaitInputRelease::OnReleaseCallback(int32 InputID)
+void UAbilityTask_WaitInputRelease::OnReleaseCallback()
 {
 	float ElapsedTime = GetWorld()->GetTimeSeconds() - StartTime;
 
-	if (!Ability.IsValid())
+	if (!Ability.IsValid() || !AbilitySystemComponent.IsValid())
 	{
 		return;
 	}
-	UGameplayAbility* MyAbility = Ability.Get();
-	check(MyAbility->IsInstantiated());
-	if (MyAbility->GetCurrentAbilitySpec()->InputID != InputID)
+
+	AbilitySystemComponent->AbilityReplicatedEventDelegate(EAbilityReplicatedClientEvent::InputReleased, GetAbilitySpecHandle(), GetActivationPredictionKey()).Remove(DelegateHandle);
+
+	if (IsPredictingClient())
 	{
-		//This key is for some other ability
-		return;
+		// Tell the server about this
+		AbilitySystemComponent->ServerSetReplicatedClientEvent(EAbilityReplicatedClientEvent::InputReleased, GetAbilitySpecHandle(), GetActivationPredictionKey(), AbilitySystemComponent->ScopedPredictionKey);
 	}
 
 	// We are done. Kill us so we don't keep getting broadcast messages
@@ -39,33 +39,18 @@ UAbilityTask_WaitInputRelease* UAbilityTask_WaitInputRelease::WaitInputRelease(c
 
 void UAbilityTask_WaitInputRelease::Activate()
 {
+	StartTime = GetWorld()->GetTimeSeconds();
 	if (Ability.IsValid())
 	{
-		FGameplayAbilitySpec* Spec = Ability->GetCurrentAbilitySpec();
-		if (Spec)
+		DelegateHandle = AbilitySystemComponent->AbilityReplicatedEventDelegate(EAbilityReplicatedClientEvent::InputReleased, GetAbilitySpecHandle(), GetActivationPredictionKey()).AddUObject(this, &UAbilityTask_WaitInputRelease::OnReleaseCallback);
+		if (IsForRemoteClient())
 		{
-			if (!Spec->InputPressed)
-			{
-				// Input was already released before we got here, we are done.
-				OnRelease.Broadcast(0.f);
-				EndTask();
-			}
-			else
-			{
-				StartTime = GetWorld()->GetTimeSeconds();
-				AbilitySystemComponent->AbilityKeyReleaseCallbacks.AddDynamic(this, &UAbilityTask_WaitInputRelease::OnReleaseCallback);
-				RegisteredCallback = true;
-			}
+			AbilitySystemComponent->CallReplicatedEventDelegateIfSet(EAbilityReplicatedClientEvent::InputReleased, GetAbilitySpecHandle(), GetActivationPredictionKey());
 		}
 	}
 }
 
 void UAbilityTask_WaitInputRelease::OnDestroy(bool AbilityEnded)
 {
-	if (RegisteredCallback && Ability.IsValid())
-	{
-		AbilitySystemComponent->AbilityKeyReleaseCallbacks.RemoveDynamic(this, &UAbilityTask_WaitInputRelease::OnReleaseCallback);
-	}
-
 	Super::OnDestroy(AbilityEnded);
 }
diff --git a/Engine/Source/Runtime/GameplayAbilities/Private/Abilities/Tasks/AbilityTask_WaitTargetData.cpp b/Engine/Source/Runtime/GameplayAbilities/Private/Abilities/Tasks/AbilityTask_WaitTargetData.cpp
--- a/Engine/Source/Runtime/GameplayAbilities/Private/Abilities/Tasks/AbilityTask_WaitTargetData.cpp
+++ b/Engine/Source/Runtime/GameplayAbilities/Private/Abilities/Tasks/AbilityTask_WaitTargetData.cpp
@@ -32,7 +32,7 @@ UAbilityTask_WaitTargetData* UAbilityTask_WaitTargetData::WaitTargetDataUsingAct
 void UAbilityTask_WaitTargetData::Activate()
 {
 	// Need to handle case where target actor was passed into task
-	if (Ability.IsValid() && (TargetClass == NULL))
+	if (Ability.IsValid() && (TargetClass == nullptr))
 	{
 		if (TargetActor.IsValid())
 		{
@@ -82,7 +82,7 @@ bool UAbilityTask_WaitTargetData::BeginSpawningActor(UObject* WorldContextObject
 			if (Class != NULL)
 			{
 				UWorld* World = GEngine->GetWorldFromContextObject(WorldContextObject);
-				SpawnedActor = World->SpawnActorDeferred<AGameplayAbilityTargetActor>(Class, FVector::ZeroVector, FRotator::ZeroRotator, NULL, NULL, true);
+				SpawnedActor = World->SpawnActorDeferred<AGameplayAbilityTargetActor>(Class, FVector::ZeroVector, FRotator::ZeroRotator, nullptr, nullptr, true);
 			}
 
 			if (SpawnedActor)
@@ -163,21 +163,6 @@ void UAbilityTask_WaitTargetData::FinalizeTargetActor(AGameplayAbilityTargetActo
 		}
 		else if (ConfirmationType == EGameplayTargetingConfirmation::UserConfirmed)
 		{
-			// If not locally controlled, check if we already have a confirm/cancel repped to us
-			if (!Ability->GetCurrentActorInfo()->IsLocallyControlled())
-			{
-				if (AbilitySystemComponent->ReplicatedConfirmAbility)
-				{
-					SpawnedActor->ConfirmTargeting();
-					return;
-				}
-				else if (AbilitySystemComponent->ReplicatedCancelAbility)
-				{
-					SpawnedActor->CancelTargeting();
-					return;
-				}
-			}
-
 			// Bind to the Cancel/Confirm Delegates (called from local confirm or from repped confirm)
 			SpawnedActor->BindToConfirmCancelInputs();
 		}
@@ -206,35 +191,25 @@ void UAbilityTask_WaitTargetData::RegisterTargetDataCallbacks()
 		// Register with the TargetData callbacks if we are expecting client to send them
 		if (!bShouldProduceTargetDataOnServer)
 		{
-			// Problem here - if there's targeting data just sitting around, fire events don't get hooked up because of this if-else, even if we don't end the task.
-			if (AbilitySystemComponent->ReplicatedTargetData.IsValid(0))
-			{
-				ValidData.Broadcast(AbilitySystemComponent->ReplicatedTargetData);
-				if (ConfirmationType == EGameplayTargetingConfirmation::CustomMulti)
-				{
-					//Since multifire is supported, we still need to hook up the callbacks
-					OnTargetDataReplicatedCallbackDelegateHandle = AbilitySystemComponent->ReplicatedTargetDataDelegate.AddUObject(this, &UAbilityTask_WaitTargetData::OnTargetDataReplicatedCallback);
-					AbilitySystemComponent->ReplicatedTargetDataCancelledDelegate.AddDynamic(this, &UAbilityTask_WaitTargetData::OnTargetDataReplicatedCancelledCallback);
-				}
-				else
-				{
-					EndTask();
-				}
-			}
-			else
-			{
-				OnTargetDataReplicatedCallbackDelegateHandle = AbilitySystemComponent->ReplicatedTargetDataDelegate.AddUObject(this, &UAbilityTask_WaitTargetData::OnTargetDataReplicatedCallback);
-				AbilitySystemComponent->ReplicatedTargetDataCancelledDelegate.AddDynamic(this, &UAbilityTask_WaitTargetData::OnTargetDataReplicatedCancelledCallback);
-			}
+			FGameplayAbilitySpecHandle	SpecHandle = GetAbilitySpecHandle();
+			FPredictionKey ActivationPredictionKey = GetActivationPredictionKey();
+
+			//Since multifire is supported, we still need to hook up the callbacks
+			AbilitySystemComponent->AbilityTargetDataSetDelegate(SpecHandle, ActivationPredictionKey ).AddUObject(this, &UAbilityTask_WaitTargetData::OnTargetDataReplicatedCallback);
+			AbilitySystemComponent->AbilityTargetDataCancelledDelegate(SpecHandle, ActivationPredictionKey ).AddUObject(this, &UAbilityTask_WaitTargetData::OnTargetDataReplicatedCancelledCallback);
+
+			AbilitySystemComponent->CallReplicatedTargetDataDelegatesIfSet(SpecHandle, ActivationPredictionKey );
 		}
 	}
 }
 
 /** Valid TargetData was replicated to use (we are server, was sent from client) */
-void UAbilityTask_WaitTargetData::OnTargetDataReplicatedCallback(FGameplayAbilityTargetDataHandle Data)
+void UAbilityTask_WaitTargetData::OnTargetDataReplicatedCallback(FGameplayAbilityTargetDataHandle Data, FGameplayTag ActivationTag)
 {
 	check(AbilitySystemComponent.IsValid());
 
+	AbilitySystemComponent->ConsumeClientReplicatedTargetData(GetAbilitySpecHandle(), GetActivationPredictionKey());
+
 	/** 
 	 *  Call into the TargetActor to sanitize/verify the data. If this returns false, we are rejecting
 	 *	the replicated target data and will treat this as a cancel.
@@ -263,9 +238,7 @@ void UAbilityTask_WaitTargetData::OnTargetDataReplicatedCallback(FGameplayAbilit
 void UAbilityTask_WaitTargetData::OnTargetDataReplicatedCancelledCallback()
 {
 	check(AbilitySystemComponent.IsValid());
-
 	Cancelled.Broadcast(FGameplayAbilityTargetDataHandle());
-	
 	EndTask();
 }
 
@@ -276,9 +249,19 @@ void UAbilityTask_WaitTargetData::OnTargetDataReadyCallback(FGameplayAbilityTarg
 
 	FScopedPredictionWindow	ScopedPrediction(AbilitySystemComponent.Get(), ShouldReplicateDataToServer());
 	
-	if (ShouldReplicateDataToServer())
+	const FGameplayAbilityActorInfo* Info = Ability->GetCurrentActorInfo();
+	if (IsPredictingClient())
 	{
-		AbilitySystemComponent->ServerSetReplicatedTargetData(Data, AbilitySystemComponent->ScopedPredictionKey);
+		if (!TargetActor->ShouldProduceTargetDataOnServer)
+		{
+			FGameplayTag ApplicationTag; // Fixme: where would this be useful?
+			AbilitySystemComponent->ServerSetReplicatedTargetData(GetAbilitySpecHandle(), GetActivationPredictionKey(), Data, ApplicationTag, AbilitySystemComponent->ScopedPredictionKey);
+		}
+		else if (ConfirmationType == EGameplayTargetingConfirmation::UserConfirmed)
+		{
+			// We aren't going to send the target data, but we will send a generic confirmed message.
+			AbilitySystemComponent->ServerSetReplicatedClientEvent(EAbilityReplicatedClientEvent::GenericConfirm, GetAbilitySpecHandle(), GetActivationPredictionKey(), AbilitySystemComponent->ScopedPredictionKey);
+		}
 	}
 
 	ValidData.Broadcast(Data);
@@ -294,9 +277,11 @@ void UAbilityTask_WaitTargetData::OnTargetDataCancelledCallback(FGameplayAbility
 {
 	check(AbilitySystemComponent.IsValid());
 
-	AbilitySystemComponent->ServerSetReplicatedTargetDataCancelled();
+	if (IsPredictingClient())
+	{
+		AbilitySystemComponent->ServerSetReplicatedTargetDataCancelled(GetAbilitySpecHandle(), GetActivationPredictionKey(), AbilitySystemComponent->ScopedPredictionKey );
+	}
 	Cancelled.Broadcast(Data);
-
 	EndTask();
 }
 
@@ -325,13 +310,8 @@ void UAbilityTask_WaitTargetData::ExternalCancel()
 
 void UAbilityTask_WaitTargetData::OnDestroy(bool AbilityEnded)
 {
-	AbilitySystemComponent->ConsumeAbilityTargetData();
-	AbilitySystemComponent->ConsumeAbilityConfirmCancel();
 	AbilitySystemComponent->SetUserAbilityActivationInhibited(false);
 
-	AbilitySystemComponent->ReplicatedTargetDataDelegate.Remove(OnTargetDataReplicatedCallbackDelegateHandle);
-	AbilitySystemComponent->ReplicatedTargetDataCancelledDelegate.RemoveDynamic(this, &UAbilityTask_WaitTargetData::OnTargetDataReplicatedCancelledCallback);
-
 	if (TargetActor.IsValid())
 	{
 		TargetActor->Destroy();
diff --git a/Engine/Source/Runtime/GameplayAbilities/Private/AbilitySystemComponent.cpp b/Engine/Source/Runtime/GameplayAbilities/Private/AbilitySystemComponent.cpp
--- a/Engine/Source/Runtime/GameplayAbilities/Private/AbilitySystemComponent.cpp
+++ b/Engine/Source/Runtime/GameplayAbilities/Private/AbilitySystemComponent.cpp
@@ -1337,7 +1337,7 @@ void UAbilitySystemComponent::DisplayDebug(class UCanvas* Canvas, const class FD
 			FColor AbilityTextColor = FColor(128, 128, 128);
 			if (AbilitySpec.IsActive())
 			{
-				StatusText = TEXT(" (Active)");
+				StatusText = FString::Printf(TEXT(" (Active %d)"), AbilitySpec.ActiveCount);
 				AbilityTextColor = FColor::Yellow;
 			}
 			else if (BlockedAbilityBindings.IsValidIndex(AbilitySpec.InputID) && BlockedAbilityBindings[AbilitySpec.InputID])
@@ -1355,9 +1355,11 @@ void UAbilitySystemComponent::DisplayDebug(class UCanvas* Canvas, const class FD
 				StatusText = TEXT(" (CantActivate)");
 				AbilityTextColor = FColor::Red;
 			}
-			
+
+			FString InputPressedStr = AbilitySpec.InputPressed ? TEXT("(InputPressed)") : TEXT("");
+
 			Canvas->SetDrawColor(AbilityTextColor);
-			YL = Canvas->DrawText(GEngine->GetTinyFont(), FString::Printf(TEXT("%s %s"), *ASC_CleanupName(GetNameSafe(AbilitySpec.Ability)), *StatusText), 4.f, YPos);
+			YL = Canvas->DrawText(GEngine->GetTinyFont(), FString::Printf(TEXT("%s %s %s"), *ASC_CleanupName(GetNameSafe(AbilitySpec.Ability)), *StatusText, *InputPressedStr), 4.f, YPos);
 			YPos += YL;
 	
 			if (AbilitySpec.IsActive())
diff --git a/Engine/Source/Runtime/GameplayAbilities/Private/AbilitySystemComponent_Abilities.cpp b/Engine/Source/Runtime/GameplayAbilities/Private/AbilitySystemComponent_Abilities.cpp
--- a/Engine/Source/Runtime/GameplayAbilities/Private/AbilitySystemComponent_Abilities.cpp
+++ b/Engine/Source/Runtime/GameplayAbilities/Private/AbilitySystemComponent_Abilities.cpp
@@ -498,7 +498,15 @@ void UAbilitySystemComponent::NotifyAbilityEnded(FGameplayAbilitySpecHandle Hand
 		ClearAnimatingAbility(Ability);
 	}
 
-	Spec->ActiveCount--;
+	// check to make sure we do not cause a roll over to uint8 by decrementing when it is 0
+	if (ensure(Spec->ActiveCount > 0))
+	{
+		Spec->ActiveCount--;
+	}
+	else
+	{
+		ABILITY_LOG(Warning, TEXT("NotifyAbilityEnded called when the Spec->ActiveCount <= 0"));
+	}
 	
 	/** If this is instanced per execution, mark pending kill and remove it from our instanced lists if we are the authority */
 	if (Ability->GetInstancingPolicy() == EGameplayAbilityInstancingPolicy::InstancedPerExecution)
@@ -523,11 +531,22 @@ void UAbilitySystemComponent::NotifyAbilityEnded(FGameplayAbilitySpecHandle Hand
 	MarkAbilitySpecDirty(*Spec);
 }
 
+void UAbilitySystemComponent::CancelAbility(UGameplayAbility* Ability)
+{
+	for (FGameplayAbilitySpec& Spec : ActivatableAbilities.Items)
+	{
+		if (Spec.Ability == Ability)
+		{
+			CancelAbilitySpec(Spec, nullptr);
+		}
+	}
+}
+
 void UAbilitySystemComponent::CancelAbilities(const FGameplayTagContainer* WithTags, const FGameplayTagContainer* WithoutTags, UGameplayAbility* Ignore)
 {
 	for (FGameplayAbilitySpec& Spec : ActivatableAbilities.Items)
 	{
-		if (!Spec.Ability)
+		if (!Spec.IsActive() || Spec.Ability == nullptr)
 		{
 			continue;
 		}
@@ -535,7 +554,7 @@ void UAbilitySystemComponent::CancelAbilities(const FGameplayTagContainer* WithT
 		bool WithTagPass = (!WithTags || Spec.Ability->AbilityTags.MatchesAny(*WithTags, false));
 		bool WithoutTagPass = (!WithoutTags || !Spec.Ability->AbilityTags.MatchesAny(*WithoutTags, false));
 
-		if (Spec.IsActive() && Spec.Ability && WithTagPass && WithoutTagPass)
+		if (WithTagPass && WithoutTagPass)
 		{
 			CancelAbilitySpec(Spec, Ignore);
 		}
@@ -566,6 +585,17 @@ void UAbilitySystemComponent::CancelAbilitySpec(FGameplayAbilitySpec& Spec, UGam
 	MarkAbilitySpecDirty(Spec);
 }
 
+void UAbilitySystemComponent::CancelAllAbilities(UGameplayAbility* Ignore)
+{
+	for (FGameplayAbilitySpec& Spec : ActivatableAbilities.Items)
+	{
+		if (Spec.Ability && Spec.Ability->IsActive())
+		{
+			CancelAbilitySpec(Spec, Ignore);
+		}
+	}
+}
+
 void UAbilitySystemComponent::ApplyAbilityBlockAndCancelTags(const FGameplayTagContainer& AbilityTags, UGameplayAbility* RequestingAbility, bool bEnableBlockTags, const FGameplayTagContainer& BlockTags, bool bExecuteCancelTags, const FGameplayTagContainer& CancelTags)
 {
 	if (bEnableBlockTags)
@@ -793,7 +823,15 @@ bool UAbilitySystemComponent::TryActivateAbility(FGameplayAbilitySpecHandle Hand
 		return false;
 	}
 
-	Spec->ActiveCount++;
+	// make sure we do not incur a roll over if we go over the uint8 max, this will need to be updated if the var size changes
+	if (ensure(Spec->ActiveCount < UINT8_MAX))
+	{
+		Spec->ActiveCount++;
+	}
+	else
+	{
+		ABILITY_LOG(Warning, TEXT("TryActivateAbility called when the Spec->ActiveCount >= UINT8_MAX"));
+	}
 
 	// Setup a fresh ActivationInfo for this AbilitySpec.
 	Spec->ActivationInfo = FGameplayAbilityActivationInfo(ActorInfo->OwnerActor.Get());
@@ -838,11 +876,11 @@ bool UAbilitySystemComponent::TryActivateAbility(FGameplayAbilitySpecHandle Hand
 		// Tell the server to start the ability, this will return to the client 
 		if (TriggerEventData)
 		{
-			ServerTryActivateAbilityWithEventData(Handle, Spec->InputPressed, FPredictionKey(), *TriggerEventData);
+			ServerTryActivateAbilityWithEventData(Handle, FPredictionKey(), *TriggerEventData);
 		}
 		else
 		{
-			ServerTryActivateAbility(Handle, Spec->InputPressed, FPredictionKey());
+			ServerTryActivateAbility(Handle, FPredictionKey());
 		}
 	}
 	else if (Ability->GetNetExecutionPolicy() == EGameplayAbilityNetExecutionPolicy::LocalPredicted)
@@ -855,11 +893,11 @@ bool UAbilitySystemComponent::TryActivateAbility(FGameplayAbilitySpecHandle Hand
 		// This must be called immediately after GeneratePredictionKey to prevent problems with recursively activating abilities
 		if (TriggerEventData)
 		{
-			ServerTryActivateAbilityWithEventData(Handle, Spec->InputPressed, ScopedPredictionKey, *TriggerEventData);
+			ServerTryActivateAbilityWithEventData(Handle, ScopedPredictionKey, *TriggerEventData);
 		}
 		else
 		{
-			ServerTryActivateAbility(Handle, Spec->InputPressed, ScopedPredictionKey);
+			ServerTryActivateAbility(Handle, ScopedPredictionKey);
 		}
 
 		// If this PredictionKey is rejected, we will call OnClientActivateAbilityFailed.
@@ -906,27 +944,27 @@ bool UAbilitySystemComponent::TryActivateAbility(FGameplayAbilitySpecHandle Hand
 	return true;
 }
 
-void UAbilitySystemComponent::ServerTryActivateAbility_Implementation(FGameplayAbilitySpecHandle Handle, bool InputPressed, FPredictionKey PredictionKey)
+void UAbilitySystemComponent::ServerTryActivateAbility_Implementation(FGameplayAbilitySpecHandle Handle, FPredictionKey PredictionKey)
 {
-	InternalServerTryActiveAbility(Handle, InputPressed, PredictionKey, nullptr);
+	InternalServerTryActiveAbility(Handle, PredictionKey, nullptr);
 }
 
-bool UAbilitySystemComponent::ServerTryActivateAbility_Validate(FGameplayAbilitySpecHandle Handle, bool InputPressed, FPredictionKey PredictionKey)
+bool UAbilitySystemComponent::ServerTryActivateAbility_Validate(FGameplayAbilitySpecHandle Handle, FPredictionKey PredictionKey)
 {
 	return true;
 }
 
-void UAbilitySystemComponent::ServerTryActivateAbilityWithEventData_Implementation(FGameplayAbilitySpecHandle Handle, bool InputPressed, FPredictionKey PredictionKey, FGameplayEventData TriggerEventData)
+void UAbilitySystemComponent::ServerTryActivateAbilityWithEventData_Implementation(FGameplayAbilitySpecHandle Handle, FPredictionKey PredictionKey, FGameplayEventData TriggerEventData)
 {
-	InternalServerTryActiveAbility(Handle, InputPressed, PredictionKey, &TriggerEventData);
+	InternalServerTryActiveAbility(Handle, PredictionKey, &TriggerEventData);
 }
 
-bool UAbilitySystemComponent::ServerTryActivateAbilityWithEventData_Validate(FGameplayAbilitySpecHandle Handle, bool InputPressed, FPredictionKey PredictionKey, FGameplayEventData TriggerEventData)
+bool UAbilitySystemComponent::ServerTryActivateAbilityWithEventData_Validate(FGameplayAbilitySpecHandle Handle, FPredictionKey PredictionKey, FGameplayEventData TriggerEventData)
 {
 	return true;
 }
 
-void UAbilitySystemComponent::InternalServerTryActiveAbility(FGameplayAbilitySpecHandle Handle, bool InputPressed, const FPredictionKey& PredictionKey, const FGameplayEventData* TriggerEventData)
+void UAbilitySystemComponent::InternalServerTryActiveAbility(FGameplayAbilitySpecHandle Handle, const FPredictionKey& PredictionKey, const FGameplayEventData* TriggerEventData)
 {
 #if !(UE_BUILD_SHIPPING || UE_BUILD_TEST)
 	if (DenyClientActivation > 0)
@@ -944,6 +982,9 @@ void UAbilitySystemComponent::InternalServerTryActiveAbility(FGameplayAbilitySpe
 		return;
 	}
 
+	// Consume any pending target info, to clear out cancels from old executions
+	ConsumeAllClientReplicatedData(Handle, PredictionKey);
+
 	FScopedPredictionWindow ScopedPredictionWindow(this, PredictionKey);
 
 	const UGameplayAbility* AbilityToActivate = Spec->Ability;
@@ -996,8 +1037,8 @@ void UAbilitySystemComponent::InternalServerTryActiveAbility(FGameplayAbilitySpe
 		return;
 	}
 
-	UGameplayAbility* InstancedAbility = NULL;
-	Spec->InputPressed = InputPressed; // Tasks that check input may need this to be right immediately on ability startup.
+	UGameplayAbility* InstancedAbility = nullptr;
+	Spec->InputPressed = true;
 
 	// Attempt to activate the ability (server side) and tell the client if it succeeded or failed.
 	if (TryActivateAbility(Handle, PredictionKey, &InstancedAbility, nullptr, TriggerEventData))
@@ -1007,6 +1048,7 @@ void UAbilitySystemComponent::InternalServerTryActiveAbility(FGameplayAbilitySpe
 	else
 	{
 		ClientActivateAbilityFailed(Handle, PredictionKey.Current);
+		Spec->InputPressed = false;
 	}
 	MarkAbilitySpecDirty(*Spec);
 }
@@ -1174,11 +1216,6 @@ void UAbilitySystemComponent::ClientActivateAbilitySucceed_Implementation(FGamep
 	}
 }
 
-void UAbilitySystemComponent::ClientAbilityNotifyRejected_Implementation(int32 InputID)
-{
-	TargetingRejectedConfirmationDelegate.Broadcast(InputID);
-}
-
 bool UAbilitySystemComponent::TriggerAbilityFromGameplayEvent(FGameplayAbilitySpecHandle Handle, FGameplayAbilityActorInfo* ActorInfo, FGameplayTag EventTag, const FGameplayEventData* Payload, UAbilitySystemComponent& Component)
 {
 	FGameplayAbilitySpec* Spec = FindAbilitySpecFromHandle(Handle);
@@ -1388,14 +1425,14 @@ void UAbilitySystemComponent::BindToInputComponent(UInputComponent* InputCompone
 	// Pressed event
 	{
 		FInputActionBinding AB(ConfirmBindName, IE_Pressed);
-		AB.ActionDelegate.GetDelegateForManualSet().BindUObject(this, &UAbilitySystemComponent::InputConfirm);
+		AB.ActionDelegate.GetDelegateForManualSet().BindUObject(this, &UAbilitySystemComponent::LocalInputConfirm);
 		InputComponent->AddActionBinding(AB);
 	}
 
 	// 
 	{
 		FInputActionBinding AB(CancelBindName, IE_Pressed);
-		AB.ActionDelegate.GetDelegateForManualSet().BindUObject(this, &UAbilitySystemComponent::InputCancel);
+		AB.ActionDelegate.GetDelegateForManualSet().BindUObject(this, &UAbilitySystemComponent::LocalInputCancel);
 		InputComponent->AddActionBinding(AB);
 	}
 }
@@ -1415,14 +1452,14 @@ void UAbilitySystemComponent::BindAbilityActivationToInputComponent(UInputCompon
 		// Pressed event
 		{
 			FInputActionBinding AB(FName(*BindStr), IE_Pressed);
-			AB.ActionDelegate.GetDelegateForManualSet().BindUObject(this, &UAbilitySystemComponent::AbilityInputPressed, idx);
+			AB.ActionDelegate.GetDelegateForManualSet().BindUObject(this, &UAbilitySystemComponent::AbilityLocalInputPressed, idx);
 			InputComponent->AddActionBinding(AB);
 		}
 
 		// Released event
 		{
 			FInputActionBinding AB(FName(*BindStr), IE_Released);
-			AB.ActionDelegate.GetDelegateForManualSet().BindUObject(this, &UAbilitySystemComponent::AbilityInputReleased, idx);
+			AB.ActionDelegate.GetDelegateForManualSet().BindUObject(this, &UAbilitySystemComponent::AbilityLocalInputReleased, idx);
 			InputComponent->AddActionBinding(AB);
 		}
 	}
@@ -1430,21 +1467,20 @@ void UAbilitySystemComponent::BindAbilityActivationToInputComponent(UInputCompon
 	// Bind Confirm/Cancel. Note: these have to come last!
 	{
 		FInputActionBinding AB(FName(*BindInfo.ConfirmTargetCommand), IE_Pressed);
-		AB.ActionDelegate.GetDelegateForManualSet().BindUObject(this, &UAbilitySystemComponent::InputConfirm);
+		AB.ActionDelegate.GetDelegateForManualSet().BindUObject(this, &UAbilitySystemComponent::LocalInputConfirm);
 		InputComponent->AddActionBinding(AB);
 	}
 	
 	{
 		FInputActionBinding AB(FName(*BindInfo.CancelTargetCommand), IE_Pressed);
-		AB.ActionDelegate.GetDelegateForManualSet().BindUObject(this, &UAbilitySystemComponent::InputCancel);
+		AB.ActionDelegate.GetDelegateForManualSet().BindUObject(this, &UAbilitySystemComponent::LocalInputCancel);
 		InputComponent->AddActionBinding(AB);
 	}
 
 }
 
-void UAbilitySystemComponent::AbilityInputPressed(int32 InputID)
+void UAbilitySystemComponent::AbilityLocalInputPressed(int32 InputID)
 {
-	// FIXME: Maps or multicast delegate to actually do this
 	for (FGameplayAbilitySpec& Spec : ActivatableAbilities.Items)
 	{
 		if (Spec.InputID == InputID)
@@ -1454,28 +1490,15 @@ void UAbilitySystemComponent::AbilityInputPressed(int32 InputID)
 				Spec.InputPressed = true;
 				if (Spec.IsActive())
 				{
-					// The ability is active, so just pipe the input event to it
-					if (Spec.Ability->GetInstancingPolicy() == EGameplayAbilityInstancingPolicy::NonInstanced)
+					if (Spec.Ability->bReplicateInputDirectly && IsOwnerActorAuthoritative() == false)
 					{
-						Spec.Ability->InputPressed(Spec.Handle, AbilityActorInfo.Get(), Spec.ActivationInfo);
-					}
-					else
-					{
-						TArray<UGameplayAbility*> Instances = Spec.GetAbilityInstances();
-						for (UGameplayAbility* Instance : Instances)
-						{
-							Instance->InputPressed(Spec.Handle, AbilityActorInfo.Get(), Spec.ActivationInfo);
-						}						
+						ServerSetInputPressed(Spec.Handle);
 					}
 
-					FScopedPredictionWindow ScopedPrediction(this, true);
-					//We don't require (AbilityKeyPressCallbacks.IsBound() || AbilityKeyReleaseCallbacks.IsBound())), because we don't necessarily know now if we'll need this data later.
-					if (GetOwnerRole() != ROLE_Authority)
-					{
-						// Tell the server we pressed input.
-						ServerSetReplicatedAbilityKeyState(InputID, true, ScopedPredictionKey);
-					}
-					AbilityKeyPressCallbacks.Broadcast(InputID);
+					AbilitySpecInputPressed(Spec);
+
+					// Invoke the InputPressed event. This is not replicated here. If someone is listening, they may replicate the InputPressed event to the server.
+					InvokeReplicatedClientEvent(EAbilityReplicatedClientEvent::InputPressed, Spec.Handle, Spec.ActivationInfo.GetActivationPredictionKey());					
 				}
 				else
 				{
@@ -1490,30 +1513,57 @@ void UAbilitySystemComponent::AbilityInputPressed(int32 InputID)
 	}
 }
 
-void UAbilitySystemComponent::AbilityInputReleased(int32 InputID)
+void UAbilitySystemComponent::AbilityLocalInputReleased(int32 InputID)
 {
-	// FIXME: Maps or multicast delegate to actually do this
 	for (FGameplayAbilitySpec& Spec : ActivatableAbilities.Items)
 	{
 		if (Spec.InputID == InputID)
 		{
-			if (Spec.Ability)
+			Spec.InputPressed = false;
+			if (Spec.Ability && Spec.IsActive())
 			{
-				AbilitySpecInputReleased(Spec);
-				FPredictionKey NewKey = FPredictionKey::CreateNewPredictionKey(this);
-				FScopedPredictionWindow ScopedPrediction(this, true);
-				//We don't require (AbilityKeyPressCallbacks.IsBound() || AbilityKeyReleaseCallbacks.IsBound())), because we don't necessarily know now if we'll need this data later.
-				if (GetOwnerRole() != ROLE_Authority)
+				if (Spec.Ability->bReplicateInputDirectly && IsOwnerActorAuthoritative() == false)
 				{
-					// Tell the server we released input.
-					ServerSetReplicatedAbilityKeyState(InputID, false, ScopedPredictionKey);
+					ServerSetInputReleased(Spec.Handle);
 				}
-				AbilityKeyReleaseCallbacks.Broadcast(InputID);
+
+				AbilitySpecInputReleased(Spec);
+				
+				InvokeReplicatedClientEvent(EAbilityReplicatedClientEvent::InputReleased, Spec.Handle, Spec.ActivationInfo.GetActivationPredictionKey());
 			}
 		}
 	}
 }
 
+void UAbilitySystemComponent::ServerSetInputPressed_Implementation(FGameplayAbilitySpecHandle AbilityHandle)
+{
+	FGameplayAbilitySpec* Spec = FindAbilitySpecFromHandle(AbilityHandle);
+	if (Spec)
+	{
+		AbilitySpecInputPressed(*Spec);
+	}
+
+}
+
+void UAbilitySystemComponent::ServerSetInputReleased_Implementation(FGameplayAbilitySpecHandle AbilityHandle)
+{
+	FGameplayAbilitySpec* Spec = FindAbilitySpecFromHandle(AbilityHandle);
+	if (Spec)
+	{
+		AbilitySpecInputReleased(*Spec);
+	}
+}
+
+bool UAbilitySystemComponent::ServerSetInputPressed_Validate(FGameplayAbilitySpecHandle AbilityHandle)
+{
+	return true;
+}
+
+bool UAbilitySystemComponent::ServerSetInputReleased_Validate(FGameplayAbilitySpecHandle AbilityHandle)
+{
+	return true;
+}
+
 void UAbilitySystemComponent::AbilitySpecInputPressed(FGameplayAbilitySpec& Spec)
 {
 	Spec.InputPressed = true;
@@ -1577,62 +1627,14 @@ void UAbilitySystemComponent::TryActivateInputHeldAbilities()
 	}
 }
 
-void UAbilitySystemComponent::InputConfirm()
+void UAbilitySystemComponent::LocalInputConfirm()
 {
-	FScopedPredictionWindow ScopedPrediction(this, true);
-
-	if (GetOwnerRole() != ROLE_Authority && ConfirmCallbacks.IsBound())
-	{
-		// Tell the server we confirmed input.
-		ServerSetReplicatedConfirm(true, ScopedPredictionKey);
-	}
-	
-	ConfirmCallbacks.Broadcast();
+	GenericLocalConfirmCallbacks.Broadcast();
 }
 
-void UAbilitySystemComponent::InputCancel()
-{
-	FScopedPredictionWindow ScopedPrediction(this, true);
-
-	if (GetOwnerRole() != ROLE_Authority && CancelCallbacks.IsBound())
-	{
-		// Tell the server we confirmed input.
-		ServerSetReplicatedConfirm(false, ScopedPredictionKey);
-	}
-
-	CancelCallbacks.Broadcast();
-}
-
-bool UAbilitySystemComponent::ServerInputPress_Validate(FGameplayAbilitySpecHandle Handle, FPredictionKey ScopedPedictionKey)
-{
-	return true;
-}
-
-void UAbilitySystemComponent::ServerInputPress_Implementation(FGameplayAbilitySpecHandle Handle, FPredictionKey ScopedPedictionKey)
-{
-	FScopedPredictionWindow ScopedPrediction(this, ScopedPedictionKey);
-
-	FGameplayAbilitySpec* Spec = FindAbilitySpecFromHandle(Handle);
-	if (Spec)
-	{
-		AbilitySpecInputPressed(*Spec);
-	}
-}
-
-bool UAbilitySystemComponent::ServerInputRelease_Validate(FGameplayAbilitySpecHandle Handle, FPredictionKey ScopedPedictionKey)
-{
-	return true;
-}
-
-void UAbilitySystemComponent::ServerInputRelease_Implementation(FGameplayAbilitySpecHandle Handle, FPredictionKey ScopedPedictionKey)
-{
-	FScopedPredictionWindow ScopedPrediction(this, ScopedPedictionKey);
-
-	FGameplayAbilitySpec* Spec = FindAbilitySpecFromHandle(Handle);
-	if (Spec)
-	{
-		AbilitySpecInputReleased(*Spec);
-	}
+void UAbilitySystemComponent::LocalInputCancel()
+{	
+	GenericLocalCancelCallbacks.Broadcast();
 }
 
 void UAbilitySystemComponent::TargetConfirm()
@@ -1653,7 +1655,6 @@ void UAbilitySystemComponent::TargetConfirm()
 			}
 			else
 			{
-				TargetActor->NotifyPlayerControllerOfRejectedConfirmation();
 				LeftoverTargetActors.Add(TargetActor);
 			}
 		}
@@ -1676,111 +1677,6 @@ void UAbilitySystemComponent::TargetCancel()
 
 // --------------------------------------------------------------------------
 
-void UAbilitySystemComponent::ServerSetReplicatedAbilityKeyState_Implementation(int32 InputID, bool Pressed, FPredictionKey PredictionKey)
-{
-	FScopedPredictionWindow ScopedPrediction(this, PredictionKey);
-
-	for (FGameplayAbilitySpec& Spec : ActivatableAbilities.Items)
-	{
-		if (Spec.InputID == InputID)
-		{
-			if (Spec.Ability)
-			{
-				Spec.InputPressed = Pressed;
-
-				// Call InputPressed() / InputReleased() on the affected abilities
-				if (Pressed)
-				{
-					AbilitySpecInputPressed(Spec);
-				}
-				else
-				{
-					AbilitySpecInputReleased(Spec);
-				}
-			}
-		}
-	}
-
-	if (Pressed)
-	{
-		AbilityKeyPressCallbacks.Broadcast(InputID);
-	}
-	else
-	{
-		AbilityKeyReleaseCallbacks.Broadcast(InputID);
-	}
-} 
-
-bool UAbilitySystemComponent::ServerSetReplicatedAbilityKeyState_Validate(int32 InputID, bool Pressed, FPredictionKey PredictionKey)
-{
-	return true;
-}
-
-// --------------------------------------------------------------------------
-
-void UAbilitySystemComponent::ServerSetReplicatedConfirm_Implementation(bool Confirmed, FPredictionKey PredictionKey)
-{
-	FScopedPredictionWindow ScopedPrediction(this, PredictionKey);
-
-	if (Confirmed)
-	{
-		ReplicatedConfirmAbility = true;
-		ConfirmCallbacks.Broadcast();
-	}
-	else
-	{
-		ReplicatedCancelAbility = true;
-		CancelCallbacks.Broadcast();
-	}
-}
-
-bool UAbilitySystemComponent::ServerSetReplicatedConfirm_Validate(bool Confirmed, FPredictionKey PredictionKey)
-{
-	return true;
-}
-
-// -------
-
-void UAbilitySystemComponent::ServerSetReplicatedTargetData_Implementation(FGameplayAbilityTargetDataHandle Confirmed, FPredictionKey PredictionKey)
-{
-	FScopedPredictionWindow ScopedPrediction(this, PredictionKey);
-
-	ReplicatedTargetData = Confirmed;
-	ReplicatedTargetDataDelegate.Broadcast(ReplicatedTargetData);
-}
-
-bool UAbilitySystemComponent::ServerSetReplicatedTargetData_Validate(FGameplayAbilityTargetDataHandle Confirmed, FPredictionKey PredictionKey)
-{
-	return true;
-}
-
-// -------
-
-void UAbilitySystemComponent::ServerSetReplicatedTargetDataCancelled_Implementation()
-{
-	ReplicatedTargetDataCancelledDelegate.Broadcast();
-}
-
-bool UAbilitySystemComponent::ServerSetReplicatedTargetDataCancelled_Validate()
-{
-	return true;
-}
-
-// -------
-
-void UAbilitySystemComponent::ConsumeAbilityConfirmCancel()
-{
-	ReplicatedConfirmAbility = false;
-	ReplicatedCancelAbility = false;
-	ConfirmCallbacks.Clear();
-	CancelCallbacks.Clear();
-}
-
-void UAbilitySystemComponent::ConsumeAbilityTargetData()
-{
-	ReplicatedTargetData.Clear();
-}
-
 void UAbilitySystemComponent::OnRep_SimulatedTasks()
 {
 	for (UAbilityTask* SimulatedTask : SimulatedTasks)
@@ -2237,5 +2133,211 @@ UGameplayAbility* UAbilitySystemComponent::GetAnimatingAbility()
 {
 	return LocalAnimMontageInfo.AnimatingAbility;
 }
+
+// -----------------------------------------------------------------------------
+
+
+
+void UAbilitySystemComponent::ConfirmAbilityTargetData(FGameplayAbilitySpecHandle AbilityHandle, FPredictionKey AbilityOriginalPredictionKey, const FGameplayAbilityTargetDataHandle& TargetData, const FGameplayTag& ApplicationTag)
+{
+	FAbilityClientDataCache* CachedData = AbilityClientDataMap.Find(FGameplayAbilitySpecHandleAndPredictionKey(AbilityHandle, AbilityOriginalPredictionKey));
+	if (CachedData)
+	{
+		CachedData->TargetSetDelegate.Broadcast(TargetData, ApplicationTag);
+	}
+}
+
+void UAbilitySystemComponent::CancelAbilityTargetData(FGameplayAbilitySpecHandle AbilityHandle, FPredictionKey AbilityOriginalPredictionKey)
+{
+	FAbilityClientDataCache* CachedData = AbilityClientDataMap.Find(FGameplayAbilitySpecHandleAndPredictionKey(AbilityHandle, AbilityOriginalPredictionKey));
+	if (CachedData)
+	{
+		CachedData->Reset();
+		CachedData->TargetCancelledDelegate.Broadcast();
+	}
+}
+
+void UAbilitySystemComponent::ConsumeAllClientReplicatedData(FGameplayAbilitySpecHandle AbilityHandle, FPredictionKey AbilityOriginalPredictionKey)
+{
+	FAbilityClientDataCache* CachedData = AbilityClientDataMap.Find(FGameplayAbilitySpecHandleAndPredictionKey(AbilityHandle, AbilityOriginalPredictionKey));
+	if (CachedData)
+	{
+		CachedData->Reset();
+	}
+}
+
+void UAbilitySystemComponent::ConsumeClientReplicatedTargetData(FGameplayAbilitySpecHandle AbilityHandle, FPredictionKey AbilityOriginalPredictionKey)
+{
+	FAbilityClientDataCache* CachedData = AbilityClientDataMap.Find(FGameplayAbilitySpecHandleAndPredictionKey(AbilityHandle, AbilityOriginalPredictionKey));
+	if (CachedData)
+	{
+		CachedData->TargetData.Clear();
+		CachedData->bTargetConfirmed = false;
+		CachedData->bTargetCancelled = false;
+	}
+}
+
+void UAbilitySystemComponent::ConsumeClientReplicatedEvent(EAbilityReplicatedClientEvent::Type EventType, FGameplayAbilitySpecHandle AbilityHandle, FPredictionKey AbilityOriginalPredictionKey)
+{
+	FAbilityClientDataCache* CachedData = AbilityClientDataMap.Find(FGameplayAbilitySpecHandleAndPredictionKey(AbilityHandle, AbilityOriginalPredictionKey));
+	if (CachedData)
+	{
+		CachedData->GenericEvents[EventType].bTriggered = false;
+	}
+}
+
+// --------------------------------------------------------------------------
+
+void UAbilitySystemComponent::ServerSetReplicatedClientEvent_Implementation(EAbilityReplicatedClientEvent::Type EventType, FGameplayAbilitySpecHandle AbilityHandle, FPredictionKey AbilityOriginalPredictionKey,  FPredictionKey CurrentPredictionKey)
+{
+	FScopedPredictionWindow ScopedPrediction(this, CurrentPredictionKey);
+
+	InvokeReplicatedClientEvent(EventType, AbilityHandle, AbilityOriginalPredictionKey);
+}
+
+void UAbilitySystemComponent::InvokeReplicatedClientEvent(EAbilityReplicatedClientEvent::Type EventType, FGameplayAbilitySpecHandle AbilityHandle, FPredictionKey AbilityOriginalPredictionKey)
+{
+	FAbilityClientDataCache& ReplicatedData = AbilityClientDataMap.FindOrAdd(FGameplayAbilitySpecHandleAndPredictionKey(AbilityHandle, AbilityOriginalPredictionKey));
+	ReplicatedData.GenericEvents[(uint8)EventType].bTriggered = true;
+	ReplicatedData.GenericEvents[EventType].Delegate.Broadcast();
+}
+
+bool UAbilitySystemComponent::ServerSetReplicatedClientEvent_Validate(EAbilityReplicatedClientEvent::Type EventType, FGameplayAbilitySpecHandle AbilityHandle, FPredictionKey AbilityOriginalPredictionKey,  FPredictionKey CurrentPredictionKey)
+{
+	if (EventType >= EAbilityReplicatedClientEvent::MAX)
+	{
+		return false;
+	}
+	return true;
+}
+
+// -------
+
+void UAbilitySystemComponent::ServerSetReplicatedTargetData_Implementation(FGameplayAbilitySpecHandle AbilityHandle, FPredictionKey AbilityOriginalPredictionKey, FGameplayAbilityTargetDataHandle ReplicatedTargetData, FGameplayTag ApplicationTag, FPredictionKey CurrentPredictionKey)
+{
+	FScopedPredictionWindow ScopedPrediction(this, CurrentPredictionKey);
+
+	// Always adds to cache to store the new data
+	FAbilityClientDataCache& ReplicatedData = AbilityClientDataMap.FindOrAdd(FGameplayAbilitySpecHandleAndPredictionKey(AbilityHandle, AbilityOriginalPredictionKey));
+
+	ReplicatedData.TargetData = ReplicatedTargetData;
+	ReplicatedData.ApplicationTag = ApplicationTag;
+	ReplicatedData.bTargetConfirmed = true;
+	ReplicatedData.bTargetCancelled = false;
+	ReplicatedData.TargetSetDelegate.Broadcast(ReplicatedData.TargetData, ReplicatedData.ApplicationTag);
+}
+
+bool UAbilitySystemComponent::ServerSetReplicatedTargetData_Validate(FGameplayAbilitySpecHandle AbilityHandle, FPredictionKey AbilityOriginalPredictionKey, FGameplayAbilityTargetDataHandle ReplicatedTargetData, FGameplayTag ApplicationTag, FPredictionKey CurrentPredictionKey)
+{
+	return true;
+}
+
+// -------
+
+void UAbilitySystemComponent::ServerSetReplicatedTargetDataCancelled_Implementation(FGameplayAbilitySpecHandle AbilityHandle, FPredictionKey AbilityOriginalPredictionKey, FPredictionKey CurrentPredictionKey)
+{
+	FScopedPredictionWindow ScopedPrediction(this, CurrentPredictionKey);
+
+	// Always adds to cache to store the new data
+	FAbilityClientDataCache& ReplicatedData = AbilityClientDataMap.FindOrAdd(FGameplayAbilitySpecHandleAndPredictionKey(AbilityHandle, AbilityOriginalPredictionKey));
+
+	ReplicatedData.Reset();
+	ReplicatedData.bTargetCancelled = true;
+	ReplicatedData.TargetCancelledDelegate.Broadcast();
+}
+
+bool UAbilitySystemComponent::ServerSetReplicatedTargetDataCancelled_Validate(FGameplayAbilitySpecHandle AbilityHandle, FPredictionKey AbilityOriginalPredictionKey, FPredictionKey CurrentPredictionKey)
+{
+	return true;
+}
+
+void UAbilitySystemComponent::CallAllReplicatedDelegatesIfSet(FGameplayAbilitySpecHandle AbilityHandle, FPredictionKey AbilityOriginalPredictionKey)
+{
+	FAbilityClientDataCache* CachedData = AbilityClientDataMap.Find(FGameplayAbilitySpecHandleAndPredictionKey(AbilityHandle, AbilityOriginalPredictionKey));
+	if (CachedData)
+	{
+		if (CachedData->bTargetConfirmed)
+		{
+			CachedData->TargetSetDelegate.Broadcast(CachedData->TargetData, CachedData->ApplicationTag);
+		}
+		else if (CachedData->bTargetCancelled)
+		{
+			CachedData->TargetCancelledDelegate.Broadcast();
+		}
+
+		for (int32 idx=0; idx < EAbilityReplicatedClientEvent::MAX; ++idx)
+		{
+			if (CachedData->GenericEvents[idx].bTriggered)
+			{
+				CachedData->GenericEvents[idx].Delegate.Broadcast();
+			}
+		}
+	}
+}
+
+bool UAbilitySystemComponent::CallReplicatedTargetDataDelegatesIfSet(FGameplayAbilitySpecHandle AbilityHandle, FPredictionKey AbilityOriginalPredictionKey)
+{
+	bool CalledDelegate = false;
+	FAbilityClientDataCache* CachedData = AbilityClientDataMap.Find(FGameplayAbilitySpecHandleAndPredictionKey(AbilityHandle, AbilityOriginalPredictionKey));
+	if (CachedData)
+	{
+		if (CachedData->bTargetConfirmed)
+		{
+			CachedData->TargetSetDelegate.Broadcast(CachedData->TargetData, CachedData->ApplicationTag);
+			CalledDelegate = true;
+		}
+		else if (CachedData->bTargetCancelled)
+		{
+			CachedData->TargetCancelledDelegate.Broadcast();
+			CalledDelegate = true;
+		}
+	}
+
+	return CalledDelegate;
+}
+
+bool UAbilitySystemComponent::CallReplicatedEventDelegateIfSet(EAbilityReplicatedClientEvent::Type EventType, FGameplayAbilitySpecHandle AbilityHandle, FPredictionKey AbilityOriginalPredictionKey)
+{
+	FAbilityClientDataCache* CachedData = AbilityClientDataMap.Find(FGameplayAbilitySpecHandleAndPredictionKey(AbilityHandle, AbilityOriginalPredictionKey));
+	if (CachedData && CachedData->GenericEvents[EventType].bTriggered)
+	{
+		// Already triggered, fire off delegate
+		CachedData->GenericEvents[EventType].Delegate.Broadcast();
+		return true;
+	}
+	return false;
+}
+
+bool UAbilitySystemComponent::CallOrAddReplicatedDelegate(EAbilityReplicatedClientEvent::Type EventType, FGameplayAbilitySpecHandle AbilityHandle, FPredictionKey AbilityOriginalPredictionKey, FSimpleMulticastDelegate::FDelegate Delegate)
+{
+	FAbilityClientDataCache& CachedData = AbilityClientDataMap.FindOrAdd(FGameplayAbilitySpecHandleAndPredictionKey(AbilityHandle, AbilityOriginalPredictionKey));
+	if (CachedData.GenericEvents[EventType].bTriggered)
+	{
+		// Already triggered, fire off delegate
+		Delegate.Execute();
+		return true;
+	}
+	
+	// Not triggered yet, so just add the delegate
+	CachedData.GenericEvents[EventType].Delegate.Add(Delegate);
+	return false;
+}
+
+FAbilityTargetDataSetDelegate& UAbilitySystemComponent::AbilityTargetDataSetDelegate(FGameplayAbilitySpecHandle AbilityHandle, FPredictionKey AbilityOriginalPredictionKey)
+{
+	return AbilityClientDataMap.FindOrAdd(FGameplayAbilitySpecHandleAndPredictionKey(AbilityHandle, AbilityOriginalPredictionKey)).TargetSetDelegate;
+}
+
+FSimpleMulticastDelegate& UAbilitySystemComponent::AbilityTargetDataCancelledDelegate(FGameplayAbilitySpecHandle AbilityHandle, FPredictionKey AbilityOriginalPredictionKey)
+{
+	return AbilityClientDataMap.FindOrAdd(FGameplayAbilitySpecHandleAndPredictionKey(AbilityHandle, AbilityOriginalPredictionKey)).TargetCancelledDelegate;
+}
+
+FSimpleMulticastDelegate& UAbilitySystemComponent::AbilityReplicatedEventDelegate(EAbilityReplicatedClientEvent::Type EventType, FGameplayAbilitySpecHandle AbilityHandle, FPredictionKey AbilityOriginalPredictionKey)
+{
+	return AbilityClientDataMap.FindOrAdd(FGameplayAbilitySpecHandleAndPredictionKey(AbilityHandle, AbilityOriginalPredictionKey)).GenericEvents[EventType].Delegate;
+}
+
+
 #undef LOCTEXT_NAMESPACE
 
diff --git a/Engine/Source/Runtime/GameplayAbilities/Public/Abilities/GameplayAbility.h b/Engine/Source/Runtime/GameplayAbilities/Public/Abilities/GameplayAbility.h
--- a/Engine/Source/Runtime/GameplayAbilities/Public/Abilities/GameplayAbility.h
+++ b/Engine/Source/Runtime/GameplayAbilities/Public/Abilities/GameplayAbility.h
@@ -120,16 +120,7 @@ public:
 	//		EndAbility()			- The ability has ended. This is intended to be called by the ability to end itself.
 	//	
 	// ----------------------------------------------------------------------------------------------------------------
-	
-	/** Input binding stub. */
-	virtual void InputPressed(const FGameplayAbilitySpecHandle Handle, const FGameplayAbilityActorInfo* ActorInfo, const FGameplayAbilityActivationInfo ActivationInfo) {};
-
-	/** Input binding stub. */
-	virtual void InputReleased(const FGameplayAbilitySpecHandle Handle, const FGameplayAbilityActorInfo* ActorInfo, const FGameplayAbilityActivationInfo ActivationInfo) {};
-
-	/** If this is set, the ability will try to activate if the input is held. */
-	UPROPERTY(EditDefaultsOnly, Category = Input)
-	bool bActivateOnInputHeld;
+		
 
 	/** Returns true if this ability can be activated right now. Has no side effects */
 	virtual bool CanActivateAbility(const FGameplayAbilitySpecHandle Handle, const FGameplayAbilityActorInfo* ActorInfo, OUT FGameplayTagContainer* OptionalRelevantTags = nullptr) const;
@@ -241,6 +232,24 @@ public:
 	/** Called when the avatar actor is set/changes */
 	virtual void OnAvatarSet(const FGameplayAbilityActorInfo* ActorInfo, const FGameplayAbilitySpec& Spec);
 
+	// --------------------------------------
+	//	Input
+	// --------------------------------------
+
+	/** Input binding stub. */
+	virtual void InputPressed(const FGameplayAbilitySpecHandle Handle, const FGameplayAbilityActorInfo* ActorInfo, const FGameplayAbilityActivationInfo ActivationInfo) {};
+
+	/** Input binding stub. */
+	virtual void InputReleased(const FGameplayAbilitySpecHandle Handle, const FGameplayAbilityActorInfo* ActorInfo, const FGameplayAbilityActivationInfo ActivationInfo) {};
+
+	/** If this is set, the ability will try to activate if the input is held. */
+	UPROPERTY(EditDefaultsOnly, Category = Input)
+	bool bActivateOnInputHeld;
+
+	/** If true, this ability will always replicate input press/release events to the server. */
+	UPROPERTY(EditDefaultsOnly, Category = Input)
+	bool bReplicateInputDirectly;
+
 protected:
 
 	// --------------------------------------
@@ -308,15 +317,15 @@ protected:
 	UFUNCTION(BlueprintCallable, Category = Ability, FriendlyName = "CommitAbility")
 	virtual bool K2_CommitAbility();
 
-	/** Attempts to commit the ability's cooldown only. */
+	/** Attempts to commit the ability's cooldown only. If BroadcastCommitEvent is true, it will broadcast the commit event that tasks like WaitAbilityCommit are listening for. */
 	UFUNCTION(BlueprintCallable, Category = Ability, FriendlyName = "CommitAbilityCooldown")
-	virtual bool K2_CommitAbilityCooldown();
+	virtual bool K2_CommitAbilityCooldown(bool BroadcastCommitEvent=false);
 
-	/** Attempts to commit the ability's cost only. */
+	/** Attempts to commit the ability's cost only. If BroadcastCommitEvent is true, it will broadcast the commit event that tasks like WaitAbilityCommit are listening for. */
 	UFUNCTION(BlueprintCallable, Category = Ability, FriendlyName = "CommitAbilityCost")
-	virtual bool K2_CommitAbilityCost();
+	virtual bool K2_CommitAbilityCost(bool BroadcastCommitEvent=false);
 
-	/** Checks the ability's cooldown, but does not apply it. */
+	/** Checks the ability's cooldown, but does not apply it.*/
 	UFUNCTION(BlueprintCallable, Category = Ability, FriendlyName = "CheckAbilityCooldown")
 	virtual bool K2_CheckAbilityCooldown();
 
@@ -351,7 +360,7 @@ protected:
 	//	CancelAbility
 	// --------------------------------------
 
-	/** Destroys instanced-per-execution abilities. Instance-per-actor abilities should 'reset'. Non instance abilities - what can we do? */
+	/** Destroys instanced-per-execution abilities. Instance-per-actor abilities should 'reset'. Any active ability state tasks receive the 'OnAbilityStateInterrupted' event. Non instance abilities - what can we do? */
 	virtual void CancelAbility(const FGameplayAbilitySpecHandle Handle, const FGameplayAbilityActorInfo* ActorInfo, const FGameplayAbilityActivationInfo ActivationInfo);
 
 	/** Destroys instanced-per-execution abilities. Instance-per-actor abilities should 'reset'. Non instance abilities - what can we do? */
@@ -371,6 +380,10 @@ protected:
 	UFUNCTION(BlueprintCallable, Category = Ability)
 	void CancelTaskByInstanceName(FName InstanceName);
 
+	/** Ends any active ability state task with the given name. If name is 'None' all active states will be ended (in an arbitrary order). */
+	UFUNCTION(BlueprintCallable, Category = Ability)
+	void EndAbilityState(FName OptionalStateNameToEnd);
+
 	// -------------------------------------
 	//	EndAbility
 	// -------------------------------------
diff --git a/Engine/Source/Runtime/GameplayAbilities/Public/Abilities/GameplayAbilityTargetActor.h b/Engine/Source/Runtime/GameplayAbilities/Public/Abilities/GameplayAbilityTargetActor.h
--- a/Engine/Source/Runtime/GameplayAbilities/Public/Abilities/GameplayAbilityTargetActor.h
+++ b/Engine/Source/Runtime/GameplayAbilities/Public/Abilities/GameplayAbilityTargetActor.h
@@ -45,8 +45,6 @@ public:
 	UFUNCTION()
 	virtual void ConfirmTargeting();
 
-	void NotifyPlayerControllerOfRejectedConfirmation();
-
 	/** Outside code is saying 'stop everything and just forget about it' */
 	UFUNCTION()
 	virtual void CancelTargeting();
@@ -94,4 +92,7 @@ public:
 	/** Draw the debug information (if applicable) for this targeting actor. */
 	UPROPERTY(BlueprintReadWrite, EditAnywhere, Replicated, meta = (ExposeOnSpawn = true), Category = Targeting)
 	bool bDebug;
+
+	FDelegateHandle GenericConfirmHandle;
+	FDelegateHandle GenericCancelHandle;
 };
\ No newline at end of file
diff --git a/Engine/Source/Runtime/GameplayAbilities/Public/Abilities/GameplayAbilityTargetTypes.h b/Engine/Source/Runtime/GameplayAbilities/Public/Abilities/GameplayAbilityTargetTypes.h
--- a/Engine/Source/Runtime/GameplayAbilities/Public/Abilities/GameplayAbilityTargetTypes.h
+++ b/Engine/Source/Runtime/GameplayAbilities/Public/Abilities/GameplayAbilityTargetTypes.h
@@ -168,7 +168,7 @@ struct GAMEPLAYABILITIES_API FGameplayAbilityTargetDataHandle
 		Data.Add(TSharedPtr<FGameplayAbilityTargetData>(DataPtr));
 	}
 
-	TArray<TSharedPtr<FGameplayAbilityTargetData>>	Data;
+	TArray<TSharedPtr<FGameplayAbilityTargetData> >	Data;
 
 	void Clear()
 	{
@@ -587,3 +587,77 @@ struct TStructOpsTypeTraits<FGameplayAbilityTargetData_SingleTargetHit> : public
 
 /** Generic callback for returning when target data is available */
 DECLARE_MULTICAST_DELEGATE_OneParam(FAbilityTargetData, FGameplayAbilityTargetDataHandle);
+
+
+// ----------------------------------------------------
+
+/** Generic callback for returning when target data is available */
+DECLARE_MULTICAST_DELEGATE_TwoParams(FAbilityTargetDataSetDelegate, FGameplayAbilityTargetDataHandle, FGameplayTag);
+
+/** These are generic, nonpayload carrying events that the client can send to the server via ServerSetReplicatedClientEvent. */
+UENUM()
+namespace EAbilityReplicatedClientEvent
+{
+	enum Type
+	{	
+		/** A generic confirmation to commit the ability */
+		GenericConfirm = 0,
+		/** A generic cancellation event. Not necessarily a canellation of the ability or targeting. Could be used to cancel out of a channelling portion of ability. */
+		GenericCancel,
+		/** Additional input presses of the ability (Press X to activate ability, press X again while it is active to do other things within the GameplayAbility's logic) */
+		InputPressed,	
+		/** Input release event of the ability */
+		InputReleased,
+
+		MAX
+	};
+}
+
+struct FAbilityReplicatedData
+{
+	FAbilityReplicatedData() : bTriggered(false) {}
+	/** Event has triggered */
+	bool bTriggered;
+
+	FSimpleMulticastDelegate Delegate;
+};
+
+/** Struct defining the cached data for a specific gameplay ability. This data is generally syncronized client->server in a network game. */
+struct FAbilityClientDataCache
+{
+	/** What elements this activation is targeting */
+	FGameplayAbilityTargetDataHandle TargetData;
+
+	/** What tag to pass through when doing an application */
+	FGameplayTag ApplicationTag;
+
+	/** True if we've been positively confirmed our targeting, false if we don't know */
+	bool bTargetConfirmed;
+
+	/** True if we've been positively cancelled our targeting, false if we don't know */
+	bool bTargetCancelled;
+
+	/** Delegate to call whenever this is modified */
+	FAbilityTargetDataSetDelegate TargetSetDelegate;
+
+	/** Delegate to call whenever this is confirmed (without target data) */
+	FSimpleMulticastDelegate TargetCancelledDelegate;
+
+	/** Generic events that contain no payload data */
+	FAbilityReplicatedData	GenericEvents[EAbilityReplicatedClientEvent::MAX];
+
+	FAbilityClientDataCache() : bTargetConfirmed(false), bTargetCancelled(false) {}
+
+	/** Resets any cached data, leaves delegates up */
+	void Reset()
+	{
+		bTargetConfirmed = bTargetCancelled = false;
+		TargetData = FGameplayAbilityTargetDataHandle();
+		ApplicationTag = FGameplayTag();
+		for (int32 i=0; i < (int32) EAbilityReplicatedClientEvent::MAX; ++i)
+		{
+			GenericEvents[i].bTriggered = false;
+		}
+
+	}
+};
\ No newline at end of file
diff --git a/Engine/Source/Runtime/GameplayAbilities/Public/Abilities/GameplayAbilityTypes.h b/Engine/Source/Runtime/GameplayAbilities/Public/Abilities/GameplayAbilityTypes.h
--- a/Engine/Source/Runtime/GameplayAbilities/Public/Abilities/GameplayAbilityTypes.h
+++ b/Engine/Source/Runtime/GameplayAbilities/Public/Abilities/GameplayAbilityTypes.h
@@ -390,6 +390,46 @@ struct TStructOpsTypeTraits< FGameplayAbilitySpecContainer > : public TStructOps
 	};
 };
 
+// ----------------------------------------------------
+
+USTRUCT()
+struct FGameplayAbilitySpecHandleAndPredictionKey
+{
+	GENERATED_USTRUCT_BODY()
+
+	FGameplayAbilitySpecHandleAndPredictionKey()
+	: PredictionKeyAtCreation(0)
+	{}
+
+	FGameplayAbilitySpecHandleAndPredictionKey(const FGameplayAbilitySpecHandle& HandleRef, const FPredictionKey& PredictionKeyAtCreationRef)
+		: AbilityHandle(HandleRef), PredictionKeyAtCreation(PredictionKeyAtCreationRef.Current)
+	{}
+
+	bool operator==(const FGameplayAbilitySpecHandleAndPredictionKey& Other) const
+	{
+		return AbilityHandle == Other.AbilityHandle && PredictionKeyAtCreation == Other.PredictionKeyAtCreation;
+	}
+
+	bool operator!=(const FGameplayAbilitySpecHandleAndPredictionKey& Other) const
+	{
+		return AbilityHandle != Other.AbilityHandle || PredictionKeyAtCreation != Other.PredictionKeyAtCreation;
+	}
+
+	friend uint32 GetTypeHash(const FGameplayAbilitySpecHandleAndPredictionKey& Handle)
+	{
+		return GetTypeHash(Handle.AbilityHandle) ^ Handle.PredictionKeyAtCreation;
+	}
+
+private:
+
+	UPROPERTY()
+	FGameplayAbilitySpecHandle AbilityHandle;
+
+	UPROPERTY()
+	int32 PredictionKeyAtCreation;
+};
+
+
 // ----------------------------------------------------
 
 
diff --git a/Engine/Source/Runtime/GameplayAbilities/Public/Abilities/Tasks/AbilityTask.h b/Engine/Source/Runtime/GameplayAbilities/Public/Abilities/Tasks/AbilityTask.h
--- a/Engine/Source/Runtime/GameplayAbilities/Public/Abilities/Tasks/AbilityTask.h
+++ b/Engine/Source/Runtime/GameplayAbilities/Public/Abilities/Tasks/AbilityTask.h
@@ -98,6 +98,12 @@ class GAMEPLAYABILITIES_API UAbilityTask : public UObject
 	/** Proper way to get the avatar actor of the ability that owns this task (usually a pawn, tower, etc) */
 	AActor* GetAvatarActor() const;
 
+	/** Returns spec handle for owning ability */
+	FGameplayAbilitySpecHandle GetAbilitySpecHandle() const;
+
+	/** Returns ActivationPredictionKey of owning ability */
+	FPredictionKey GetActivationPredictionKey() const;
+
 	/** Helper function for instantiating and initializing a new task */
 	template <class T>
 	static T*	NewTask(UObject* WorldContextObject, FName InstanceName = FName())
@@ -135,10 +141,23 @@ class GAMEPLAYABILITIES_API UAbilityTask : public UObject
 	/** Am I actually running this as a simulated task. (This will be true on clients that simulating. This will be false on the server and the owning client) */
 	bool bIsSimulating;
 
+	/** Returns true if the ability is a locally predicted ability running on a client. Usually this means we need to tell the server something. */
+	bool IsPredictingClient() const;
+
+	/** Returns true if we are executing the ability on the server for a non locally controlled client */
+	bool IsForRemoteClient() const;
+
 protected:	
 
 	/** End and CleanUp the task - may be called by the task itself or by the owning ability if the ability is ending. Do NOT call directly! Call EndTask() or AbilityEnded() */
 	virtual void OnDestroy(bool AbilityIsEnding);
+
+	/** Helper method for registering client replicated callbacks */
+	
+	bool CallOrAddReplicatedDelegate(EAbilityReplicatedClientEvent::Type Event, FSimpleMulticastDelegate::FDelegate Delegate)
+	{
+		return AbilitySystemComponent->CallOrAddReplicatedDelegate(Event, GetAbilitySpecHandle(), GetActivationPredictionKey(), Delegate);
+	}
 };
 
 //For searching through lists of ability instances
diff --git a/Engine/Source/Runtime/GameplayAbilities/Public/Abilities/Tasks/AbilityTask_StartAbilityState.h b/Engine/Source/Runtime/GameplayAbilities/Public/Abilities/Tasks/AbilityTask_StartAbilityState.h
new file mode 100644
--- /dev/null
+++ b/Engine/Source/Runtime/GameplayAbilities/Public/Abilities/Tasks/AbilityTask_StartAbilityState.h
@@ -0,0 +1,63 @@
+// Copyright 1998-2015 Epic Games, Inc. All Rights Reserved.
+
+#pragma once
+
+#include "AbilityTask.h"
+#include "Abilities/GameplayAbilityTypes.h"
+#include "AbilityTask_StartAbilityState.generated.h"
+
+DECLARE_DYNAMIC_MULTICAST_DELEGATE(FAbilityStateDelegate);
+
+/**
+ * An ability state is simply an ability task that provides a way to handle the ability being interrupted.
+ * You can use ability states to do state-specific cleanup if the ability ends or was interrupted at a certain point during it's execution.
+ *
+ * An ability state will always result in either 'OnStateEnded' or 'OnStateInterrupted' being called.
+ *
+ * 'OnStateEnded' will be called if:
+ * - The ability itself ends via AGameplayAbility::EndAbility
+ * - The ability state is manually ended via AGameplayAbility::EndAbilityState
+ * - Another ability state is started will 'bEndCurrentState' set to true
+ *
+ * 'OnStateInterrupted' will be called if:
+ * - The ability itself is cancelled via AGameplayAbility::CancelAbility
+ */
+UCLASS(MinimalAPI)
+class UAbilityTask_StartAbilityState : public UAbilityTask
+{
+	GENERATED_UCLASS_BODY()
+
+	/** Invoked if 'EndAbilityState' is called or if 'EndAbility' is called and this state is active. */
+	UPROPERTY(BlueprintAssignable)
+	FAbilityStateDelegate OnStateEnded;
+
+	/** Invoked if the ability was interrupted and this state is active. */
+	UPROPERTY(BlueprintAssignable)
+	FAbilityStateDelegate OnStateInterrupted;
+
+	virtual void Activate() override;
+
+	virtual FString GetDebugString() const override;
+
+	/**
+	 * Starts a new ability state.
+	 *
+	 * @param StateName The name of the state.
+	 * @param bEndCurrentState If true, all other active ability states will be ended.
+	 */
+	UFUNCTION(BlueprintCallable, Meta = (HidePin = "WorldContextObject", DefaultToSelf = "WorldContextObject", BlueprintInternalUseOnly = "true", HideSpawnParms = "Instigator"), Category = "Ability|Tasks")
+	static UAbilityTask_StartAbilityState* StartAbilityState(UObject* WorldContextObject, FName StateName, bool bEndCurrentState = true);
+
+private:
+
+	FDelegateHandle EndStateHandle;
+	FDelegateHandle InterruptStateHandle;
+
+	bool bEndCurrentState;
+
+	virtual void OnDestroy(bool AbilityEnded) override;
+
+	void OnEndState(FName StateNameToEnd);
+
+	void OnInterruptState();
+};
diff --git a/Engine/Source/Runtime/GameplayAbilities/Public/Abilities/Tasks/AbilityTask_WaitConfirmCancel.h b/Engine/Source/Runtime/GameplayAbilities/Public/Abilities/Tasks/AbilityTask_WaitConfirmCancel.h
--- a/Engine/Source/Runtime/GameplayAbilities/Public/Abilities/Tasks/AbilityTask_WaitConfirmCancel.h
+++ b/Engine/Source/Runtime/GameplayAbilities/Public/Abilities/Tasks/AbilityTask_WaitConfirmCancel.h
@@ -28,6 +28,12 @@ class UAbilityTask_WaitConfirmCancel : public UAbilityTask
 	UFUNCTION()
 	void OnCancelCallback();
 
+	UFUNCTION()
+	void OnLocalConfirmCallback();
+
+	UFUNCTION()
+	void OnLocalCancelCallback();
+
 	UFUNCTION(BlueprintCallable, meta=(HidePin = "WorldContextObject", DefaultToSelf = "WorldContextObject", BlueprintInternalUseOnly = "true", FriendlyName="Wait for Confirm Input"), Category="Ability|Tasks")
 	static UAbilityTask_WaitConfirmCancel* WaitConfirmCancel(UObject* WorldContextObject);	
 
diff --git a/Engine/Source/Runtime/GameplayAbilities/Public/Abilities/Tasks/AbilityTask_WaitInputPress.h b/Engine/Source/Runtime/GameplayAbilities/Public/Abilities/Tasks/AbilityTask_WaitInputPress.h
--- a/Engine/Source/Runtime/GameplayAbilities/Public/Abilities/Tasks/AbilityTask_WaitInputPress.h
+++ b/Engine/Source/Runtime/GameplayAbilities/Public/Abilities/Tasks/AbilityTask_WaitInputPress.h
@@ -20,7 +20,7 @@ class UAbilityTask_WaitInputPress : public UAbilityTask
 	FInputPressDelegate		OnPress;
 
 	UFUNCTION()
-	void OnPressCallback(int32 InputID);
+	void OnPressCallback();
 
 	virtual void Activate() override;
 
@@ -33,5 +33,5 @@ protected:
 	virtual void OnDestroy(bool AbilityEnded) override;
 
 	float StartTime;
-	bool RegisteredCallback;
+	FDelegateHandle DelegateHandle;
 };
\ No newline at end of file
diff --git a/Engine/Source/Runtime/GameplayAbilities/Public/Abilities/Tasks/AbilityTask_WaitInputRelease.h b/Engine/Source/Runtime/GameplayAbilities/Public/Abilities/Tasks/AbilityTask_WaitInputRelease.h
--- a/Engine/Source/Runtime/GameplayAbilities/Public/Abilities/Tasks/AbilityTask_WaitInputRelease.h
+++ b/Engine/Source/Runtime/GameplayAbilities/Public/Abilities/Tasks/AbilityTask_WaitInputRelease.h
@@ -20,7 +20,7 @@ class UAbilityTask_WaitInputRelease : public UAbilityTask
 	FInputReleaseDelegate	OnRelease;
 
 	UFUNCTION()
-	void OnReleaseCallback(int32 InputID);
+	void OnReleaseCallback();
 
 	virtual void Activate() override;
 
@@ -33,5 +33,5 @@ protected:
 	virtual void OnDestroy(bool AbilityEnded) override;
 
 	float StartTime;
-	bool RegisteredCallback;
+	FDelegateHandle DelegateHandle;
 };
\ No newline at end of file
diff --git a/Engine/Source/Runtime/GameplayAbilities/Public/Abilities/Tasks/AbilityTask_WaitTargetData.h b/Engine/Source/Runtime/GameplayAbilities/Public/Abilities/Tasks/AbilityTask_WaitTargetData.h
--- a/Engine/Source/Runtime/GameplayAbilities/Public/Abilities/Tasks/AbilityTask_WaitTargetData.h
+++ b/Engine/Source/Runtime/GameplayAbilities/Public/Abilities/Tasks/AbilityTask_WaitTargetData.h
@@ -22,7 +22,7 @@ class GAMEPLAYABILITIES_API UAbilityTask_WaitTargetData: public UAbilityTask
 	FWaitTargetDataDelegate	Cancelled;
 
 	UFUNCTION()
-	void OnTargetDataReplicatedCallback(FGameplayAbilityTargetDataHandle Data);
+	void OnTargetDataReplicatedCallback(FGameplayAbilityTargetDataHandle Data, FGameplayTag ActivationTag);
 
 	UFUNCTION()
 	void OnTargetDataReplicatedCancelledCallback();
diff --git a/Engine/Source/Runtime/GameplayAbilities/Public/AbilitySystemComponent.h b/Engine/Source/Runtime/GameplayAbilities/Public/AbilitySystemComponent.h
--- a/Engine/Source/Runtime/GameplayAbilities/Public/AbilitySystemComponent.h
+++ b/Engine/Source/Runtime/GameplayAbilities/Public/AbilitySystemComponent.h
@@ -50,12 +50,15 @@ class GAMEPLAYABILITIES_API UAbilitySystemComponent : public UActorComponent, pu
 {
 	GENERATED_UCLASS_BODY()
 
-	/** Used to register callbacks to ability-key input */
-	DECLARE_DYNAMIC_MULTICAST_DELEGATE_OneParam(FAbilityAbilityKey, /*UGameplayAbility*, Ability, */int32, InputID);
-
 	/** Used to register callbacks to confirm/cancel input */
 	DECLARE_DYNAMIC_MULTICAST_DELEGATE(FAbilityConfirmOrCancel);
 
+	/** Used to notify ability state tasks that a state is being ended - should not be used directly */
+	DECLARE_MULTICAST_DELEGATE_OneParam(FAbilityStateEnded, FName);
+
+	/** Used to notify ability state tasks that an interrupt is occurring - should not be used directly */
+	DECLARE_MULTICAST_DELEGATE(FAbilityStateInterrupted);
+
 	friend struct FActiveGameplayEffectAction_Add;
 	friend FGameplayEffectSpec;
 	friend class AAbilitySystemDebugHUD;
@@ -500,9 +503,27 @@ class GAMEPLAYABILITIES_API UAbilitySystemComponent : public UActorComponent, pu
 	/** Removes the specified ability */
 	void ClearAbility(const FGameplayAbilitySpecHandle& Handle);
 
+	// --------------------------------------------
+	// Ability Cancelling/Interrupts
+	// --------------------------------------------
+
+	/** Used by the ability state task to handle when a state is ended - should not use directly */
+	FAbilityStateEnded OnAbilityStateEnded;
+
+	/** Used by the ability state task to handle when a state is interrupted - should not use directly */
+	FAbilityStateInterrupted OnAbilityStateInterrupted;
+
+	/** Cancels the specified ability. */
+	void CancelAbility(UGameplayAbility* Ability);
+
 	/** Cancel all abilities with the specified tags. Will not cancel the Ignore instance */
 	void CancelAbilities(const FGameplayTagContainer* WithTags=nullptr, const FGameplayTagContainer* WithoutTags=nullptr, UGameplayAbility* Ignore=nullptr);
 
+	/** Cancels all abilities regardless of tags. Will not cancel the ignore instance */
+	void CancelAllAbilities(UGameplayAbility* Ignore=nullptr);
+
+	// ----------------------------------------------------------------------------------------------------------------
+
 	/** 
 	 * Called from ability activation or native code, will apply the correct ability blocking tags and cancel existing abilities. Subclasses can override the behavior 
 	 * @param AbilityTags The tags of the ability that has block and cancel flags
@@ -542,9 +563,13 @@ class GAMEPLAYABILITIES_API UAbilitySystemComponent : public UActorComponent, pu
 	 *	instancing or anything else that the AbilitySystemComponent would provide, then it doesn't need the component to function.
 	 */
 protected:
+
 	UPROPERTY(ReplicatedUsing=OnRep_ActivateAbilities, BlueprintReadOnly, Category = "Abilities")
 	FGameplayAbilitySpecContainer	ActivatableAbilities;
 
+	/** Maps from an ability spec to the target data. Used to track replicated data and callbacks */
+	TMap<FGameplayAbilitySpecHandleAndPredictionKey, FAbilityClientDataCache> AbilityClientDataMap;
+
 	/** Will be called from GiveAbility or from OnRep. Initializes events (triggers and inputs) with the given ability */
 	virtual void OnGiveAbility(FGameplayAbilitySpec& AbilitySpec);
 
@@ -578,10 +603,10 @@ public:
 	void	OnRep_ActivateAbilities();
 
 	UFUNCTION(Server, reliable, WithValidation)
-	void	ServerTryActivateAbility(FGameplayAbilitySpecHandle AbilityToActivate, bool InputPressed, FPredictionKey PredictionKey);
+	void	ServerTryActivateAbility(FGameplayAbilitySpecHandle AbilityToActivate, FPredictionKey PredictionKey);
 
 	UFUNCTION(Server, reliable, WithValidation)
-	void	ServerTryActivateAbilityWithEventData(FGameplayAbilitySpecHandle AbilityToActivate, bool InputPressed, FPredictionKey PredictionKey, FGameplayEventData TriggerEventData);
+	void	ServerTryActivateAbilityWithEventData(FGameplayAbilitySpecHandle AbilityToActivate, FPredictionKey PredictionKey, FGameplayEventData TriggerEventData);
 
 	/** Replicate that an ability has ended, to the client or server as appropriate */
 	void	ReplicateEndAbility(FGameplayAbilitySpecHandle Handle, FGameplayAbilityActivationInfo ActivationInfo, UGameplayAbility* Ability);
@@ -603,11 +628,6 @@ public:
 	UFUNCTION(Client, Reliable)
 	void	ClientActivateAbilitySucceed(FGameplayAbilitySpecHandle AbilityToActivate, int16 PredictionKey, FGameplayEventData TriggerEventData);
 
-	/** Attempted to confirm targeting, but the targeting actor rejected it. */
-	UFUNCTION(Client, Unreliable)
-	void	ClientAbilityNotifyRejected(int32 InputID);
-
-
 	// ----------------------------------------------------------------------------------------------------------------
 
 	// This is meant to be used to inhibit activating an ability from an input perspective. (E.g., the menu is pulled up, another game mechanism is consuming all input, etc)
@@ -629,36 +649,28 @@ public:
 	
 	virtual void BindAbilityActivationToInputComponent(UInputComponent* InputComponent, FGameplayAbiliyInputBinds BindInfo);
 
-	void AbilityInputPressed(int32 InputID);
-
-	void AbilityInputReleased(int32 InputID);
+	void AbilityLocalInputPressed(int32 InputID);
 
-	void AbilitySpecInputPressed(FGameplayAbilitySpec& Spec);
-
-	void AbilitySpecInputReleased(FGameplayAbilitySpec& Spec);
+	void AbilityLocalInputReleased(int32 InputID);
 
 	void TryActivateInputHeldAbilities();
 
-	/** Sent by abilities to tell server when activation input is pressed. (Sent by default in order not to short-circuit WaitInputPress/WaitInputRelease tasks) */
-	UFUNCTION(Server, reliable, WithValidation)
-	void ServerInputPress(FGameplayAbilitySpecHandle AbilityToActivate, FPredictionKey ScopedPedictionKey);
-
-	/** Sent by abilities to tell server when activation input is released. (Sent by default in order not to short-circuit WaitInputPress/WaitInputRelease tasks) */
-	UFUNCTION(Server, reliable, WithValidation)
-	void ServerInputRelease(FGameplayAbilitySpecHandle AbilityToActivate, FPredictionKey ScopedPedictionKey);
-
 	UFUNCTION(BlueprintCallable, Category="Abilities")
-	void InputConfirm();
+	void LocalInputConfirm();
 
 	UFUNCTION(BlueprintCallable, Category="Abilities")
-	void InputCancel();
+	void LocalInputCancel();
+
+	/** Generic local callback for generic ConfirmEvent that any ability can listen to */
+	FAbilityConfirmOrCancel	GenericLocalConfirmCallbacks;
 
-	FAbilityAbilityKey	AbilityKeyPressCallbacks;
-	FAbilityAbilityKey	AbilityKeyReleaseCallbacks;
-	FAbilityConfirmOrCancel	ConfirmCallbacks;
-	FAbilityConfirmOrCancel	CancelCallbacks;
+	/** Generic local callback for generic CancelEvent that any ability can listen to */
+	FAbilityConfirmOrCancel	GenericLocalCancelCallbacks;
 
+	/** A generic callback anytime an ability is activated (started) */	
 	FGenericAbilityDelegate AbilityActivatedCallbacks;
+
+	/** A generic callback anytime an ability is commited (cost/cooldown applied) */
 	FGenericAbilityDelegate AbilityCommitedCallbacks;
 
 	/** Executes a gameplay event. Returns the number of successful ability activations triggered by the event */
@@ -736,7 +748,7 @@ public:
 protected:
 
 	/** Implementation of ServerTryActivateAbility */
-	virtual void InternalServerTryActiveAbility(FGameplayAbilitySpecHandle AbilityToActivate, bool InputPressed, const FPredictionKey& PredictionKey, const FGameplayEventData* TriggerEventData);
+	virtual void InternalServerTryActiveAbility(FGameplayAbilitySpecHandle AbilityToActivate, const FPredictionKey& PredictionKey, const FGameplayEventData* TriggerEventData);
 
 	/** Called when a prediction key that played a montage is rejected */
 	void OnPredictiveMontageRejected(UAnimMontage* PredictiveMontage);
@@ -816,39 +828,96 @@ public:
 	// -----------------------------------------------------------------------------
 
 	/**
-	 *	While these appear to be state, these are actually synchronization events w/ some payload data
+	 *	Client -> Server RPCs used by activate abilities
+	 *	
+	 *	There are a few ways active abilities on clients communicate with the server through the AbilitySystemComponent.
+	 *	1. Through TargetData events. Either that client sends TargetData via ServerSetReplicatedTargetData or cancels it via ServerSetReplicatedTargetDataCancelled.
+	 *	2. Through Generic Client Events. These are non payload carrying events defined by the EAbilityReplicatedClientEvent enum.
+	 *	
+	 *	The general pattern is:
+	 *		-Clients and server register with these events via AbilityTargetDataSetDelegate / AbilityReplicatedEventDelegate.
+	 *		-These events themselves can be instigated on the client, via InvokeReplicatedClientEvent (which would trigger the callback locally).
+	 *		-Client is ultimately responbile to call ServerSetReplicatedClientEvent/ServerSetReplicatedTargetData to replicate the event to the server.
+	 *		-The server may get the replicated event before the server is actually listening for it! We store whether the event was recevied or not, for the given ability spec/prediction key pair.
+	 *			-Use functions like CallReplicatedEventDelegateIfSet/CallReplicatedTargetDataDelegatesIfSet after registering the delegates on the server. This invokes the delegates if they were already received.
+	 *			-Alternatively, use CallOrAddReplicatedDelegate to do it in one fell swoop (not good if you need multiple callbacks though).
+	 *	
+	 *	3. Through direct input replication (assuming you are using the provided input binding code in AbilitySystemComponent). 
+	 *	   This is a per ability setting defined by UGameplayAbility::bReplicateInputDirectly. It is generally NOT a good idea too use this! It is better to use the above Generic Client Eveents.
+	 *		
+	 *	   When bReplicateInputDirectly, the server RPCs ServerSetInputPressed/ServerSetInputReleased will be called whenever the ability is active. These will call AbilitySpecInputReleased/AbilitySpecInputPressed
+	 *	   and ultimately pipe down in UGameplayAbility::InputPressed/Released.
+	 *	   
+	 *	   This is useful if you truly need to maintain tracked input state on the server. For example if your ability logic requires polling of input state and arbitrary times (anim events, touched event, etc).
+	 *	   If your ability just requires event based input (button was pressed, released) then it is better to use Generic Client Events which will only replicate when the ability is listening for those events.
+	 *		
 	 */
 
+	/** Replicates the Generic CLient Event to the server. */
 	UFUNCTION(Server, reliable, WithValidation)
-	void ServerSetReplicatedAbilityKeyState(int32 InputID, bool Pressed, FPredictionKey PredictionKey);
+	void ServerSetReplicatedClientEvent(EAbilityReplicatedClientEvent::Type EventType, FGameplayAbilitySpecHandle AbilityHandle, FPredictionKey AbilityOriginalPredictionKey, FPredictionKey CurrentPredictionKey);
 
-	UFUNCTION(Server, reliable, WithValidation)
-	void ServerSetReplicatedConfirm(bool Confirmed, FPredictionKey PredictionKey);
+	/** Calls local callbacks that are registered with the given Generic Client Event */
+	void InvokeReplicatedClientEvent(EAbilityReplicatedClientEvent::Type EventType, FGameplayAbilitySpecHandle AbilityHandle, FPredictionKey AbilityOriginalPredictionKey);
 
+	/**  */
 	UFUNCTION(Server, reliable, WithValidation)
-	void ServerSetReplicatedTargetData(FGameplayAbilityTargetDataHandle ReplicatedTargetData, FPredictionKey PredictionKey);
+	void ServerSetReplicatedTargetData(FGameplayAbilitySpecHandle AbilityHandle, FPredictionKey AbilityOriginalPredictionKey, FGameplayAbilityTargetDataHandle ReplicatedTargetData, FGameplayTag ApplicationTag, FPredictionKey CurrentPredictionKey);
 
+	/** Replicates to the server that targeting has been cancelled */
 	UFUNCTION(Server, reliable, WithValidation)
-	void ServerSetReplicatedTargetDataCancelled();
+	void ServerSetReplicatedTargetDataCancelled(FGameplayAbilitySpecHandle AbilityHandle, FPredictionKey AbilityOriginalPredictionKey, FPredictionKey CurrentPredictionKey);
 
-	void ConsumeAbilityConfirmCancel();
+	/** Sets the current target data and calls applicable callbacks */
+	void ConfirmAbilityTargetData(FGameplayAbilitySpecHandle AbilityHandle, FPredictionKey AbilityOriginalPredictionKey, const FGameplayAbilityTargetDataHandle& TargetData, const FGameplayTag& ApplicationTag);
 
-	void ConsumeAbilityTargetData();
+	/** Cancels the ability target data and calls callbacks */
+	void CancelAbilityTargetData(FGameplayAbilitySpecHandle AbilityHandle, FPredictionKey AbilityOriginalPredictionKey);
 
-	//Note: Ability key state is stored on the ability spec, not here on the ASC with confirm/cancel.
-	bool ReplicatedConfirmAbility;
-	bool ReplicatedCancelAbility;
+	/** Deletes all cached ability client data (Was: ConsumeAbilityTargetData)*/
+	void ConsumeAllClientReplicatedData(FGameplayAbilitySpecHandle AbilityHandle, FPredictionKey AbilityOriginalPredictionKey);
+
+	/** Consumes cached TargetData from client (only TargetData) */
+	void ConsumeClientReplicatedTargetData(FGameplayAbilitySpecHandle AbilityHandle, FPredictionKey AbilityOriginalPredictionKey);
+
+	/** Consumes the given Generic Client Event (unsets it). */
+	void ConsumeClientReplicatedEvent(EAbilityReplicatedClientEvent::Type EventType, FGameplayAbilitySpecHandle AbilityHandle, FPredictionKey AbilityOriginalPredictionKey);
+	
+	/** Calls any ClientReplicated delegates that have been sent (TargetData or Generic Client Events). Note this can be dangerous if multiple places in an ability register events and then call this function. */
+	void CallAllReplicatedDelegatesIfSet(FGameplayAbilitySpecHandle AbilityHandle, FPredictionKey AbilityOriginalPredictionKey);
 
-	FGameplayAbilityTargetDataHandle ReplicatedTargetData;
+	/** Calls the TargetData Confirm/Cancel events if they have been sent. */
+	bool CallReplicatedTargetDataDelegatesIfSet(FGameplayAbilitySpecHandle AbilityHandle, FPredictionKey AbilityOriginalPredictionKey);
 
-	/** ReplicatedTargetData was received */
-	FAbilityTargetData	ReplicatedTargetDataDelegate;
+	/** Calls a given Generic Client Event delegate if the event has already been sent */
+	bool CallReplicatedEventDelegateIfSet(EAbilityReplicatedClientEvent::Type EventType, FGameplayAbilitySpecHandle AbilityHandle, FPredictionKey AbilityOriginalPredictionKey);
 
-	/** ReplicatedTargetData was 'cancelled' for this activation */
-	FAbilityConfirmOrCancel	ReplicatedTargetDataCancelledDelegate;
+	/** Calls passed in delegate if the Client Event has already been sent. If not, it adds the delegate to our multicast callback that will fire when it does. */
+	bool CallOrAddReplicatedDelegate(EAbilityReplicatedClientEvent::Type EventType, FGameplayAbilitySpecHandle AbilityHandle, FPredictionKey AbilityOriginalPredictionKey, FSimpleMulticastDelegate::FDelegate Delegate);
 
-	/** Targeting actor rejected a confirmation attempt */
-	FTargetingRejectedConfirmation TargetingRejectedConfirmationDelegate;
+	/** Returns TargetDataSet delegate for a given Ability/PredictionKey pair */
+	FAbilityTargetDataSetDelegate& AbilityTargetDataSetDelegate(FGameplayAbilitySpecHandle AbilityHandle, FPredictionKey AbilityOriginalPredictionKey);
+
+	/** Returns TargetData Cancelled delegate for a given Ability/PredictionKey pair */
+	FSimpleMulticastDelegate& AbilityTargetDataCancelledDelegate(FGameplayAbilitySpecHandle AbilityHandle, FPredictionKey AbilityOriginalPredictionKey);
+
+	/** Returns Generic Client Event for a given Ability/PredictionKey pair */
+	FSimpleMulticastDelegate& AbilityReplicatedEventDelegate(EAbilityReplicatedClientEvent::Type EventType, FGameplayAbilitySpecHandle AbilityHandle, FPredictionKey AbilityOriginalPredictionKey);
+
+	// Direct Input state replication. These will be called if bReplicateInputDirectly is true on the ability and is generally not a good thing to use. (Instead, prefer to use Generic Client Events).
+	UFUNCTION(Server, reliable, WithValidation)
+	void ServerSetInputPressed(FGameplayAbilitySpecHandle AbilityHandle);
+
+	UFUNCTION(Server, reliable, WithValidation)
+	void ServerSetInputReleased(FGameplayAbilitySpecHandle AbilityHandle);
+
+	/** Called on local player always. Called on server only if bReplicateInputDirectly is set on the GameplayAbility. */
+	void AbilitySpecInputReleased(FGameplayAbilitySpec& Spec);
+
+	/** Called on local player always. Called on server only if bReplicateInputDirectly is set on the GameplayAbility. */
+	void AbilitySpecInputPressed(FGameplayAbilitySpec& Spec);
+	
+	// ---------------------------------------------------------------------
 
 	/** Tasks that run on simulated proxies */
 	UPROPERTY(ReplicatedUsing=OnRep_SimulatedTasks)
diff --git a/Engine/Source/Runtime/GameplayAbilities/Public/GameplayEffect.h b/Engine/Source/Runtime/GameplayAbilities/Public/GameplayEffect.h
--- a/Engine/Source/Runtime/GameplayAbilities/Public/GameplayEffect.h
+++ b/Engine/Source/Runtime/GameplayAbilities/Public/GameplayEffect.h
@@ -20,7 +20,7 @@ class UAbilitySystemComponent;
 class UGameplayModMagnitudeCalculation;
 class UGameplayEffectExecutionCalculation;
 
-DECLARE_DELEGATE_RetVal_OneParam(bool, FActiveGameplayEffectQueryCustomMatch, const FActiveGameplayEffect);
+DECLARE_DELEGATE_RetVal_OneParam(bool, FActiveGameplayEffectQueryCustomMatch, const FActiveGameplayEffect&);
 
 /** Enumeration outlining the possible gameplay effect magnitude calculation policies */
 UENUM()
diff --git a/Engine/Source/Runtime/GameplayAbilities/Public/GameplayPrediction.h b/Engine/Source/Runtime/GameplayAbilities/Public/GameplayPrediction.h
--- a/Engine/Source/Runtime/GameplayAbilities/Public/GameplayPrediction.h
+++ b/Engine/Source/Runtime/GameplayAbilities/Public/GameplayPrediction.h
@@ -23,7 +23,7 @@ DECLARE_DELEGATE(FPredictionKeyEvent);
  *	-Ability activation
  *	-Triggered Events
  *	-GameplayEffect application:
- *		-Attribute modification	(EXCEPTIONS: meta attributes such as Damage/Healing and certain cases of ongoing % based modifiers. Explained below)
+ *		-Attribute modification	(EXCEPTIONS: Executions do not currently predict, only attribute modifiers)
  *		-GameplayTag modification
  *	-Gameplay Cue events (both from within predictive gameplay effect and on their own)
  *		
@@ -40,7 +40,7 @@ DECLARE_DELEGATE(FPredictionKeyEvent);
  *	1. "Can I do this?" Basic protocol for prediction.
  *	2. "Undo" How to undo side effects when a prediction fails.
  *	3. "Redo" How to avoid replaying side effects that we predicted locally but that also get replicated from the server.
- *	4. "Completeness" How to be sue we /really/ predicted all side effects.
+ *	4. "Completeness" How to be sure we /really/ predicted all side effects.
  *	5. "Dependencies" How to manage dependent prediction and chains of predicted events.
  *	6. "Override" How to override state predictively that is otherwise replicated/owned by the server.
  *	
