Merging

//depot/UE4-Orion/Engine/Source/Runtime/GameplayAbilities/*/Abilities/...

to //depot/UE4/Engine/Source/Runtime/GameplayAbilities/*/Abilities/...

diff --git a/Engine/Source/Runtime/GameplayAbilities/Private/Abilities/Tasks/AbilityTask_WaitConfirmCancel.cpp b/Engine/Source/Runtime/GameplayAbilities/Private/Abilities/Tasks/AbilityTask_WaitConfirmCancel.cpp
--- a/Engine/Source/Runtime/GameplayAbilities/Private/Abilities/Tasks/AbilityTask_WaitConfirmCancel.cpp
+++ b/Engine/Source/Runtime/GameplayAbilities/Private/Abilities/Tasks/AbilityTask_WaitConfirmCancel.cpp
@@ -16,18 +16,21 @@ UAbilityTask_WaitConfirmCancel::UAbilityTask_WaitConfirmCancel(const FObjectInit
 
 void UAbilityTask_WaitConfirmCancel::OnConfirmCallback()
 {
-	OnConfirm.Broadcast();
+	
 	if (AbilitySystemComponent.IsValid())
 	{
+		AbilitySystemComponent->ConsumeReplicatedClientEvent(EAbilityReplicatedClientEvent::GenericConfirm, GetAbilitySpecHandle(), GetActivationPredictionKey());
+		OnConfirm.Broadcast();
 		EndTask();
 	}
 }
 
 void UAbilityTask_WaitConfirmCancel::OnCancelCallback()
 {
-	OnCancel.Broadcast();
 	if (AbilitySystemComponent.IsValid())
 	{
+		AbilitySystemComponent->ConsumeReplicatedClientEvent(EAbilityReplicatedClientEvent::GenericCancel, GetAbilitySpecHandle(), GetActivationPredictionKey());
+		OnCancel.Broadcast();
 		EndTask();
 	}
 }
diff --git a/Engine/Source/Runtime/GameplayAbilities/Private/Abilities/Tasks/AbilityTask_WaitInputPress.cpp b/Engine/Source/Runtime/GameplayAbilities/Private/Abilities/Tasks/AbilityTask_WaitInputPress.cpp
--- a/Engine/Source/Runtime/GameplayAbilities/Private/Abilities/Tasks/AbilityTask_WaitInputPress.cpp
+++ b/Engine/Source/Runtime/GameplayAbilities/Private/Abilities/Tasks/AbilityTask_WaitInputPress.cpp
@@ -26,6 +26,10 @@ void UAbilityTask_WaitInputPress::OnPressCallback()
 		// Tell the server about this
 		AbilitySystemComponent->ServerSetReplicatedClientEvent(EAbilityReplicatedClientEvent::InputPressed, GetAbilitySpecHandle(), GetActivationPredictionKey(), AbilitySystemComponent->ScopedPredictionKey);
 	}
+	else
+	{
+		AbilitySystemComponent->ConsumeReplicatedClientEvent(EAbilityReplicatedClientEvent::InputPressed, GetAbilitySpecHandle(), GetActivationPredictionKey());
+	}
 
 	// We are done. Kill us so we don't keep getting broadcast messages
 	OnPress.Broadcast(ElapsedTime);
diff --git a/Engine/Source/Runtime/GameplayAbilities/Private/Abilities/Tasks/AbilityTask_WaitInputRelease.cpp b/Engine/Source/Runtime/GameplayAbilities/Private/Abilities/Tasks/AbilityTask_WaitInputRelease.cpp
--- a/Engine/Source/Runtime/GameplayAbilities/Private/Abilities/Tasks/AbilityTask_WaitInputRelease.cpp
+++ b/Engine/Source/Runtime/GameplayAbilities/Private/Abilities/Tasks/AbilityTask_WaitInputRelease.cpp
@@ -26,6 +26,10 @@ void UAbilityTask_WaitInputRelease::OnReleaseCallback()
 		// Tell the server about this
 		AbilitySystemComponent->ServerSetReplicatedClientEvent(EAbilityReplicatedClientEvent::InputReleased, GetAbilitySpecHandle(), GetActivationPredictionKey(), AbilitySystemComponent->ScopedPredictionKey);
 	}
+	else
+	{
+		AbilitySystemComponent->ConsumeReplicatedClientEvent(EAbilityReplicatedClientEvent::InputReleased, GetAbilitySpecHandle(), GetActivationPredictionKey());
+	}
 
 	// We are done. Kill us so we don't keep getting broadcast messages
 	OnRelease.Broadcast(ElapsedTime);
diff --git a/Engine/Source/Runtime/GameplayAbilities/Public/Abilities/GameplayAbility.h b/Engine/Source/Runtime/GameplayAbilities/Public/Abilities/GameplayAbility.h
--- a/Engine/Source/Runtime/GameplayAbilities/Public/Abilities/GameplayAbility.h
+++ b/Engine/Source/Runtime/GameplayAbilities/Public/Abilities/GameplayAbility.h
@@ -635,7 +635,7 @@ protected:
 	//
 	// ----------------------------------------------------------------------------------------------------------------
 	
-	/** Abilities with these tags are canceled when this ability is executed */
+	/** Abilities with these tags are cancelled when this ability is executed */
 	UPROPERTY(EditDefaultsOnly, Category = Tags)
 	FGameplayTagContainer CancelAbilitiesWithTag;
 
