fix bad merge

diff --git a/Engine/Source/Runtime/GameplayAbilities/Private/Abilities/Tasks/AbilityTask_WaitConfirmCancel.cpp b/Engine/Source/Runtime/GameplayAbilities/Private/Abilities/Tasks/AbilityTask_WaitConfirmCancel.cpp
--- a/Engine/Source/Runtime/GameplayAbilities/Private/Abilities/Tasks/AbilityTask_WaitConfirmCancel.cpp
+++ b/Engine/Source/Runtime/GameplayAbilities/Private/Abilities/Tasks/AbilityTask_WaitConfirmCancel.cpp
@@ -19,7 +19,7 @@ void UAbilityTask_WaitConfirmCancel::OnConfirmCallback()
 	
 	if (AbilitySystemComponent.IsValid())
 	{
-		AbilitySystemComponent->ConsumeReplicatedClientEvent(EAbilityReplicatedClientEvent::GenericConfirm, GetAbilitySpecHandle(), GetActivationPredictionKey());
+		AbilitySystemComponent->ConsumeClientReplicatedEvent(EAbilityReplicatedClientEvent::GenericConfirm, GetAbilitySpecHandle(), GetActivationPredictionKey());
 		OnConfirm.Broadcast();
 		EndTask();
 	}
@@ -29,7 +29,7 @@ void UAbilityTask_WaitConfirmCancel::OnCancelCallback()
 {
 	if (AbilitySystemComponent.IsValid())
 	{
-		AbilitySystemComponent->ConsumeReplicatedClientEvent(EAbilityReplicatedClientEvent::GenericCancel, GetAbilitySpecHandle(), GetActivationPredictionKey());
+		AbilitySystemComponent->ConsumeClientReplicatedEvent(EAbilityReplicatedClientEvent::GenericCancel, GetAbilitySpecHandle(), GetActivationPredictionKey());
 		OnCancel.Broadcast();
 		EndTask();
 	}
diff --git a/Engine/Source/Runtime/GameplayAbilities/Private/Abilities/Tasks/AbilityTask_WaitInputPress.cpp b/Engine/Source/Runtime/GameplayAbilities/Private/Abilities/Tasks/AbilityTask_WaitInputPress.cpp
--- a/Engine/Source/Runtime/GameplayAbilities/Private/Abilities/Tasks/AbilityTask_WaitInputPress.cpp
+++ b/Engine/Source/Runtime/GameplayAbilities/Private/Abilities/Tasks/AbilityTask_WaitInputPress.cpp
@@ -28,7 +28,7 @@ void UAbilityTask_WaitInputPress::OnPressCallback()
 	}
 	else
 	{
-		AbilitySystemComponent->ConsumeReplicatedClientEvent(EAbilityReplicatedClientEvent::InputPressed, GetAbilitySpecHandle(), GetActivationPredictionKey());
+		AbilitySystemComponent->ConsumeClientReplicatedEvent(EAbilityReplicatedClientEvent::InputPressed, GetAbilitySpecHandle(), GetActivationPredictionKey());
 	}
 
 	// We are done. Kill us so we don't keep getting broadcast messages
diff --git a/Engine/Source/Runtime/GameplayAbilities/Private/Abilities/Tasks/AbilityTask_WaitInputRelease.cpp b/Engine/Source/Runtime/GameplayAbilities/Private/Abilities/Tasks/AbilityTask_WaitInputRelease.cpp
--- a/Engine/Source/Runtime/GameplayAbilities/Private/Abilities/Tasks/AbilityTask_WaitInputRelease.cpp
+++ b/Engine/Source/Runtime/GameplayAbilities/Private/Abilities/Tasks/AbilityTask_WaitInputRelease.cpp
@@ -28,7 +28,7 @@ void UAbilityTask_WaitInputRelease::OnReleaseCallback()
 	}
 	else
 	{
-		AbilitySystemComponent->ConsumeReplicatedClientEvent(EAbilityReplicatedClientEvent::InputReleased, GetAbilitySpecHandle(), GetActivationPredictionKey());
+		AbilitySystemComponent->ConsumeClientReplicatedEvent(EAbilityReplicatedClientEvent::InputReleased, GetAbilitySpecHandle(), GetActivationPredictionKey());
 	}
 
 	// We are done. Kill us so we don't keep getting broadcast messages
