non unity fix

diff --git a/Engine/Source/Runtime/GameplayAbilities/Private/Abilities/Tasks/AbilityTask.cpp b/Engine/Source/Runtime/GameplayAbilities/Private/Abilities/Tasks/AbilityTask.cpp
--- a/Engine/Source/Runtime/GameplayAbilities/Private/Abilities/Tasks/AbilityTask.cpp
+++ b/Engine/Source/Runtime/GameplayAbilities/Private/Abilities/Tasks/AbilityTask.cpp
@@ -182,4 +182,9 @@ bool UAbilityTask::IsForRemoteClient() const
 	}
 
 	return false;
-}
\ No newline at end of file
+}
+
+bool UAbilityTask::CallOrAddReplicatedDelegate(EAbilityReplicatedClientEvent::Type Event, FSimpleMulticastDelegate::FDelegate Delegate)
+{
+	return AbilitySystemComponent->CallOrAddReplicatedDelegate(Event, GetAbilitySpecHandle(), GetActivationPredictionKey(), Delegate);
+}
diff --git a/Engine/Source/Runtime/GameplayAbilities/Public/Abilities/Tasks/AbilityTask.h b/Engine/Source/Runtime/GameplayAbilities/Public/Abilities/Tasks/AbilityTask.h
--- a/Engine/Source/Runtime/GameplayAbilities/Public/Abilities/Tasks/AbilityTask.h
+++ b/Engine/Source/Runtime/GameplayAbilities/Public/Abilities/Tasks/AbilityTask.h
@@ -154,10 +154,7 @@ protected:
 
 	/** Helper method for registering client replicated callbacks */
 	
-	bool CallOrAddReplicatedDelegate(EAbilityReplicatedClientEvent::Type Event, FSimpleMulticastDelegate::FDelegate Delegate)
-	{
-		return AbilitySystemComponent->CallOrAddReplicatedDelegate(Event, GetAbilitySpecHandle(), GetActivationPredictionKey(), Delegate);
-	}
+	bool CallOrAddReplicatedDelegate(EAbilityReplicatedClientEvent::Type Event, FSimpleMulticastDelegate::FDelegate Delegate);
 };
 
 //For searching through lists of ability instances
