Dev to Main integration from CL #2446290

diff --git a/Engine/Source/Runtime/Engine/Private/NetworkDriver.cpp b/Engine/Source/Runtime/Engine/Private/NetworkDriver.cpp
--- a/Engine/Source/Runtime/Engine/Private/NetworkDriver.cpp
+++ b/Engine/Source/Runtime/Engine/Private/NetworkDriver.cpp
@@ -1747,12 +1747,12 @@ bool FPacketSimulationSettings::ParseSettings(const TCHAR* Cmd)
 
 FNetViewer::FNetViewer(UNetConnection* InConnection, float DeltaSeconds) :
 	InViewer(InConnection->PlayerController),
-	Viewer(InConnection->Viewer),
+	ViewTarget(InConnection->ViewTarget),
 	ViewLocation(ForceInit),
 	ViewDir(ForceInit)
 {
 	// Get viewer coordinates.
-	ViewLocation = Viewer->GetActorLocation();
+	ViewLocation = ViewTarget->GetActorLocation();
 	if (InViewer)
 	{
 		FRotator ViewRotation = InViewer->GetControlRotation();
@@ -1765,8 +1765,8 @@ FNetViewer::FNetViewer(UNetConnection* InConnection, float DeltaSeconds) :
 	if (InConnection->TickCount & 1)
 	{
 		float PredictSeconds = (InConnection->TickCount & 2) ? 0.4f : 0.9f;
-		Ahead = PredictSeconds * Viewer->GetVelocity();
-		APawn* ViewerPawn = Cast<APawn>(Viewer);
+		Ahead = PredictSeconds * ViewTarget->GetVelocity();
+		APawn* ViewerPawn = Cast<APawn>(ViewTarget);
 		if( ViewerPawn && ViewerPawn->GetMovementBase() && ViewerPawn->GetMovementBase()->GetOwner() )
 		{
 			Ahead += PredictSeconds * ViewerPawn->GetMovementBase()->GetOwner()->GetVelocity();
@@ -1787,7 +1787,7 @@ FNetViewer::FNetViewer(UNetConnection* InConnection, float DeltaSeconds) :
 				World = ViewerPawn->GetWorld();
 			}
 			check( World );
-			World->LineTraceSingle(Hit, ViewLocation, Hit.Location, FCollisionQueryParams(NAME_ServerForwardView, true, Viewer), FCollisionObjectQueryParams(ECC_WorldStatic));
+			World->LineTraceSingle(Hit, ViewLocation, Hit.Location, FCollisionQueryParams(NAME_ServerForwardView, true, ViewTarget), FCollisionObjectQueryParams(ECC_WorldStatic));
 			ViewLocation = Hit.Location;
 		}
 	}
@@ -1801,7 +1801,7 @@ FActorPriority::FActorPriority(UNetConnection* InConnection, UActorChannel* InCh
 	Priority = 0;
 	for (int32 i = 0; i < Viewers.Num(); i++)
 	{
-		Priority = FMath::Max<int32>(Priority, FMath::RoundToInt(65536.0f * Actor->GetNetPriority(Viewers[i].ViewLocation, Viewers[i].ViewDir, Viewers[i].InViewer, InChannel, Time, bLowBandwidth)));
+		Priority = FMath::Max<int32>(Priority, FMath::RoundToInt(65536.0f * Actor->GetNetPriority(Viewers[i].ViewLocation, Viewers[i].ViewDir, Viewers[i].InViewer, Viewers[i].ViewTarget, InChannel, Time, bLowBandwidth)));
 	}
 }
 
@@ -1909,7 +1909,7 @@ int32 UNetDriver::ServerReplicateActors(float DeltaSeconds)
 			}
 			bFoundReadyConnection = true;
 			
-			Connection->Viewer = Connection->PlayerController ? Connection->PlayerController->GetViewTarget() : OwningActor->GetOwner();
+			Connection->ViewTarget = Connection->PlayerController ? Connection->PlayerController->GetViewTarget() : OwningActor->GetOwner();
 			//@todo - eliminate this mallocs if the connection isn't going to actually be updated this frame (currently needed to verify owner relevancy below)
 			Connection->OwnedConsiderList.Empty(NetRelevantActorCount);
 
@@ -1919,21 +1919,21 @@ int32 UNetDriver::ServerReplicateActors(float DeltaSeconds)
 				APlayerController* ChildPlayerController = Child->PlayerController;
 				if (ChildPlayerController != NULL)
 				{
-					Child->Viewer = ChildPlayerController->GetViewTarget();
+					Child->ViewTarget = ChildPlayerController->GetViewTarget();
 					Child->OwnedConsiderList.Empty(NetRelevantActorCount);
 				}
 				else
 				{
-					Child->Viewer = NULL;
+					Child->ViewTarget = NULL;
 				}
 			}
 		}
 		else
 		{
-			Connection->Viewer = NULL;
+			Connection->ViewTarget = NULL;
 			for (int32 ChildIdx = 0; ChildIdx < Connection->Children.Num(); ChildIdx++)
 			{
-				Connection->Children[ChildIdx]->Viewer = NULL;
+				Connection->Children[ChildIdx]->ViewTarget = NULL;
 			}
 		}
 	}
@@ -2078,11 +2078,11 @@ int32 UNetDriver::ServerReplicateActors(float DeltaSeconds)
 							bool bCloseChannel = true;
 							while (Connection != NULL)
 							{
-								if (Connection->Viewer != NULL)
+								if (Connection->ViewTarget != NULL)
 								{
 									if (ActorOwner == Connection->PlayerController || 
 										(Connection->PlayerController && ActorOwner == Connection->PlayerController->GetPawn()) ||
-										Connection->Viewer->IsRelevancyOwnerFor(Actor, ActorOwner, Connection->OwningActor))
+										Connection->ViewTarget->IsRelevancyOwnerFor(Actor, ActorOwner, Connection->OwningActor))
 									{
 										// For performance reasons, make sure we don't resize the array. It should already be appropriately sized above!
 										ensure(Connection->OwnedConsiderList.Num() < Connection->OwnedConsiderList.Max());
@@ -2174,7 +2174,7 @@ int32 UNetDriver::ServerReplicateActors(float DeltaSeconds)
 				}
 			}
 		}
-		else if (Connection->Viewer)
+		else if (Connection->ViewTarget)
 		{
 			int32 j;
 			int32 ConsiderCount	= 0;
@@ -2225,7 +2225,7 @@ int32 UNetDriver::ServerReplicateActors(float DeltaSeconds)
 				new(ConnectionViewers) FNetViewer(Connection, DeltaSeconds);
 				for (j = 0; j < Connection->Children.Num(); j++)
 				{
-					if (Connection->Children[j]->Viewer != NULL)
+					if (Connection->Children[j]->ViewTarget != NULL)
 					{
 						new(ConnectionViewers) FNetViewer(Connection->Children[j], DeltaSeconds);
 					}
@@ -2240,7 +2240,7 @@ int32 UNetDriver::ServerReplicateActors(float DeltaSeconds)
 
 				// determine whether we should priority sort the list of relevant actors based on the saturation/bandwidth of the current connection
 				//@note - if the server is currently CPU saturated then do not sort until framerate improves
-				check(World == Connection->Viewer->GetWorld());
+				check(World == Connection->ViewTarget->GetWorld());
 				AGameMode const* const GameMode = World->GetAuthGameMode();
 				bool bLowNetBandwidth = !bCPUSaturated && (Connection->CurrentNetSpeed / float(GameMode->NumPlayers + GameMode->NumBots) < 500.f );
 
@@ -2278,7 +2278,7 @@ int32 UNetDriver::ServerReplicateActors(float DeltaSeconds)
 								float LastReplicationTime  = Channel ? (Connection->Driver->Time - Channel->LastUpdateTime) : Connection->Driver->SpawnPrioritySeconds;
 								for (int32 viewerIdx = 0; viewerIdx < ConnectionViewers.Num(); viewerIdx++)
 								{
-									if (!Actor->GetNetDormancy(ConnectionViewers[viewerIdx].ViewLocation, ConnectionViewers[viewerIdx].ViewDir, ConnectionViewers[viewerIdx].InViewer, Channel, LastReplicationTime, bLowNetBandwidth))
+									if (!Actor->GetNetDormancy(ConnectionViewers[viewerIdx].ViewLocation, ConnectionViewers[viewerIdx].ViewDir, ConnectionViewers[viewerIdx].InViewer, ConnectionViewers[viewerIdx].ViewTarget, Channel, Time, bLowNetBandwidth))
 									{
 										ShouldGoDormant = false;
 										break;
@@ -2309,7 +2309,7 @@ int32 UNetDriver::ServerReplicateActors(float DeltaSeconds)
 						bool Relevant = false;
 						for (int32 viewerIdx = 0; viewerIdx < ConnectionViewers.Num(); viewerIdx++)
 						{
-							if(Actor->IsNetRelevantFor(ConnectionViewers[viewerIdx].InViewer, ConnectionViewers[viewerIdx].Viewer, ConnectionViewers[viewerIdx].ViewLocation))
+							if(Actor->IsNetRelevantFor(ConnectionViewers[viewerIdx].InViewer, ConnectionViewers[viewerIdx].ViewTarget, ConnectionViewers[viewerIdx].ViewLocation))
 							{
 								Relevant = true;
 								break;
@@ -2447,7 +2447,7 @@ int32 UNetDriver::ServerReplicateActors(float DeltaSeconds)
 							{
 								for (int32 k = 0; k < ConnectionViewers.Num(); k++)
 								{
-									if (Actor->IsNetRelevantFor(ConnectionViewers[k].InViewer, ConnectionViewers[k].Viewer, ConnectionViewers[k].ViewLocation))
+									if (Actor->IsNetRelevantFor(ConnectionViewers[k].InViewer, ConnectionViewers[k].ViewTarget, ConnectionViewers[k].ViewLocation))
 									{
 										bIsRelevant = true;
 										break;
@@ -2582,7 +2582,7 @@ int32 UNetDriver::ServerReplicateActors(float DeltaSeconds)
 				{
 					for (int32 h = 0; h < ConnectionViewers.Num(); h++)
 					{
-						if (Actor->IsNetRelevantFor(ConnectionViewers[h].InViewer, ConnectionViewers[h].Viewer, ConnectionViewers[h].ViewLocation))
+						if (Actor->IsNetRelevantFor(ConnectionViewers[h].InViewer, ConnectionViewers[h].ViewTarget, ConnectionViewers[h].ViewLocation))
 						{
 							UE_LOG(LogNetTraffic, Log, TEXT(" Saturated. Mark %s NetUpdateTime to be checked for next tick"), *Actor->GetName());
 							Actor->bPendingNetUpdate = true;
diff --git a/Engine/Source/Runtime/Engine/Private/Pawn.cpp b/Engine/Source/Runtime/Engine/Private/Pawn.cpp
--- a/Engine/Source/Runtime/Engine/Private/Pawn.cpp
+++ b/Engine/Source/Runtime/Engine/Private/Pawn.cpp
@@ -251,34 +251,6 @@ FVector APawn::GetGravityDirection()
 	return FVector(0.f,0.f,-1.f);
 }
 
-
-float APawn::GetNetPriority(const FVector& ViewPos, const FVector& ViewDir, APlayerController* Viewer, UActorChannel* InChannel, float Time, bool bLowBandwidth)
-{
-	if ( !bHidden )
-	{
-		FVector Dir = GetActorLocation() - ViewPos;
-		float DistSq = Dir.SizeSquared();
-
-		// adjust priority based on distance and whether pawn is in front of viewer or is controlled
-		if ( (ViewDir | Dir) < 0.f )
-		{
-			if ( DistSq > NEARSIGHTTHRESHOLDSQUARED )
-				Time *= 0.3f;
-			else if ( DistSq > CLOSEPROXIMITYSQUARED )
-				Time *= 0.5f;
-		}
-		else if ( Controller && (DistSq < FARSIGHTTHRESHOLDSQUARED) && (FMath::Square(ViewDir | Dir) > 0.5f * DistSq) )
-		{
-			Time *= 2.f;
-		}
-		else if (DistSq > MEDSIGHTTHRESHOLDSQUARED)
-		{
-			Time *= 0.5f;
-		}
-	}
-	return NetPriority * Time;
-}
-
 bool APawn::ShouldTickIfViewportsOnly() const 
 { 
 	return IsLocallyControlled() && Cast<APlayerController>(GetController()); 
@@ -1045,10 +1017,10 @@ bool APawn::IsBasedOnActor(const AActor* Other) const
 }
 
 
-bool APawn::IsNetRelevantFor(const APlayerController* RealViewer, const AActor* Viewer, const FVector& SrcLocation) const
+bool APawn::IsNetRelevantFor(const APlayerController* RealViewer, const AActor* ViewTarget, const FVector& SrcLocation) const
 {
-	if( bAlwaysRelevant || RealViewer == Controller || IsOwnedBy(Viewer) || IsOwnedBy(RealViewer) || this==Viewer || Viewer==Instigator
-		|| IsBasedOnActor(Viewer) || (Viewer && Viewer->IsBasedOnActor(this)))
+	if (bAlwaysRelevant || RealViewer == Controller || IsOwnedBy(ViewTarget) || IsOwnedBy(RealViewer) || this == ViewTarget || ViewTarget == Instigator
+		|| IsBasedOnActor(ViewTarget) || (ViewTarget && ViewTarget->IsBasedOnActor(this)))
 	{
 		return true;
 	}
@@ -1062,7 +1034,7 @@ bool APawn::IsNetRelevantFor(const APlayerController* RealViewer, const AActor*
 		AActor* BaseActor = MovementBase ? MovementBase->GetOwner() : NULL;
 		if ( MovementBase && BaseActor && GetMovementComponent() && ((Cast<const USkeletalMeshComponent>(MovementBase)) || (BaseActor == GetOwner())) )
 		{
-			return BaseActor->IsNetRelevantFor( RealViewer, Viewer, SrcLocation );
+			return BaseActor->IsNetRelevantFor(RealViewer, ViewTarget, SrcLocation);
 		}
 	}
 
diff --git a/Engine/Source/Runtime/Engine/Private/PhysicsEngine/BodyInstance.cpp b/Engine/Source/Runtime/Engine/Private/PhysicsEngine/BodyInstance.cpp
--- a/Engine/Source/Runtime/Engine/Private/PhysicsEngine/BodyInstance.cpp
+++ b/Engine/Source/Runtime/Engine/Private/PhysicsEngine/BodyInstance.cpp
@@ -4147,7 +4147,6 @@ float FBodyInstance::GetDistanceToBody(const FVector& Point, FVector& OutPointOn
 	int32 NumShapes = RigidBody->getShapes(PShapes.GetData(), PShapes.Num());
 
 	const PxVec3 PPoint = U2PVector(Point);
-	PxTransform PCompTM(U2PTransform(OwnerComponent->ComponentToWorld));
 	float MinDistanceSqr=BIG_NUMBER;
 	bool bFoundValidBody = false;
 
@@ -4157,7 +4156,7 @@ float FBodyInstance::GetDistanceToBody(const FVector& Point, FVector& OutPointOn
 		PxShape* PShape = PShapes[ShapeIdx];
 		check(PShape);
 		PxGeometry& PGeom = PShape->getGeometry().any();
-		PxTransform PGlobalPose = PCompTM.transform(PShape->getLocalPose());
+		PxTransform PGlobalPose = PxShapeExt::getGlobalPose(*PShape, *RigidBody);
 		PxGeometryType::Enum GeomType = PShape->getGeometryType();
 		
 		if (GeomType == PxGeometryType::eTRIANGLEMESH)
diff --git a/Engine/Source/Runtime/Engine/Private/PlayerController.cpp b/Engine/Source/Runtime/Engine/Private/PlayerController.cpp
--- a/Engine/Source/Runtime/Engine/Private/PlayerController.cpp
+++ b/Engine/Source/Runtime/Engine/Private/PlayerController.cpp
@@ -35,6 +35,9 @@
 #include "GameFramework/GameState.h"
 #include "GameFramework/GameMode.h"
 #include "Engine/ChildConnection.h"
+#include "Engine/GameEngine.h"
+#include "Engine/GameInstance.h"
+#include "VisualLogger/VisualLogger.h"
 
 DEFINE_LOG_CATEGORY(LogPlayerController);
 
@@ -47,7 +50,6 @@ const float RetryServerCheckSpectatorThrottleTime = 0.25f;
 
 APlayerController::APlayerController(const FObjectInitializer& ObjectInitializer)
 	: Super(ObjectInitializer)
-	, SlateOperations(new FReply( FReply::Unhandled() ))
 {
 	NetPriority = 3.0f;
 	CheatClass = UCheatManager::StaticClass();
@@ -80,7 +82,7 @@ APlayerController::APlayerController(const FObjectInitializer& ObjectInitializer
 	}
 }
 
-float APlayerController::GetNetPriority(const FVector& ViewPos, const FVector& ViewDir, APlayerController* Viewer, UActorChannel* InChannel, float Time, bool bLowBandwidth)
+float APlayerController::GetNetPriority(const FVector& ViewPos, const FVector& ViewDir, APlayerController* Viewer, AActor* ViewTarget, UActorChannel* InChannel, float Time, bool bLowBandwidth)
 {
 	if ( Viewer == this )
 	{
@@ -1167,13 +1169,16 @@ void APlayerController::ClientSetHUD_Implementation(TSubclassOf<AHUD> NewHUDClas
 {
 	if ( MyHUD != NULL )
 	{
+
 		MyHUD->Destroy();
 		MyHUD = NULL;
 	}
+
 	FActorSpawnParameters SpawnInfo;
 	SpawnInfo.Owner = this;
 	SpawnInfo.Instigator = Instigator;
 	SpawnInfo.ObjectFlags |= RF_Transient;	// We never want to save HUDs into a map
+
 	MyHUD = GetWorld()->SpawnActor<AHUD>(NewHUDClass, SpawnInfo );
 }
 
@@ -3615,8 +3620,10 @@ void APlayerController::SetPlayer( UPlayer* InPlayer )
 	check(InPlayer!=NULL);
 
 	// Detach old player.
-	if( InPlayer->PlayerController )
+	if (InPlayer->PlayerController)
+	{
 		InPlayer->PlayerController->Player = NULL;
+	}
 
 	// Set the viewport.
 	Player = InPlayer;
@@ -3633,6 +3640,8 @@ void APlayerController::SetPlayer( UPlayer* InPlayer )
 	ULocalPlayer *LP = Cast<ULocalPlayer>(InPlayer);
 	if (LP != NULL)
 	{
+		// Clients need this marked as local (server already knew at construction time)
+		SetAsLocalPlayerController();
 		LP->InitOnlineSession();
 		InitInputSystem();
 	}
@@ -3802,7 +3811,7 @@ void APlayerController::TickActor( float DeltaSeconds, ELevelTick TickType, FAct
 #endif // !(UE_BUILD_SHIPPING || UE_BUILD_TEST)
 }
 
-bool APlayerController::IsNetRelevantFor(const APlayerController* RealViewer, const AActor* Viewer, const FVector& SrcLocation) const
+bool APlayerController::IsNetRelevantFor(const APlayerController* RealViewer, const AActor* ViewTarget, const FVector& SrcLocation) const
 {
 	return ( this==RealViewer );
 }
@@ -4367,9 +4376,10 @@ void FInputModeGameOnly::ApplyInputMode(FReply& SlateOperations, class UGameView
 void APlayerController::SetInputMode(const FInputModeDataBase& InData)
 {
 	UGameViewportClient* GameViewportClient = GetWorld()->GetGameViewport();
-	if (GameViewportClient != nullptr)
+	ULocalPlayer* LocalPlayer = Cast< ULocalPlayer >( Player );
+	if ( GameViewportClient && LocalPlayer )
 	{
-		InData.ApplyInputMode(*SlateOperations, *GameViewportClient);
+		InData.ApplyInputMode( LocalPlayer->GetSlateOperations(), *GameViewportClient );
 	}
 }
 
@@ -4447,3 +4457,11 @@ void FDynamicForceFeedbackDetails::Update(FForceFeedbackValues& Values) const
 		Values.RightSmall = FMath::Clamp(Intensity, Values.RightSmall, 1.f);
 	}
 }
+
+void APlayerController::OnServerStartedVisualLogger_Implementation(bool bIsLogging)
+{
+#if ENABLE_VISUAL_LOG
+	FVisualLogger::Get().SetIsRecordingOnServer(bIsLogging);
+	ClientMessage(FString::Printf(TEXT("Visual Loggger is %s."), FVisualLogger::Get().IsRecordingOnServer() ? TEXT("now recording") : TEXT("disabled")));
+#endif
+}
diff --git a/Engine/Source/Runtime/Engine/Private/Replication.cpp b/Engine/Source/Runtime/Engine/Private/Replication.cpp
--- a/Engine/Source/Runtime/Engine/Private/Replication.cpp
+++ b/Engine/Source/Runtime/Engine/Private/Replication.cpp
@@ -36,12 +36,12 @@ static bool				SavedbReversePlayback;
 static float			SavedPosition;
 static uint8			SavedReplicationForceIsPlaying;
 
-float AMatineeActor::GetNetPriority(const FVector& ViewPos, const FVector& ViewDir, APlayerController* Viewer, UActorChannel* InChannel, float Time, bool bLowBandwidth)
+float AMatineeActor::GetNetPriority(const FVector& ViewPos, const FVector& ViewDir, APlayerController* Viewer, AActor* ViewTarget, UActorChannel* InChannel, float Time, bool bLowBandwidth)
 {
-	float Result = Super::GetNetPriority(ViewPos, ViewDir, Viewer, InChannel, Time, bLowBandwidth);
+	float Result = Super::GetNetPriority(ViewPos, ViewDir, Viewer, ViewTarget, InChannel, Time, bLowBandwidth);
 	// attempt to replicate MatineeActors approximately in the order that they were spawned to reduce ordering issues
-	// when LDs make multipler matinees affect the same target(s)
-	// not great, but without a full depedancy setup this is the best we can do
+	// when LDs make multiple matinees affect the same target(s)
+	// not great, but without a full dependency setup this is the best we can do
 	if (InChannel == NULL)
 	{
 		Result += 1.0f - CreationTime / GetWorld()->TimeSeconds;
diff --git a/Engine/Source/Runtime/Engine/Private/SkeletalMeshComponentPhysics.cpp b/Engine/Source/Runtime/Engine/Private/SkeletalMeshComponentPhysics.cpp
--- a/Engine/Source/Runtime/Engine/Private/SkeletalMeshComponentPhysics.cpp
+++ b/Engine/Source/Runtime/Engine/Private/SkeletalMeshComponentPhysics.cpp
@@ -1747,9 +1747,46 @@ FVector USkeletalMeshComponent::GetSkinnedVertexPosition(int32 VertexIndex) cons
 
 extern float DebugLineLifetime;
 
+float USkeletalMeshComponent::GetDistanceToCollision(const FVector& Point, FVector& ClosestPointOnCollision) const
+{
+	ClosestPointOnCollision = Point;
+	float ClosestPointDistance = -1.f;
+	bool bHasResult = false;
+
+	for (int32 BodyIdx = 0; BodyIdx < Bodies.Num(); ++BodyIdx)
+	{
+		FBodyInstance* BodyInstance = Bodies[BodyIdx];
+		if (BodyInstance && BodyInstance->IsValidBodyInstance() && (BodyInstance->GetCollisionEnabled() != ECollisionEnabled::NoCollision))
+		{
+			FVector ClosestPoint;
+			const float Distance = Bodies[BodyIdx]->GetDistanceToBody(Point, ClosestPoint);
+
+			if (Distance < 0.f)
+			{
+				// Invalid result, impossible to be better than ClosestPointDistance
+				continue;
+			}
+
+			if (!bHasResult || (Distance < ClosestPointDistance))
+			{
+				bHasResult = true;
+				ClosestPointDistance = Distance;
+				ClosestPointOnCollision = ClosestPoint;
+
+				// If we're inside collision, we're not going to find anything better, so abort search we've got our best find.
+				if (Distance <= KINDA_SMALL_NUMBER)
+				{
+					break;
+				}
+			}
+		}
+	}
+
+	return ClosestPointDistance;
+}
+
 bool USkeletalMeshComponent::LineTraceComponent(struct FHitResult& OutHit, const FVector Start, const FVector End, const struct FCollisionQueryParams& Params)
 {
-	UPhysicsAsset* const PhysicsAsset = GetPhysicsAsset();
 	UWorld* const World = GetWorld();
 	bool bHaveHit = false;
 
diff --git a/Engine/Source/Runtime/Engine/Private/SkeletalMeshMerge.cpp b/Engine/Source/Runtime/Engine/Private/SkeletalMeshMerge.cpp
--- a/Engine/Source/Runtime/Engine/Private/SkeletalMeshMerge.cpp
+++ b/Engine/Source/Runtime/Engine/Private/SkeletalMeshMerge.cpp
@@ -131,7 +131,15 @@ bool FSkeletalMeshMerge::FinalizeMesh()
 			{
 				FName SrcBoneName = SrcMesh->RefSkeleton.GetBoneName(i);
 				int32 DestBoneIndex = NewRefSkeleton.FindBoneIndex(SrcBoneName);
-				check(DestBoneIndex != INDEX_NONE); // All bones should be here now!
+
+				if (DestBoneIndex == INDEX_NONE)
+				{
+					// Missing bones shouldn't be possible, but can happen with invalid meshes;
+					// map any bone we are missing to the 'root'.
+
+					DestBoneIndex = 0;
+				}
+
 				MeshInfo.SrcToDestRefSkeletonMap[i] = DestBoneIndex;
 			}
 		}
@@ -250,10 +258,10 @@ void FSkeletalMeshMerge::GenerateNewSectionArray( TArray<FNewSectionInfo>& NewSe
 	{
 		// source mesh
 		USkeletalMesh* SrcMesh = SrcMeshList[MeshIdx];
-		FSkeletalMeshResource* SrcResource = SrcMesh->GetImportedResource();
 
 		if( SrcMesh )
 		{
+			FSkeletalMeshResource* SrcResource = SrcMesh->GetImportedResource();
 			int32 SourceLODIdx = FMath::Min(LODIdx, SrcResource->LODModels.Num()-1);
 			FStaticLODModel& SrcLODModel = SrcResource->LODModels[SourceLODIdx];
 			FSkeletalMeshLODInfo& SrcLODInfo = SrcMesh->LODInfo[SourceLODIdx];
@@ -466,8 +474,11 @@ void FSkeletalMeshMerge::GenerateLODModel( int32 LODIdx )
 			{
 				FName SrcLODBoneName = MergeSectionInfo.SkelMesh->RefSkeleton.GetBoneName( SrcLODModel.RequiredBones[Idx] );
 				int32 MergeBoneIndex = NewRefSkeleton.FindBoneIndex(SrcLODBoneName);
-				check(MergeBoneIndex != INDEX_NONE);
-				MergeLODModel.RequiredBones.AddUnique(MergeBoneIndex);
+				
+				if (MergeBoneIndex != INDEX_NONE)
+				{
+					MergeLODModel.RequiredBones.AddUnique(MergeBoneIndex);
+				}
 			}
 
 			// keep track of the max number of influences used by the vertices of the chunk
@@ -731,16 +742,21 @@ void FSkeletalMeshMerge::OverrideReferenceSkeletonPose(const TArray<FRefPoseOver
 		const FRefPoseOverride& PoseOverride = PoseOverrides[i];
 		const FReferenceSkeleton& SourceSkeleton = PoseOverride.SkeletalMesh->RefSkeleton;
 
-		for (int32 j = 0, BoneMax = PoseOverride.BoneNames.Num(); j < BoneMax; ++j)
+		for (int32 j = 0, BoneMax = PoseOverride.Overrides.Num(); j < BoneMax; ++j)
 		{
-			const FName& BoneName = PoseOverride.BoneNames[j];
+			const FName& BoneName = PoseOverride.Overrides[j].BoneName;
 			int32 SourceBoneIndex = SourceSkeleton.FindBoneIndex(BoneName);
 
 			if (SourceBoneIndex != INDEX_NONE)
 			{
-				OverrideReferenceBonePose(SourceBoneIndex, SourceSkeleton, TargetSkeleton);
+				bool bOverrideBone = (PoseOverride.Overrides[j].OverrideMode == FRefPoseOverride::ChildrenOnly) ? false : true;
 
-				bool bOverrideChildren = PoseOverride.OverrideChildren[j];
+				if (bOverrideBone)
+				{
+					OverrideReferenceBonePose(SourceBoneIndex, SourceSkeleton, TargetSkeleton);
+				}
+
+				bool bOverrideChildren = (PoseOverride.Overrides[j].OverrideMode == FRefPoseOverride::BoneOnly) ? false : true;
 
 				if (bOverrideChildren)
 				{
@@ -885,17 +901,22 @@ void FSkeletalMeshMerge::OverrideMergedSockets(const TArray<FRefPoseOverride>& P
 		const TArray<USkeletalMeshSocket*>& SkeletonSocketList = PoseOverride.SkeletalMesh->Skeleton->Sockets;
 		const TArray<USkeletalMeshSocket*>& MeshSocketList = const_cast<USkeletalMesh*>(PoseOverride.SkeletalMesh)->GetMeshOnlySocketList();
 
-		for (int32 j = 0, BoneMax = PoseOverride.BoneNames.Num(); j < BoneMax; ++j)
+		for (int32 j = 0, BoneMax = PoseOverride.Overrides.Num(); j < BoneMax; ++j)
 		{
-			const FName& BoneName = PoseOverride.BoneNames[j];
+			const FName& BoneName = PoseOverride.Overrides[j].BoneName;
 			int32 SourceBoneIndex = SourceSkeleton.FindBoneIndex(BoneName);
 
 			if (SourceBoneIndex != INDEX_NONE)
 			{
-				OverrideBoneSockets(BoneName, SkeletonSocketList);
-				OverrideBoneSockets(BoneName, MeshSocketList);
+				bool bOverrideBone = (PoseOverride.Overrides[j].OverrideMode == FRefPoseOverride::ChildrenOnly) ? false : true;
+
+				if (bOverrideBone)
+				{
+					OverrideBoneSockets(BoneName, SkeletonSocketList);
+					OverrideBoneSockets(BoneName, MeshSocketList);
+				}
 
-				bool bOverrideChildren = PoseOverride.OverrideChildren[j];
+				bool bOverrideChildren = (PoseOverride.Overrides[j].OverrideMode == FRefPoseOverride::BoneOnly) ? false : true;
 
 				if (bOverrideChildren)
 				{
diff --git a/Engine/Source/Runtime/Engine/Private/StaticMeshRender.cpp b/Engine/Source/Runtime/Engine/Private/StaticMeshRender.cpp
--- a/Engine/Source/Runtime/Engine/Private/StaticMeshRender.cpp
+++ b/Engine/Source/Runtime/Engine/Private/StaticMeshRender.cpp
@@ -705,7 +705,7 @@ void FStaticMeshSceneProxy::GetDynamicMeshElements(const TArray<const FSceneView
 										Collector.RegisterOneFrameMaterialProxy(VertexColorVisualizationMaterialInstance);
 										MeshElement.MaterialRenderProxy = VertexColorVisualizationMaterialInstance;
 									}
-									else if( bDrawComplexWireframeCollision || (bInCollisionView && bDrawComplexCollision && AllowDebugViewmodes()) )
+									else if (AllowDebugViewmodes() && (bDrawComplexWireframeCollision || (bInCollisionView && bDrawComplexCollision)))
 									{
 										if (LODModel.Sections[SectionIndex].bEnableCollision)
 										{
diff --git a/Engine/Source/Runtime/Engine/Private/StreamableManager.cpp b/Engine/Source/Runtime/Engine/Private/StreamableManager.cpp
--- a/Engine/Source/Runtime/Engine/Private/StreamableManager.cpp
+++ b/Engine/Source/Runtime/Engine/Private/StreamableManager.cpp
@@ -277,6 +277,23 @@ void FStreamableManager::RequestAsyncLoad(const TArray<FStringAssetReference>& T
 	}
 }
 
+void FStreamableManager::RequestAsyncLoad( const FStringAssetReference& TargetToStream, FStreamableDelegate DelegateToCall )
+{
+	TSharedRef< FStreamableRequest > NewRequest = MakeShareable( new FStreamableRequest() );
+	NewRequest->CompletionDelegate = DelegateToCall;
+
+	if ( FStreamable* Streamable = StreamInternal( TargetToStream ) )
+	{
+		Streamable->AddRelatedRequest( NewRequest );
+
+		if ( Streamable->Target )
+		{
+			CheckCompletedRequests( TargetToStream, Streamable );
+		}
+	}
+}
+
+
 void FStreamableManager::FindInMemory(FStringAssetReference& InOutTargetName, struct FStreamable* Existing)
 {
 	check(Existing);
diff --git a/Engine/Source/Runtime/Engine/Private/UnrealEngine.cpp b/Engine/Source/Runtime/Engine/Private/UnrealEngine.cpp
--- a/Engine/Source/Runtime/Engine/Private/UnrealEngine.cpp
+++ b/Engine/Source/Runtime/Engine/Private/UnrealEngine.cpp
@@ -7592,39 +7592,58 @@ UNetDriver* UEngine::FindNamedNetDriver(const UPendingNetGame* InPendingNetGame,
 	return FindNamedNetDriver_Local(GetWorldContextFromPendingNetGameChecked(InPendingNetGame).ActiveNetDrivers, NetDriverName);
 }
 
-bool CreateNamedNetDriver_Local(UEngine *Engine, FWorldContext &Context, FName NetDriverName, FName NetDriverDefinition)
+UNetDriver* CreateNetDriver_Local(UEngine *Engine, FWorldContext &Context, FName NetDriverDefinition)
 {
-	UNetDriver* NetDriver = FindNamedNetDriver_Local(Context.ActiveNetDrivers, NetDriverName);
-	if (NetDriver == NULL)
+	UNetDriver* NetDriver = nullptr;
+	for (int32 Index = 0; Index < Engine->NetDriverDefinitions.Num(); Index++)
 	{
-		for (int32 Index = 0; Index < Engine->NetDriverDefinitions.Num(); Index++)
+		FNetDriverDefinition& NetDriverDef = Engine->NetDriverDefinitions[Index];
+		if (NetDriverDef.DefName == NetDriverDefinition)
 		{
-			FNetDriverDefinition& NetDriverDef = Engine->NetDriverDefinitions[Index];
-			if (NetDriverDef.DefName == NetDriverDefinition)
+			// find the class to load
+			UClass* NetDriverClass = StaticLoadClass(UNetDriver::StaticClass(), NULL, *NetDriverDef.DriverClassName.ToString(), NULL, LOAD_Quiet, NULL);
+
+			// if it fails, then fall back to standard fallback
+			if (NetDriverClass == nullptr || !NetDriverClass->GetDefaultObject<UNetDriver>()->IsAvailable())
 			{
-				// find the class to load
-				UClass* NetDriverClass = StaticLoadClass(UNetDriver::StaticClass(), NULL, *NetDriverDef.DriverClassName.ToString(), NULL, LOAD_Quiet, NULL);
+				NetDriverClass = StaticLoadClass(UNetDriver::StaticClass(), NULL, *NetDriverDef.DriverClassNameFallback.ToString(), NULL, LOAD_None, NULL);
+			}
 
-				// if it fails, then fall back to standard fallback
-				if (NetDriverClass == NULL || !NetDriverClass->GetDefaultObject<UNetDriver>()->IsAvailable())
-				{
-					NetDriverClass = StaticLoadClass(UNetDriver::StaticClass(), NULL, *NetDriverDef.DriverClassNameFallback.ToString(), NULL, LOAD_None, NULL);
-				}
+			// Bail out if the net driver isn't available. The name may be incorrect or the class might not be built as part of the game configuration.
+			if (NetDriverClass == nullptr)
+			{
+				break;
+			}
 
-				// Bail out if the net driver isn't available. The name may be incorrect or the class might not be built as part of the game configuration.
-				if(NetDriverClass == NULL)
-				{
-					break;
-				}
+			// Try to create network driver.
+			NetDriver = ConstructObject<UNetDriver>(NetDriverClass);
+			check(NetDriver);
+			NetDriver->NetDriverName = NetDriver->GetFName();
 
-				// Try to create network driver.
-				NetDriver = NewObject<UNetDriver>(GetTransientPackage(), NetDriverClass);
-				check(NetDriver);
-				NetDriver->NetDriverName = NetDriverName;
-				
-				new (Context.ActiveNetDrivers) FNamedNetDriver(NetDriver, &NetDriverDef);
-				return true;
-			}
+			new (Context.ActiveNetDrivers) FNamedNetDriver(NetDriver, &NetDriverDef);
+			return NetDriver;
+		}
+	}
+	
+	UE_LOG(LogNet, Log, TEXT("CreateNamedNetDriver failed to create driver from definition %s"), *NetDriverDefinition.ToString());
+	return nullptr;
+}
+
+UNetDriver* UEngine::CreateNetDriver(UWorld *InWorld, FName NetDriverDefinition)
+{
+	return CreateNetDriver_Local(this, GetWorldContextFromWorldChecked(InWorld), NetDriverDefinition);
+}
+
+bool CreateNamedNetDriver_Local(UEngine *Engine, FWorldContext &Context, FName NetDriverName, FName NetDriverDefinition)
+{
+	UNetDriver* NetDriver = FindNamedNetDriver_Local(Context.ActiveNetDrivers, NetDriverName);
+	if (NetDriver == nullptr)
+	{
+		NetDriver = CreateNetDriver_Local(Engine, Context, NetDriverDefinition);
+		if (NetDriver)
+		{
+			NetDriver->NetDriverName = NetDriverName;
+			return true;
 		}
 	}
 
diff --git a/Engine/Source/Runtime/Engine/Public/SkeletalMeshMerge.h b/Engine/Source/Runtime/Engine/Public/SkeletalMeshMerge.h
--- a/Engine/Source/Runtime/Engine/Public/SkeletalMeshMerge.h
+++ b/Engine/Source/Runtime/Engine/Public/SkeletalMeshMerge.h
@@ -12,26 +12,70 @@ FSkeletalMeshMerge
 
 struct FRefPoseOverride
 {
-	/** The skeletal mesh that contains the reference pose. */
-	const USkeletalMesh* SkeletalMesh;
+ public:
 
-	/** The names of the bones to override. */
-	TArray<FName> BoneNames;
+	enum EBoneOverrideMode
+	{
+		BoneOnly,			// Override the bone only.
+		ChildrenOnly,		// Override the bone's children only.
+		BoneAndChildren,	// Override both the bone & children.
+		MAX,
+	};
 
 	/**
-	 * Whether the override for the bone should apply to the bone's children.
-	 * This array has a 1:1 with the 'BoneNames' array.
+	 * Constructs a FRefPoseOverride.
 	 */
-	TArray<bool> OverrideChildren;
+	FRefPoseOverride(const USkeletalMesh* ReferenceMesh)
+	{
+		this->SkeletalMesh = ReferenceMesh;
+	}
 
 	/**
 	 * Adds a bone to the list of poses to override.
 	 */
-	void AddOverride(FName BoneName, bool bOverrideChildren = false)
+	DEPRECATED(4.8, "Use AddOverride(FName, EBoneOverrideMode) instead.")
+	void AddOverride(FName BoneName, bool bOverrideChildren)
 	{
-		BoneNames.Add(BoneName);
-		OverrideChildren.Add(bOverrideChildren);
+		EBoneOverrideMode OverrideMode = BoneAndChildren;
+
+		if (!bOverrideChildren)
+		{
+			OverrideMode = BoneOnly;
+		}
+
+		AddOverride(BoneName, OverrideMode);
 	}
+
+	/**
+	 * Adds a bone to the list of poses to override.
+	 */
+	void AddOverride(FName BoneName, EBoneOverrideMode OverrideMode = BoneOnly)
+	{
+		FBoneOverrideInfo OverrideInfo;
+		OverrideInfo.BoneName = BoneName;
+		OverrideInfo.OverrideMode = OverrideMode;
+
+		Overrides.Add(OverrideInfo);
+	}
+
+ private:
+
+	 struct FBoneOverrideInfo
+	 {
+		 /** The names of the bone to override. */
+		 FName BoneName;
+
+		 /** Whether the override applies to the bone, bone's children, or both. */
+		 EBoneOverrideMode OverrideMode;
+	 };
+
+	/** The skeletal mesh that contains the reference pose. */
+	const USkeletalMesh* SkeletalMesh;
+
+	/** The list of bone overrides. */
+	TArray<FBoneOverrideInfo> Overrides;
+
+	friend class FSkeletalMeshMerge;
 };
 
 /** 
diff --git a/Engine/Source/Runtime/Engine/Public/VisualLogger/VisualLogger.inl b/Engine/Source/Runtime/Engine/Public/VisualLogger/VisualLogger.inl
--- a/Engine/Source/Runtime/Engine/Public/VisualLogger/VisualLogger.inl
+++ b/Engine/Source/Runtime/Engine/Public/VisualLogger/VisualLogger.inl
@@ -67,13 +67,13 @@ FVisualLogEntry* FVisualLogger::GetEntryToWrite(const class UObject* Object, flo
 
 	bool InitializeNewEntry = false;
 
-	UWorld* World = GetWorld(Object);
+	TWeakObjectPtr<UWorld> World = GetWorld(Object);
 
 	if (CurrentEntryPerObject.Contains(LogOwner))
 	{
 		CurrentEntry = &CurrentEntryPerObject[LogOwner];
 		InitializeNewEntry = TimeStamp > CurrentEntry->TimeStamp && ShouldCreate == ECreateIfNeeded::Create;
-		if (World)
+		if (World.IsValid())
 		{
 			World->GetTimerManager().ClearAllTimersForObject(this);
 			for (auto& CurrentPair : CurrentEntryPerObject)
@@ -138,15 +138,15 @@ FVisualLogEntry* FVisualLogger::GetEntryToWrite(const class UObject* Object, flo
 		}
 	}
 
-	if (World)
+	if (World.IsValid())
 	{
 		//set next tick timer to flush obsolete/old entries
 		World->GetTimerManager().SetTimerForNextTick(FTimerDelegate::CreateLambda(
-			[this](){
+			[this, World](){
 			for (auto& CurrentPair : CurrentEntryPerObject)
 			{
 				FVisualLogEntry* Entry = &CurrentPair.Value;
-				if (Entry->TimeStamp >= 0) // CurrentEntry->TimeStamp == -1 means it's not initialized entry information
+				if (Entry->TimeStamp >= 0 && (!World.IsValid() || Entry->TimeStamp < World->GetTimeSeconds())) // CurrentEntry->TimeStamp == -1 means it's not initialized entry information
 				{
 					for (auto* Device : OutputDevices)
 					{
diff --git a/Engine/Source/Runtime/GameplayAbilities/Private/Abilities/GameplayAbility.cpp b/Engine/Source/Runtime/GameplayAbilities/Private/Abilities/GameplayAbility.cpp
--- a/Engine/Source/Runtime/GameplayAbilities/Private/Abilities/GameplayAbility.cpp
+++ b/Engine/Source/Runtime/GameplayAbilities/Private/Abilities/GameplayAbility.cpp
@@ -146,7 +146,84 @@ bool UGameplayAbility::IsSupportedForNetworking() const
 	return Supported;
 }
 
-bool UGameplayAbility::CanActivateAbility(const FGameplayAbilitySpecHandle Handle, const FGameplayAbilityActorInfo* ActorInfo, OUT FGameplayTagContainer* OptionalRelevantTags) const
+bool UGameplayAbility::DoesAbilitySatisfyTagRequirements(const UAbilitySystemComponent& AbilitySystemComponent, FGameplayTagContainer SourceTags, FGameplayTagContainer TargetTags, OUT FGameplayTagContainer* OptionalRelevantTags) const
+{
+	bool bBlocked = false, bMissing = false;
+
+	const FGameplayTag& BlockedTag = UAbilitySystemGlobals::Get().ActivateFailTagsBlockedTag;
+	const FGameplayTag& MissingTag = UAbilitySystemGlobals::Get().ActivateFailTagsMissingTag;
+
+	// Check if any of this ability's tags are currently blocked
+	if (AbilitySystemComponent.AreAbilityTagsBlocked(AbilityTags))
+	{
+		bBlocked = true;
+	}
+
+	// Check to see the required/blocked tags for this ability
+	if (ActivationBlockedTags.Num() || ActivationRequiredTags.Num())
+	{
+		FGameplayTagContainer AbilitySystemComponentTags;
+
+		AbilitySystemComponent.GetOwnedGameplayTags(AbilitySystemComponentTags);
+
+		if (AbilitySystemComponentTags.MatchesAny(ActivationBlockedTags, false))
+		{
+			bBlocked = true;
+		}
+
+		if (!AbilitySystemComponentTags.MatchesAll(ActivationRequiredTags, true))
+		{
+			bMissing = true;
+		}
+	}
+
+	if (SourceBlockedTags.Num() || SourceRequiredTags.Num())
+	{
+		if (SourceTags.MatchesAny(SourceBlockedTags, false))
+		{
+			bBlocked = true;
+		}
+
+		if (!SourceTags.MatchesAll(SourceRequiredTags, true))
+		{
+			bMissing = true;
+		}
+	}
+
+	if (TargetBlockedTags.Num() || TargetRequiredTags.Num())
+	{
+		if (TargetTags.MatchesAny(TargetBlockedTags, false))
+		{
+			bBlocked = true;
+		}
+
+		if (!TargetTags.MatchesAll(TargetRequiredTags, true))
+		{
+			bMissing = true;
+		}
+	}
+
+	if (bBlocked)
+	{
+		if (OptionalRelevantTags && BlockedTag.IsValid())
+		{
+			OptionalRelevantTags->AddTag(BlockedTag);
+		}
+		return false;
+	}
+	if (bMissing)
+	{
+		if (OptionalRelevantTags && MissingTag.IsValid())
+		{
+			OptionalRelevantTags->AddTag(MissingTag);
+		}
+		return false;
+	}
+	
+	return true;
+}
+
+bool UGameplayAbility::CanActivateAbility(const FGameplayAbilitySpecHandle Handle, const FGameplayAbilityActorInfo* ActorInfo, FGameplayTagContainer SourceTags, FGameplayTagContainer TargetTags, OUT FGameplayTagContainer* OptionalRelevantTags) const
 {
 	// Don't set the actor info, CanActivate is called on the CDO
 
@@ -160,6 +237,12 @@ bool UGameplayAbility::CanActivateAbility(const FGameplayAbilitySpecHandle Handl
 	FGameplayTagContainer DummyContainer;
 	FGameplayTagContainer& OutTags = OptionalRelevantTags ? *OptionalRelevantTags : DummyContainer;
 
+	// make sure the ActorInfo and its ability system component are valid, if not bail out.
+	if (ActorInfo == nullptr || !ActorInfo->AbilitySystemComponent.IsValid())
+	{
+		return false;
+	}
+
 	if (ActorInfo->AbilitySystemComponent->GetUserAbilityActivationInhibited())
 	{
 		/**
@@ -174,19 +257,36 @@ bool UGameplayAbility::CanActivateAbility(const FGameplayAbilitySpecHandle Handl
 		return false;
 	}
 	
-	if (!CheckCooldown(Handle, ActorInfo))
+	if (!UAbilitySystemGlobals::Get().ShouldIgnoreCooldowns() && !CheckCooldown(Handle, ActorInfo, OptionalRelevantTags))
 	{
 		return false;
 	}
 
-	if (!CheckCost(Handle, ActorInfo))
+	if (!UAbilitySystemGlobals::Get().ShouldIgnoreCosts() && !CheckCost(Handle, ActorInfo, OptionalRelevantTags))
 	{
 		return false;
 	}
 
-	if (bHasBlueprintCanUse)
+	if (!DoesAbilitySatisfyTagRequirements(*ActorInfo->AbilitySystemComponent.Get(), SourceTags, TargetTags, OptionalRelevantTags))
+	{	// If the ability's tags are blocked, or if it has a "Blocking" tag or is missing a "Required" tag, then it can't activate.
+		return false;
+	}
+
+	FGameplayAbilitySpec* Spec = ActorInfo->AbilitySystemComponent->FindAbilitySpecFromHandle(Handle);
+	if (!Spec)
 	{
+		ABILITY_LOG(Warning, TEXT("CanActivateAbility called with invalid Handle"));
+		return false;
+	}
+
+	// Check if this ability's input binding is currently blocked
+	if (ActorInfo->AbilitySystemComponent->IsAbilityInputBlocked(Spec->InputID))
+	{
+		return false;
+	}
 
+	if (bHasBlueprintCanUse)
+	{
 		if (K2_CanActivateAbility(*ActorInfo, OutTags) == false)
 		{
 			ABILITY_LOG(Log, TEXT("CanActivateAbility %s failed, blueprint refused"), *GetName());
@@ -232,6 +332,11 @@ bool UGameplayAbility::CommitAbility(const FGameplayAbilitySpecHandle Handle, co
 
 bool UGameplayAbility::CommitAbilityCooldown(const FGameplayAbilitySpecHandle Handle, const FGameplayAbilityActorInfo* ActorInfo, const FGameplayAbilityActivationInfo ActivationInfo)
 {
+	if (UAbilitySystemGlobals::Get().ShouldIgnoreCooldowns())
+	{
+		return true;
+	}
+
 	// Last chance to fail (maybe we no longer have resources to commit since we after we started this ability activation)
 	if (!CheckCooldown(Handle, ActorInfo))
 	{
@@ -244,6 +349,11 @@ bool UGameplayAbility::CommitAbilityCooldown(const FGameplayAbilitySpecHandle Ha
 
 bool UGameplayAbility::CommitAbilityCost(const FGameplayAbilitySpecHandle Handle, const FGameplayAbilityActorInfo* ActorInfo, const FGameplayAbilityActivationInfo ActivationInfo)
 {
+	if (UAbilitySystemGlobals::Get().ShouldIgnoreCosts())
+	{
+		return true;
+	}
+
 	// Last chance to fail (maybe we no longer have resources to commit since we after we started this ability activation)
 	if (!CheckCost(Handle, ActorInfo))
 	{
@@ -264,7 +374,17 @@ bool UGameplayAbility::CommitCheck(const FGameplayAbilitySpecHandle Handle, cons
 	 *			-E.g., its possible the act of starting your ability makes it no longer activatable (CanaCtivateAbility() may be false if called here).
 	 */
 
-	return (CheckCooldown(Handle, ActorInfo) && CheckCost(Handle, ActorInfo));
+	if (!UAbilitySystemGlobals::Get().ShouldIgnoreCooldowns() && !CheckCooldown(Handle, ActorInfo))
+	{
+		return false;
+	}
+
+	if (!UAbilitySystemGlobals::Get().ShouldIgnoreCosts() && !CheckCost(Handle, ActorInfo))
+	{
+		return false;
+	}
+
+	return true;
 }
 
 void UGameplayAbility::CommitExecute(const FGameplayAbilitySpecHandle Handle, const FGameplayAbilityActorInfo* ActorInfo, const FGameplayAbilityActivationInfo ActivationInfo)
@@ -506,7 +626,7 @@ UGameplayEffect* UGameplayAbility::GetCostGameplayEffect() const
 	}
 }
 
-bool UGameplayAbility::CheckCooldown(const FGameplayAbilitySpecHandle Handle, const FGameplayAbilityActorInfo* ActorInfo) const
+bool UGameplayAbility::CheckCooldown(const FGameplayAbilitySpecHandle Handle, const FGameplayAbilityActorInfo* ActorInfo, OUT FGameplayTagContainer* OptionalRelevantTags) const
 {
 	const FGameplayTagContainer* CooldownTags = GetCooldownTags();
 	if (CooldownTags)
@@ -514,13 +634,20 @@ bool UGameplayAbility::CheckCooldown(const FGameplayAbilitySpecHandle Handle, co
 		check(ActorInfo->AbilitySystemComponent.IsValid());
 		if (CooldownTags->Num() > 0 && ActorInfo->AbilitySystemComponent->HasAnyMatchingGameplayTags(*CooldownTags))
 		{
+			const FGameplayTag& CoooldownTag = UAbilitySystemGlobals::Get().ActivateFailCooldownTag;
+
+			if (OptionalRelevantTags && CoooldownTag.IsValid())
+			{
+				OptionalRelevantTags->AddTag(CoooldownTag);
+			}
+
 			return false;
 		}
 	}
 	return true;
 }
 
-void UGameplayAbility::ApplyCooldown(const FGameplayAbilitySpecHandle Handle, const FGameplayAbilityActorInfo* ActorInfo, const FGameplayAbilityActivationInfo ActivationInfo)
+void UGameplayAbility::ApplyCooldown(const FGameplayAbilitySpecHandle Handle, const FGameplayAbilityActorInfo* ActorInfo, const FGameplayAbilityActivationInfo ActivationInfo) const
 {
 	UGameplayEffect* CooldownGE = GetCooldownGameplayEffect();
 	if (CooldownGE)
@@ -529,18 +656,27 @@ void UGameplayAbility::ApplyCooldown(const FGameplayAbilitySpecHandle Handle, co
 	}
 }
 
-bool UGameplayAbility::CheckCost(const FGameplayAbilitySpecHandle Handle, const FGameplayAbilityActorInfo* ActorInfo) const
+bool UGameplayAbility::CheckCost(const FGameplayAbilitySpecHandle Handle, const FGameplayAbilityActorInfo* ActorInfo, OUT FGameplayTagContainer* OptionalRelevantTags) const
 {
 	UGameplayEffect* CostGE = GetCostGameplayEffect();
 	if (CostGE)
 	{
 		check(ActorInfo->AbilitySystemComponent.IsValid());
-		return ActorInfo->AbilitySystemComponent->CanApplyAttributeModifiers(CostGE, GetAbilityLevel(Handle, ActorInfo), GetEffectContext(Handle, ActorInfo));
+		if (!ActorInfo->AbilitySystemComponent->CanApplyAttributeModifiers(CostGE, GetAbilityLevel(Handle, ActorInfo), GetEffectContext(Handle, ActorInfo)))
+		{
+			const FGameplayTag& CostTag = UAbilitySystemGlobals::Get().ActivateFailCostTag;
+
+			if (OptionalRelevantTags && CostTag.IsValid())
+			{
+				OptionalRelevantTags->AddTag(CostTag);
+			}
+			return false;
+		}
 	}
 	return true;
 }
 
-void UGameplayAbility::ApplyCost(const FGameplayAbilitySpecHandle Handle, const FGameplayAbilityActorInfo* ActorInfo, const FGameplayAbilityActivationInfo ActivationInfo)
+void UGameplayAbility::ApplyCost(const FGameplayAbilitySpecHandle Handle, const FGameplayAbilityActorInfo* ActorInfo, const FGameplayAbilityActivationInfo ActivationInfo) const
 {
 	UGameplayEffect* CostGE = GetCostGameplayEffect();
 	if (CostGE)
@@ -726,13 +862,13 @@ bool UGameplayAbility::K2_CommitAbilityCost(bool BroadcastCommitEvent)
 bool UGameplayAbility::K2_CheckAbilityCooldown()
 {
 	check(CurrentActorInfo);
-	return CheckCooldown(CurrentSpecHandle, CurrentActorInfo);
+	return UAbilitySystemGlobals::Get().ShouldIgnoreCooldowns() || CheckCooldown(CurrentSpecHandle, CurrentActorInfo);
 }
 
 bool UGameplayAbility::K2_CheckAbilityCost()
 {
 	check(CurrentActorInfo);
-	return CheckCost(CurrentSpecHandle, CurrentActorInfo);
+	return UAbilitySystemGlobals::Get().ShouldIgnoreCosts() || CheckCost(CurrentSpecHandle, CurrentActorInfo);
 }
 
 void UGameplayAbility::K2_EndAbility()
@@ -1056,7 +1192,7 @@ FActiveGameplayEffectHandle UGameplayAbility::K2_ApplyGameplayEffectToOwner(cons
 	return FActiveGameplayEffectHandle();
 }
 
-FActiveGameplayEffectHandle UGameplayAbility::ApplyGameplayEffectToOwner(const FGameplayAbilitySpecHandle Handle, const FGameplayAbilityActorInfo* ActorInfo, const FGameplayAbilityActivationInfo ActivationInfo, const UGameplayEffect* GameplayEffect, float GameplayEffectLevel)
+FActiveGameplayEffectHandle UGameplayAbility::ApplyGameplayEffectToOwner(const FGameplayAbilitySpecHandle Handle, const FGameplayAbilityActorInfo* ActorInfo, const FGameplayAbilityActivationInfo ActivationInfo, const UGameplayEffect* GameplayEffect, float GameplayEffectLevel) const
 {
 	if (GameplayEffect && (HasAuthorityOrPredictionKey(ActorInfo, &ActivationInfo)))
 	{
@@ -1076,7 +1212,7 @@ FActiveGameplayEffectHandle UGameplayAbility::K2_ApplyGameplayEffectSpecToOwner(
 	return ApplyGameplayEffectSpecToOwner(CurrentSpecHandle, CurrentActorInfo, CurrentActivationInfo, EffectSpecHandle);
 }
 
-FActiveGameplayEffectHandle UGameplayAbility::ApplyGameplayEffectSpecToOwner(const FGameplayAbilitySpecHandle AbilityHandle, const FGameplayAbilityActorInfo* ActorInfo, const FGameplayAbilityActivationInfo ActivationInfo, const FGameplayEffectSpecHandle SpecHandle)
+FActiveGameplayEffectHandle UGameplayAbility::ApplyGameplayEffectSpecToOwner(const FGameplayAbilitySpecHandle AbilityHandle, const FGameplayAbilityActorInfo* ActorInfo, const FGameplayAbilityActivationInfo ActivationInfo, const FGameplayEffectSpecHandle SpecHandle) const
 {
 	if (SpecHandle.IsValid() && (HasAuthorityOrPredictionKey(ActorInfo, &ActivationInfo)))
 	{
@@ -1098,7 +1234,7 @@ TArray<FActiveGameplayEffectHandle> UGameplayAbility::K2_ApplyGameplayEffectToTa
 	return BP_ApplyGameplayEffectToTarget(Target, GameplayEffect->GetClass(), GameplayEffectLevel);
 }
 
-TArray<FActiveGameplayEffectHandle> UGameplayAbility::ApplyGameplayEffectToTarget(const FGameplayAbilitySpecHandle Handle, const FGameplayAbilityActorInfo* ActorInfo, const FGameplayAbilityActivationInfo ActivationInfo, FGameplayAbilityTargetDataHandle Target, TSubclassOf<UGameplayEffect> GameplayEffectClass, float GameplayEffectLevel)
+TArray<FActiveGameplayEffectHandle> UGameplayAbility::ApplyGameplayEffectToTarget(const FGameplayAbilitySpecHandle Handle, const FGameplayAbilityActorInfo* ActorInfo, const FGameplayAbilityActivationInfo ActivationInfo, FGameplayAbilityTargetDataHandle Target, TSubclassOf<UGameplayEffect> GameplayEffectClass, float GameplayEffectLevel) const
 {
 	TArray<FActiveGameplayEffectHandle> EffectHandles;
 
@@ -1120,7 +1256,7 @@ TArray<FActiveGameplayEffectHandle> UGameplayAbility::K2_ApplyGameplayEffectSpec
 	return ApplyGameplayEffectSpecToTarget(CurrentSpecHandle, CurrentActorInfo, CurrentActivationInfo, SpecHandle, TargetData);
 }
 
-TArray<FActiveGameplayEffectHandle> UGameplayAbility::ApplyGameplayEffectSpecToTarget(const FGameplayAbilitySpecHandle AbilityHandle, const FGameplayAbilityActorInfo* ActorInfo, const FGameplayAbilityActivationInfo ActivationInfo, const FGameplayEffectSpecHandle SpecHandle, FGameplayAbilityTargetDataHandle TargetData)
+TArray<FActiveGameplayEffectHandle> UGameplayAbility::ApplyGameplayEffectSpecToTarget(const FGameplayAbilitySpecHandle AbilityHandle, const FGameplayAbilityActorInfo* ActorInfo, const FGameplayAbilityActivationInfo ActivationInfo, const FGameplayEffectSpecHandle SpecHandle, FGameplayAbilityTargetDataHandle TargetData) const
 {
 	TArray<FActiveGameplayEffectHandle> EffectHandles;
 	
diff --git a/Engine/Source/Runtime/GameplayAbilities/Private/Abilities/GameplayAbilityTargetActor.cpp b/Engine/Source/Runtime/GameplayAbilities/Private/Abilities/GameplayAbilityTargetActor.cpp
--- a/Engine/Source/Runtime/GameplayAbilities/Private/Abilities/GameplayAbilityTargetActor.cpp
+++ b/Engine/Source/Runtime/GameplayAbilities/Private/Abilities/GameplayAbilityTargetActor.cpp
@@ -73,7 +73,7 @@ void AGameplayAbilityTargetActor::CancelTargeting()
 	Destroy();
 }
 
-bool AGameplayAbilityTargetActor::IsNetRelevantFor(const APlayerController* RealViewer, const AActor* Viewer, const FVector& SrcLocation) const
+bool AGameplayAbilityTargetActor::IsNetRelevantFor(const APlayerController* RealViewer, const AActor* ViewTarget, const FVector& SrcLocation) const
 {
 	//The player who created the ability doesn't need to be updated about it - there should be local prediction in place.
 	if (RealViewer == MasterPC)
@@ -86,10 +86,10 @@ bool AGameplayAbilityTargetActor::IsNetRelevantFor(const APlayerController* Real
 
 	if (Avatar)
 	{
-		return Avatar->IsNetRelevantFor(RealViewer, Viewer, SrcLocation);
+		return Avatar->IsNetRelevantFor(RealViewer, ViewTarget, SrcLocation);
 	}
 
-	return Super::IsNetRelevantFor(RealViewer, Viewer, SrcLocation);
+	return Super::IsNetRelevantFor(RealViewer, ViewTarget, SrcLocation);
 }
 
 bool AGameplayAbilityTargetActor::GetReplicates() const
diff --git a/Engine/Source/Runtime/GameplayAbilities/Private/Abilities/GameplayAbilityWorldReticle.cpp b/Engine/Source/Runtime/GameplayAbilities/Private/Abilities/GameplayAbilityWorldReticle.cpp
--- a/Engine/Source/Runtime/GameplayAbilities/Private/Abilities/GameplayAbilityWorldReticle.cpp
+++ b/Engine/Source/Runtime/GameplayAbilities/Private/Abilities/GameplayAbilityWorldReticle.cpp
@@ -38,7 +38,7 @@ void AGameplayAbilityWorldReticle::InitializeReticle(AGameplayAbilityTargetActor
 	OnParametersInitialized();
 }
 
-bool AGameplayAbilityWorldReticle::IsNetRelevantFor(const APlayerController* RealViewer, const AActor* Viewer, const FVector& SrcLocation) const
+bool AGameplayAbilityWorldReticle::IsNetRelevantFor(const APlayerController* RealViewer, const AActor* ViewTarget, const FVector& SrcLocation) const
 {
 	//The player who created the ability doesn't need to be updated about it - there should be local prediction in place.
 	if (RealViewer == MasterPC)
@@ -46,7 +46,7 @@ bool AGameplayAbilityWorldReticle::IsNetRelevantFor(const APlayerController* Rea
 		return false;
 	}
 
-	return Super::IsNetRelevantFor(RealViewer, Viewer, SrcLocation);
+	return Super::IsNetRelevantFor(RealViewer, ViewTarget, SrcLocation);
 }
 
 bool AGameplayAbilityWorldReticle::GetReplicates() const
diff --git a/Engine/Source/Runtime/GameplayAbilities/Private/Abilities/GameplayAbility_CharacterJump.cpp b/Engine/Source/Runtime/GameplayAbilities/Private/Abilities/GameplayAbility_CharacterJump.cpp
--- a/Engine/Source/Runtime/GameplayAbilities/Private/Abilities/GameplayAbility_CharacterJump.cpp
+++ b/Engine/Source/Runtime/GameplayAbilities/Private/Abilities/GameplayAbility_CharacterJump.cpp
@@ -39,9 +39,9 @@ void UGameplayAbility_CharacterJump::InputReleased(const FGameplayAbilitySpecHan
 	}
 }
 
-bool UGameplayAbility_CharacterJump::CanActivateAbility(const FGameplayAbilitySpecHandle Handle, const FGameplayAbilityActorInfo* ActorInfo, OUT FGameplayTagContainer* OptionalRelevantTags) const
+bool UGameplayAbility_CharacterJump::CanActivateAbility(const FGameplayAbilitySpecHandle Handle, const FGameplayAbilityActorInfo* ActorInfo, FGameplayTagContainer SourceTags, FGameplayTagContainer TargetTags, OUT FGameplayTagContainer* OptionalRelevantTags) const
 {
-	if (!Super::CanActivateAbility(Handle, ActorInfo, OptionalRelevantTags))
+	if (!Super::CanActivateAbility(Handle, ActorInfo, SourceTags, TargetTags, OptionalRelevantTags))
 	{
 		return false;
 	}
diff --git a/Engine/Source/Runtime/GameplayAbilities/Private/AbilitySystemBlueprintLibrary.cpp b/Engine/Source/Runtime/GameplayAbilities/Private/AbilitySystemBlueprintLibrary.cpp
--- a/Engine/Source/Runtime/GameplayAbilities/Private/AbilitySystemBlueprintLibrary.cpp
+++ b/Engine/Source/Runtime/GameplayAbilities/Private/AbilitySystemBlueprintLibrary.cpp
@@ -23,23 +23,29 @@ UAbilitySystemComponent* UAbilitySystemBlueprintLibrary::GetAbilitySystemCompone
 
 void UAbilitySystemBlueprintLibrary::SendGameplayEventToActor(AActor* Actor, FGameplayTag EventTag, FGameplayEventData Payload)
 {
-	IAbilitySystemInterface* AbilitySystemInterface = Cast<IAbilitySystemInterface>(Actor);
-	if (AbilitySystemInterface != NULL)
+	if (Actor && !Actor->IsPendingKill())
 	{
-		UAbilitySystemComponent* AbilitySystemComponent = AbilitySystemInterface->GetAbilitySystemComponent();
-		if (AbilitySystemComponent != NULL)
+		IAbilitySystemInterface* AbilitySystemInterface = Cast<IAbilitySystemInterface>(Actor);
+		if (AbilitySystemInterface != NULL)
 		{
-			FScopedPredictionWindow NewScopedWindow(AbilitySystemComponent, true);
-			AbilitySystemComponent->HandleGameplayEvent(EventTag, &Payload);
+			UAbilitySystemComponent* AbilitySystemComponent = AbilitySystemInterface->GetAbilitySystemComponent();
+			if (AbilitySystemComponent != NULL)
+			{
+				FScopedPredictionWindow NewScopedWindow(AbilitySystemComponent, true);
+				AbilitySystemComponent->HandleGameplayEvent(EventTag, &Payload);
+			}
 		}
 	}
 }
 
-float UAbilitySystemBlueprintLibrary::GetFloatAttribute(const class AActor* Actor, FGameplayAttribute Attribute)
+float UAbilitySystemBlueprintLibrary::GetFloatAttribute(const class AActor* Actor, FGameplayAttribute Attribute, bool& bSuccessfullyFoundAttribute)
 {
+	bSuccessfullyFoundAttribute = true;
+
 	const UAbilitySystemComponent* const AbilitySystem = UAbilitySystemGlobals::GetAbilitySystemComponentFromActor(Actor);
-	if (!AbilitySystem)
+	if (!AbilitySystem || !AbilitySystem->HasAttributeSetForAttribute(Attribute))
 	{
+		bSuccessfullyFoundAttribute = false;
 		return 0.f;
 	}
 
@@ -185,6 +191,25 @@ TArray<AActor*> UAbilitySystemBlueprintLibrary::GetActorsFromTargetData(FGamepla
 	return TArray<AActor*>();
 }
 
+bool UAbilitySystemBlueprintLibrary::DoesTargetDataContainActor(FGameplayAbilityTargetDataHandle TargetData, int32 Index, AActor* Actor)
+{
+	if (TargetData.Data.IsValidIndex(Index))
+	{
+		FGameplayAbilityTargetData* Data = TargetData.Data[Index].Get();
+		if (Data)
+		{
+			for (auto DataActorIter : Data->GetActors())
+			{
+				if (DataActorIter == Actor)
+				{
+					return true;
+				}
+			}
+		}
+	}
+	return false;
+}
+
 bool UAbilitySystemBlueprintLibrary::TargetDataHasActor(FGameplayAbilityTargetDataHandle TargetData, int32 Index)
 {
 	if (TargetData.Data.IsValidIndex(Index))
diff --git a/Engine/Source/Runtime/GameplayAbilities/Private/AbilitySystemComponent.cpp b/Engine/Source/Runtime/GameplayAbilities/Private/AbilitySystemComponent.cpp
--- a/Engine/Source/Runtime/GameplayAbilities/Private/AbilitySystemComponent.cpp
+++ b/Engine/Source/Runtime/GameplayAbilities/Private/AbilitySystemComponent.cpp
@@ -239,6 +239,8 @@ FGameplayEffectContextHandle UAbilitySystemComponent::GetEffectContext() const
 {
 	FGameplayEffectContextHandle Context = FGameplayEffectContextHandle(UAbilitySystemGlobals::Get().AllocGameplayEffectContext());
 	// By default use the owner and avatar as the instigator and causer
+	check(AbilityActorInfo.IsValid());
+	
 	Context.AddInstigator(AbilityActorInfo->OwnerActor.Get(), AbilityActorInfo->AvatarActor.Get());
 	return Context;
 }
@@ -335,6 +337,11 @@ FOnActiveGameplayEffectRemoved* UAbilitySystemComponent::OnGameplayEffectRemoved
 	return nullptr;
 }
 
+FOnActiveGameplayEffectRemoved& UAbilitySystemComponent::OnAnyGameplayEffectRemovedDelegate()
+{
+	return ActiveGameplayEffects.OnActiveGameplayEffectRemovedDelegate;
+}
+
 int32 UAbilitySystemComponent::GetNumActiveGameplayEffects() const
 {
 	return ActiveGameplayEffects.GetNumGameplayEffects();
@@ -468,7 +475,7 @@ void UAbilitySystemComponent::UpdateTagMap(const FGameplayTagContainer& Containe
 
 FActiveGameplayEffectHandle UAbilitySystemComponent::ApplyGameplayEffectSpecToTarget(OUT FGameplayEffectSpec &Spec, UAbilitySystemComponent *Target, FPredictionKey PredictionKey)
 {
-	if (!UAbilitySystemGlobals::Get().PredictTargetGameplayEffects)
+	if (!UAbilitySystemGlobals::Get().ShouldPredictTargetGameplayEffects())
 	{
 		// If we don't want to predict target effects, clear prediction key
 		PredictionKey = FPredictionKey();
@@ -476,7 +483,7 @@ FActiveGameplayEffectHandle UAbilitySystemComponent::ApplyGameplayEffectSpecToTa
 
 	FActiveGameplayEffectHandle ReturnHandle;
 
-	if (!UAbilitySystemGlobals::Get().PredictTargetGameplayEffects)
+	if (!UAbilitySystemGlobals::Get().ShouldPredictTargetGameplayEffects())
 	{
 		// If we don't want to predict target effects, clear prediction key
 		PredictionKey = FPredictionKey();
@@ -730,7 +737,7 @@ float UAbilitySystemComponent::GetGameplayEffectMagnitude(FActiveGameplayEffectH
 	return ActiveGameplayEffects.GetGameplayEffectMagnitude(Handle, Attribute);
 }
 
-void UAbilitySystemComponent::InvokeGameplayCueEvent(const FGameplayEffectSpec &Spec, EGameplayCueEvent::Type EventType)
+void UAbilitySystemComponent::InvokeGameplayCueEvent(const FGameplayEffectSpecForRPC &Spec, EGameplayCueEvent::Type EventType)
 {
 	AActor* ActorAvatar = AbilityActorInfo->AvatarActor.Get();
 	if (!Spec.Def)
@@ -875,7 +882,7 @@ void UAbilitySystemComponent::RemoveAllGameplayCues()
 	}
 }
 
-void UAbilitySystemComponent::NetMulticast_InvokeGameplayCueExecuted_FromSpec_Implementation(const FGameplayEffectSpec Spec, FPredictionKey PredictionKey)
+void UAbilitySystemComponent::NetMulticast_InvokeGameplayCueExecuted_FromSpec_Implementation(const FGameplayEffectSpecForRPC Spec, FPredictionKey PredictionKey)
 {
 	if (IsOwnerActorAuthoritative() || PredictionKey.IsValidKey() == false)
 	{
@@ -944,7 +951,11 @@ TArray<float> UAbilitySystemComponent::GetActiveEffectsDuration(const FActiveGam
 
 void UAbilitySystemComponent::RemoveActiveEffectsWithTags(const FGameplayTagContainer Tags)
 {
-	RemoveActiveEffects(FActiveGameplayEffectQuery(&Tags));
+	AActor *OwningActor = GetOwner();
+	if (OwningActor && OwningActor->Role == ROLE_Authority)
+	{
+		RemoveActiveEffects(FActiveGameplayEffectQuery(&Tags));
+	}
 }
 
 void UAbilitySystemComponent::RemoveActiveEffects(const FActiveGameplayEffectQuery Query, int32 StacksToRemove)
@@ -1350,7 +1361,7 @@ void UAbilitySystemComponent::DisplayDebug(class UCanvas* Canvas, const class FD
 				StatusText = TEXT(" (TagBlocked)");
 				AbilityTextColor = FColor::Red;
 			}
-			else if (AbilitySpec.Ability->CanActivateAbility(AbilitySpec.Handle, AbilityActorInfo.Get()) == false)
+			else if (AbilitySpec.Ability->CanActivateAbility(AbilitySpec.Handle, AbilityActorInfo.Get(), FGameplayTagContainer(), FGameplayTagContainer()) == false)
 			{
 				StatusText = TEXT(" (CantActivate)");
 				AbilityTextColor = FColor::Red;
diff --git a/Engine/Source/Runtime/GameplayAbilities/Private/AbilitySystemComponent_Abilities.cpp b/Engine/Source/Runtime/GameplayAbilities/Private/AbilitySystemComponent_Abilities.cpp
--- a/Engine/Source/Runtime/GameplayAbilities/Private/AbilitySystemComponent_Abilities.cpp
+++ b/Engine/Source/Runtime/GameplayAbilities/Private/AbilitySystemComponent_Abilities.cpp
@@ -245,6 +245,44 @@ FGameplayAbilitySpecHandle UAbilitySystemComponent::GiveAbility(FGameplayAbility
 	return OwnedSpec.Handle;
 }
 
+FGameplayAbilitySpecHandle UAbilitySystemComponent::GiveAbilityAndActivateOnce(FGameplayAbilitySpec Spec)
+{
+	check(Spec.Ability);
+
+	if (Spec.Ability->GetInstancingPolicy() == EGameplayAbilityInstancingPolicy::NonInstanced || Spec.Ability->GetNetExecutionPolicy() == EGameplayAbilityNetExecutionPolicy::LocalOnly)
+	{
+		ABILITY_LOG(Error, TEXT("GiveAbilityAndActivateOnce called on ability %s that is non instanced or won't execute on server, not allowed!"), *Spec.Ability->GetName());
+
+		return FGameplayAbilitySpecHandle();
+	}
+
+	if (!IsOwnerActorAuthoritative())
+	{
+		ABILITY_LOG(Error, TEXT("GiveAbilityAndActivateOnce called on ability %s on the client, not allowed!"), *Spec.Ability->GetName());
+
+		return FGameplayAbilitySpecHandle();
+	}
+
+	FGameplayAbilitySpecHandle AddedAbilityHandle = GiveAbility(Spec);
+
+	FGameplayAbilitySpec* FoundSpec = FindAbilitySpecFromHandle(AddedAbilityHandle);
+
+	if (FoundSpec)
+	{
+		FoundSpec->RemoveAfterActivation = true;
+
+		if (!TryActivateAbility(AddedAbilityHandle, FPredictionKey()))
+		{
+			// We failed to activate it, so remove it now
+			ClearAbility(AddedAbilityHandle);
+
+			return FGameplayAbilitySpecHandle();
+		}
+	}
+
+	return AddedAbilityHandle;
+}
+
 void UAbilitySystemComponent::ClearAllAbilities()
 {
 	check(IsOwnerActorAuthoritative());	// Should be called on authority
@@ -421,7 +459,7 @@ void UAbilitySystemComponent::CheckForClearedAbilities()
 	}
 }
 
-const TArray<FGameplayAbilitySpec>& UAbilitySystemComponent::GetActivatableAbilities()
+const TArray<FGameplayAbilitySpec>& UAbilitySystemComponent::GetActivatableAbilities() const
 {
 	return ActivatableAbilities.Items;
 }
@@ -491,6 +529,7 @@ void UAbilitySystemComponent::NotifyAbilityEnded(FGameplayAbilitySpecHandle Hand
 	FGameplayAbilitySpec* Spec = FindAbilitySpecFromHandle(Handle);
 	check(Spec);
 	check(Ability);
+	ENetRole OwnerRole = GetOwnerRole();
 
 	// If AnimatingAbility ended, clear the pointer
 	if (LocalAnimMontageInfo.AnimatingAbility == Ability)
@@ -515,7 +554,7 @@ void UAbilitySystemComponent::NotifyAbilityEnded(FGameplayAbilitySpecHandle Hand
 
 		if (Ability->GetReplicationPolicy() != EGameplayAbilityReplicationPolicy::ReplicateNo)
 		{
-			if (GetOwnerRole() == ROLE_Authority)
+			if (OwnerRole == ROLE_Authority)
 			{
 				Spec->ReplicatedInstances.Remove(Ability);
 				AllReplicatedInstancedAbilities.Remove(Ability);
@@ -528,7 +567,19 @@ void UAbilitySystemComponent::NotifyAbilityEnded(FGameplayAbilitySpecHandle Hand
 			Ability->MarkPendingKill();
 		}
 	}
-	MarkAbilitySpecDirty(*Spec);
+
+	if (OwnerRole == ROLE_Authority)
+	{
+		if (Spec->RemoveAfterActivation && !Spec->IsActive())
+		{
+			// If we should remove after activation and there are no more active instances, kill it now
+			ClearAbility(Handle);
+		}
+		else
+		{
+			MarkAbilitySpecDirty(*Spec);
+		}
+	}
 }
 
 void UAbilitySystemComponent::CancelAbility(UGameplayAbility* Ability)
@@ -668,15 +719,29 @@ void UAbilitySystemComponent::OnRep_ActivateAbilities()
 	}
 
 	CheckForClearedAbilities();
+
+	// Try to run any pending activations that couldn't run before. If they don't work now, kill them
+
+	for (auto PendingAbilityInfo : PendingServerActivatedAbilities)
+	{
+		ClientActivateAbilitySucceedWithEventData_Implementation(PendingAbilityInfo.Handle, PendingAbilityInfo.PredictionKey.Current, PendingAbilityInfo.TriggerEventData);
+	}
+	PendingServerActivatedAbilities.Empty();
+
 }
 
-void UAbilitySystemComponent::GetActivateableGameplayAbilitySpecsByAllMatchingTags(const FGameplayTagContainer& GameplayTagContainer, TArray < struct FGameplayAbilitySpec* >& MatchingGameplayAbilities) const
+void UAbilitySystemComponent::GetActivatableGameplayAbilitySpecsByAllMatchingTags(const FGameplayTagContainer& GameplayTagContainer, TArray < struct FGameplayAbilitySpec* >& MatchingGameplayAbilities) const
 {
 	for (const FGameplayAbilitySpec& Spec : ActivatableAbilities.Items)
 	{		
 		if (Spec.Ability && Spec.Ability->AbilityTags.MatchesAll(GameplayTagContainer, false))
 		{
-			MatchingGameplayAbilities.Add(const_cast<FGameplayAbilitySpec*>(&Spec));
+			// Do NOT consider abilities that are blocked by tags currently.  That way, we can use the blocking to find
+			// an appropriate ability based on tags when we have more than one ability that match the GameplayTagContainer.
+			if (Spec.Ability->DoesAbilitySatisfyTagRequirements(*this, FGameplayTagContainer(), FGameplayTagContainer()))
+			{
+				MatchingGameplayAbilities.Add(const_cast<FGameplayAbilitySpec*>(&Spec));
+			}
 		}
 	}
 }
@@ -684,7 +749,7 @@ void UAbilitySystemComponent::GetActivateableGameplayAbilitySpecsByAllMatchingTa
 bool UAbilitySystemComponent::TryActivateAbilityByTag(const FGameplayTagContainer& GameplayTagContainer)
 {
 	TArray<FGameplayAbilitySpec*> AbilitiesToActivate;
-	GetActivateableGameplayAbilitySpecsByAllMatchingTags(GameplayTagContainer, AbilitiesToActivate);
+	GetActivatableGameplayAbilitySpecsByAllMatchingTags(GameplayTagContainer, AbilitiesToActivate);
 
 	UGameplayAbility* ActivatedAbility(nullptr);
 
@@ -698,6 +763,17 @@ bool UAbilitySystemComponent::TryActivateAbilityByTag(const FGameplayTagContaine
 	return bSuccess;
 }
 
+bool UAbilitySystemComponent::IsAbilityInputBlocked(int32 InputID) const
+{
+	// Check if this ability's input binding is currently blocked
+	if (InputID >= 0 && InputID < BlockedAbilityBindings.Num() && BlockedAbilityBindings[InputID] > 0)
+	{
+		return true;
+	}
+
+	return false;
+}
+
 /**
  * Attempts to activate the ability.
  *	-This function calls CanActivateAbility
@@ -707,6 +783,7 @@ bool UAbilitySystemComponent::TryActivateAbilityByTag(const FGameplayTagContaine
  */
 bool UAbilitySystemComponent::TryActivateAbility(FGameplayAbilitySpecHandle Handle, FPredictionKey InPredictionKey, UGameplayAbility** OutInstancedAbility, FOnGameplayAbilityEnded* OnGameplayAbilityEndedDelegate, const FGameplayEventData* TriggerEventData)
 {
+	FGameplayTagContainer FailureTags;
 	FGameplayAbilitySpec* Spec = FindAbilitySpecFromHandle(Handle);
 	if (!Spec)
 	{
@@ -752,51 +829,31 @@ bool UAbilitySystemComponent::TryActivateAbility(FGameplayAbilitySpecHandle Hand
 		return false;
 	}
 
-	// Check if this ability's input binding is currently blocked
-	if (Spec->InputID >= 0 && Spec->InputID < BlockedAbilityBindings.Num() && BlockedAbilityBindings[Spec->InputID] > 0)
-	{
-		return false;
-	}
-
-	// Check if any of this ability's tags are currently blocked
-	if (Ability->AbilityTags.MatchesAny(BlockedAbilityTags.GetExplicitGameplayTags(), false))
-	{
-		return false;
-	}
-
-	// Check to see the required/blocked tags for this ability
-	if (Ability->ActivationBlockedTags.Num() || Ability->ActivationRequiredTags.Num())
-	{
-		FGameplayTagContainer SourceTags;
-
-		GetOwnedGameplayTags(SourceTags);
-
-		if (SourceTags.MatchesAny(Ability->ActivationBlockedTags, false))
-		{
-			return false;
-		}
-
-		if (!SourceTags.MatchesAll(Ability->ActivationRequiredTags, true))
-		{
-			return false;
-		}
-	}
-
 	// If it's instance once the instanced ability will be set, otherwise it will be null
 	UGameplayAbility* InstancedAbility = Spec->GetPrimaryInstance();
 
 	// TODO: We should run this on on the non instanced, but run it on the instance once for now, need to fixup data
 	// (Always do a non instanced CanActivate check)
 
+	FGameplayTagContainer SourceTags;
+	FGameplayTagContainer TargetTags;
+	if (TriggerEventData != nullptr)
+	{
+		SourceTags = TriggerEventData->InstigatorTags;
+		TargetTags = TriggerEventData->TargetTags;
+	}
+
 	if (InstancedAbility)
 	{
-		if (!InstancedAbility->CanActivateAbility(Handle, ActorInfo))
+		if (!InstancedAbility->CanActivateAbility(Handle, ActorInfo, SourceTags, TargetTags, &FailureTags))
 		{
+			NotifyAbilityFailed(Handle, Ability, FailureTags);
 			return false;
 		}
 	}
-	else if (!Ability->CanActivateAbility(Handle, ActorInfo))
+	else if (!Ability->CanActivateAbility(Handle, ActorInfo, SourceTags, TargetTags, &FailureTags))
 	{
+		NotifyAbilityFailed(Handle, Ability, FailureTags);
 		return false;
 	}
 
@@ -807,7 +864,7 @@ bool UAbilitySystemComponent::TryActivateAbility(FGameplayAbilitySpecHandle Hand
 		{
 			if (Ability->bRetriggerInstancedAbility)
 			{
-				Ability->EndAbility(Handle, ActorInfo, Spec->ActivationInfo, true);
+				InstancedAbility->EndAbility(Handle, ActorInfo, Spec->ActivationInfo, true);
 			}
 			else
 			{
@@ -865,8 +922,15 @@ bool UAbilitySystemComponent::TryActivateAbility(FGameplayAbilitySpecHandle Hand
 		if (!bIsLocal && Ability->GetNetExecutionPolicy() != EGameplayAbilityNetExecutionPolicy::ServerOnly)
 		{
 			// TODO: Add support for "Server prediction keys"
-			ClientActivateAbilitySucceed(Handle, InPredictionKey.Current, TriggerEventData ? *TriggerEventData : FGameplayEventData());
-
+			if (TriggerEventData)
+			{
+				ClientActivateAbilitySucceedWithEventData(Handle, InPredictionKey.Current, *TriggerEventData);
+			}
+			else
+			{
+				ClientActivateAbilitySucceed(Handle, InPredictionKey.Current);
+			}
+			
 			// This will get copied into the instanced abilities
 			ActivationInfo.bCanBeEndedByOtherInstance = Ability->bServerRespectsRemoteAbilityCancellation;;
 		}
@@ -1010,7 +1074,15 @@ void UAbilitySystemComponent::InternalServerTryActiveAbility(FGameplayAbilitySpe
 					break;
 				case EAbilityExecutionState::Executing:
 				case EAbilityExecutionState::Succeeded:
-					ClientActivateAbilitySucceed(Handle, PredictionKey.Current, TriggerEventData ? *TriggerEventData : FGameplayEventData());
+					if (TriggerEventData)
+					{
+						ClientActivateAbilitySucceedWithEventData(Handle, PredictionKey.Current, *TriggerEventData);
+					}
+					else
+					{
+						ClientActivateAbilitySucceed(Handle, PredictionKey.Current);
+					}
+
 					// Client commands to end the ability that come in after this point are considered for this instance
 					TArray<UGameplayAbility*> Instances = Spec->GetAbilityInstances();
 					for (auto Instance : Instances)
@@ -1033,7 +1105,7 @@ void UAbilitySystemComponent::InternalServerTryActiveAbility(FGameplayAbilitySpe
 		AbilityInfo.PredictionKey = PredictionKey;
 		AbilityInfo.Handle = Handle;
 
-		PendingClientAbilities.Add(AbilityInfo);
+		PendingClientActivatedAbilities.Add(AbilityInfo);
 		return;
 	}
 
@@ -1068,7 +1140,8 @@ void UAbilitySystemComponent::ReplicateEndAbility(FGameplayAbilitySpecHandle Han
 		}
 		else
 		{
-			ServerEndAbility(Handle, ActivationInfo);
+			// This passes up the current prediction key if we have one
+			ServerEndAbility(Handle, ActivationInfo, ScopedPredictionKey);
 		}
 	}
 }
@@ -1093,12 +1166,14 @@ void UAbilitySystemComponent::RemoteEndAbility(FGameplayAbilitySpecHandle Abilit
 	}
 }
 
-void UAbilitySystemComponent::ServerEndAbility_Implementation(FGameplayAbilitySpecHandle AbilityToEnd, FGameplayAbilityActivationInfo ActivationInfo)
+void UAbilitySystemComponent::ServerEndAbility_Implementation(FGameplayAbilitySpecHandle AbilityToEnd, FGameplayAbilityActivationInfo ActivationInfo, FPredictionKey PredictionKey)
 {
+	FScopedPredictionWindow ScopedPrediction(this, PredictionKey);
+
 	RemoteEndAbility(AbilityToEnd, ActivationInfo);
 }
 
-bool UAbilitySystemComponent::ServerEndAbility_Validate(FGameplayAbilitySpecHandle AbilityToEnd, FGameplayAbilityActivationInfo ActivationInfo)
+bool UAbilitySystemComponent::ServerEndAbility_Validate(FGameplayAbilitySpecHandle AbilityToEnd, FGameplayAbilityActivationInfo ActivationInfo, FPredictionKey PredictionKey)
 {
 	return true;
 }
@@ -1142,12 +1217,26 @@ void UAbilitySystemComponent::OnClientActivateAbilityFailed(FGameplayAbilitySpec
 
 static_assert(sizeof(int16) == sizeof(FPredictionKey::KeyType), "Sizeof PredictionKey::KeyType does not match RPC parameters in AbilitySystemComponent ClientActivateAbilitySucceed_Implementation");
 
-void UAbilitySystemComponent::ClientActivateAbilitySucceed_Implementation(FGameplayAbilitySpecHandle Handle, int16 PredictionKey, FGameplayEventData TriggerEventData)
+void UAbilitySystemComponent::ClientActivateAbilitySucceed_Implementation(FGameplayAbilitySpecHandle Handle, int16 PredictionKey)
+{
+	ClientActivateAbilitySucceedWithEventData_Implementation(Handle, PredictionKey, FGameplayEventData());
+}
+
+void UAbilitySystemComponent::ClientActivateAbilitySucceedWithEventData_Implementation(FGameplayAbilitySpecHandle Handle, int16 PredictionKey, FGameplayEventData TriggerEventData)
 {
 	FGameplayAbilitySpec* Spec = FindAbilitySpecFromHandle(Handle);
 	if (!Spec)
 	{
-		// Can potentially happen in race conditions where client tries to activate ability that is removed server side before it is received.
+		// Can happen if the client gets told to activate an ability the same frame that abilities are added on the server
+		FPredictionKey TempKey;
+		TempKey.Current = PredictionKey;
+		FPendingAbilityInfo AbilityInfo;
+		AbilityInfo.PredictionKey = TempKey;
+		AbilityInfo.Handle = Handle;
+		AbilityInfo.TriggerEventData = TriggerEventData;
+
+		// This won't add it if we're currently being called from the pending list
+		PendingServerActivatedAbilities.AddUnique(AbilityInfo);
 		return;
 	}
 
@@ -1191,6 +1280,9 @@ void UAbilitySystemComponent::ClientActivateAbilitySucceed_Implementation(FGamep
 	else
 	{
 		// We haven't already executed this ability at all, so kick it off.
+
+		// The spec will now be active, and we need to keep track on the client as well.  Since we cannot call TryActivateAbility, which will increment ActiveCount on the server, we have to do this here.
+		++Spec->ActiveCount;
 		
 		if (AbilityToActivate->GetInstancingPolicy() == EGameplayAbilityInstancingPolicy::InstancedPerExecution)
 		{
@@ -1264,11 +1356,11 @@ bool UAbilitySystemComponent::TriggerAbilityFromGameplayEvent(FGameplayAbilitySp
 			&& ActorInfo->OwnerActor->Role == ROLE_Authority)
 		{
 			bool bPendingClientAbilityFound = false;
-			for (auto PendingAbilityInfo : Component.PendingClientAbilities)
+			for (auto PendingAbilityInfo : Component.PendingClientActivatedAbilities)
 			{
 				if (ScopedPredictionKey.Current == PendingAbilityInfo.PredictionKey.Base && Handle == PendingAbilityInfo.Handle) // found a match
 				{
-					Component.PendingClientAbilities.RemoveSingleSwap(PendingAbilityInfo);
+					Component.PendingClientActivatedAbilities.RemoveSingleSwap(PendingAbilityInfo);
 					bPendingClientAbilityFound = true;
 					break;
 				}
@@ -1335,6 +1427,11 @@ void UAbilitySystemComponent::NotifyAbilityActivated(FGameplayAbilitySpecHandle
 	AbilityActivatedCallbacks.Broadcast(Ability);
 }
 
+void UAbilitySystemComponent::NotifyAbilityFailed(const FGameplayAbilitySpecHandle Handle, UGameplayAbility* Ability, FGameplayTagContainer FailureReason)
+{
+	AbilityFailedCallbacks.Broadcast(Ability, FailureReason);
+}
+
 int32 UAbilitySystemComponent::HandleGameplayEvent(FGameplayTag EventTag, const FGameplayEventData* Payload)
 {
 	int32 TriggeredCount = 0;
@@ -1378,8 +1475,14 @@ void UAbilitySystemComponent::MonitoredTagChanged(const FGameplayTag Tag, int32
 				{
 					if (NewCount > 0)
 					{
+						// Populate event data so this will use the same blueprint node to activate as gameplay triggers
+						FGameplayEventData EventData;
+						EventData.EventMagnitude = NewCount;
+						EventData.EventTag = EventTag;
+						EventData.Instigator = OwnerActor;
+						EventData.Target = OwnerActor;
 						// Try to activate it
-						TryActivateAbility(Spec->Handle, FPredictionKey());
+						TryActivateAbility(Spec->Handle, FPredictionKey(), nullptr, nullptr, &EventData);
 
 						// TODO: Check client/server type
 					}
@@ -1695,6 +1798,13 @@ void UAbilitySystemComponent::OnRep_SimulatedTasks()
 	}
 }
 
+#if ENABLE_VISUAL_LOG
+void UAbilitySystemComponent::ClearDebugInstantEffects()
+{
+	ActiveGameplayEffects.DebugExecutedGameplayEffects.Empty();
+}
+#endif // ENABLE_VISUAL_LOG
+
 // ---------------------------------------------------------------------------
 
 float UAbilitySystemComponent::PlayMontage(UGameplayAbility* InAnimatingAbility, FGameplayAbilityActivationInfo ActivationInfo, UAnimMontage* NewAnimMontage, float InPlayRate, FName StartSectionName)
diff --git a/Engine/Source/Runtime/GameplayAbilities/Private/AbilitySystemGlobals.cpp b/Engine/Source/Runtime/GameplayAbilities/Private/AbilitySystemGlobals.cpp
--- a/Engine/Source/Runtime/GameplayAbilities/Private/AbilitySystemGlobals.cpp
+++ b/Engine/Source/Runtime/GameplayAbilities/Private/AbilitySystemGlobals.cpp
@@ -23,6 +23,11 @@ UAbilitySystemGlobals::UAbilitySystemGlobals(const FObjectInitializer& ObjectIni
 #if WITH_EDITORONLY_DATA
 	RegisteredReimportCallback = false;
 #endif // #if WITH_EDITORONLY_DATA
+
+#if !(UE_BUILD_SHIPPING || UE_BUILD_TEST)
+	bIgnoreAbilitySystemCooldowns = false;
+	bIgnoreAbilitySystemCosts = false;
+#endif // #if !(UE_BUILD_SHIPPING || UE_BUILD_TEST)
 }
 
 void UAbilitySystemGlobals::InitGlobalData()
@@ -33,6 +38,7 @@ void UAbilitySystemGlobals::InitGlobalData()
 	InitAttributeDefaults();
 
 	GetGameplayCueManager();
+	InitGlobalTags();
 }
 
 UCurveTable * UAbilitySystemGlobals::GetGlobalCurveTable()
@@ -238,4 +244,36 @@ UGameplayCueManager* UAbilitySystemGlobals::GetGameplayCueManager()
 void UAbilitySystemGlobals::GlobalPreGameplayEffectSpecApply(FGameplayEffectSpec& Spec, UAbilitySystemComponent* AbilitySystemComponent)
 {
 
-}
\ No newline at end of file
+}
+
+void UAbilitySystemGlobals::ToggleIgnoreAbilitySystemCooldowns()
+{
+#if !(UE_BUILD_SHIPPING || UE_BUILD_TEST)
+	bIgnoreAbilitySystemCooldowns = !bIgnoreAbilitySystemCooldowns;
+#endif // #if !(UE_BUILD_SHIPPING || UE_BUILD_TEST)
+}
+
+void UAbilitySystemGlobals::ToggleIgnoreAbilitySystemCosts()
+{
+#if !(UE_BUILD_SHIPPING || UE_BUILD_TEST)
+	bIgnoreAbilitySystemCosts = !bIgnoreAbilitySystemCosts;
+#endif // #if !(UE_BUILD_SHIPPING || UE_BUILD_TEST)
+}
+
+bool UAbilitySystemGlobals::ShouldIgnoreCooldowns() const
+{
+#if !(UE_BUILD_SHIPPING || UE_BUILD_TEST)
+	return bIgnoreAbilitySystemCooldowns;
+#else
+	return false;
+#endif // #if !(UE_BUILD_SHIPPING || UE_BUILD_TEST)
+}
+
+bool UAbilitySystemGlobals::ShouldIgnoreCosts() const
+{
+#if !(UE_BUILD_SHIPPING || UE_BUILD_TEST)
+	return bIgnoreAbilitySystemCosts;
+#else
+	return false;
+#endif // #if !(UE_BUILD_SHIPPING || UE_BUILD_TEST)
+}
diff --git a/Engine/Source/Runtime/GameplayAbilities/Private/GameplayAbilitiesExec.cpp b/Engine/Source/Runtime/GameplayAbilities/Private/GameplayAbilitiesExec.cpp
new file mode 100644
--- /dev/null
+++ b/Engine/Source/Runtime/GameplayAbilities/Private/GameplayAbilitiesExec.cpp
@@ -0,0 +1,45 @@
+// Copyright 1998-2015 Epic Games, Inc. All Rights Reserved.
+
+#include "AbilitySystemPrivatePCH.h"
+#include "AbilitySystemComponent.h"
+#if !(UE_BUILD_SHIPPING || UE_BUILD_TEST)
+#include "Misc/CoreMisc.h"
+
+struct FGameplayAbilitiesExec : public FSelfRegisteringExec
+{
+	FGameplayAbilitiesExec()
+	{
+	}
+
+	// Begin FExec Interface
+	virtual bool Exec(UWorld* Inworld, const TCHAR* Cmd, FOutputDevice& Ar) override;
+	// End FExec Interface
+};
+
+FGameplayAbilitiesExec GameplayAbilitiesExecInstance;
+
+bool FGameplayAbilitiesExec::Exec(UWorld* Inworld, const TCHAR* Cmd, FOutputDevice& Ar)
+{
+	if (Inworld == NULL)
+	{
+		return false;
+	}
+
+	UAbilitySystemGlobals& AbilitySystemGlobals = UAbilitySystemGlobals::Get();
+	bool bHandled = false;
+
+	if (FParse::Command(&Cmd, TEXT("ToggleIgnoreAbilitySystemCooldowns")))
+	{
+		AbilitySystemGlobals.ToggleIgnoreAbilitySystemCooldowns();
+		bHandled = true;
+	}
+	else if (FParse::Command(&Cmd, TEXT("ToggleIgnoreAbilitySystemCosts")))
+	{
+		AbilitySystemGlobals.ToggleIgnoreAbilitySystemCosts();
+		bHandled = true;
+	}
+
+	return bHandled;
+}
+
+#endif // !(UE_BUILD_SHIPPING || UE_BUILD_TEST)
\ No newline at end of file
diff --git a/Engine/Source/Runtime/GameplayAbilities/Private/GameplayAbilityTargetTypes.cpp b/Engine/Source/Runtime/GameplayAbilities/Private/GameplayAbilityTargetTypes.cpp
--- a/Engine/Source/Runtime/GameplayAbilities/Private/GameplayAbilityTargetTypes.cpp
+++ b/Engine/Source/Runtime/GameplayAbilities/Private/GameplayAbilityTargetTypes.cpp
@@ -9,57 +9,39 @@
 
 TArray<FActiveGameplayEffectHandle> FGameplayAbilityTargetData::ApplyGameplayEffect(const UGameplayEffect* GameplayEffect, const FGameplayEffectContextHandle& InEffectContext, float Level, FPredictionKey PredictionKey)
 {
-	// Clone the effect context per target to avoid modifying original context
-	FGameplayEffectContextHandle EffectContext = InEffectContext.Duplicate();
-	
-	FGameplayEffectSpec	SpecToApply(GameplayEffect, EffectContext, Level);
-	
-	AddTargetDataToContext(EffectContext, false);
-
-	TArray<TWeakObjectPtr<AActor> > Actors = GetActors();
-	TArray<FActiveGameplayEffectHandle>	AppliedHandles;
-	AppliedHandles.Reserve(Actors.Num());
+	// Make a temp spec and call the spec function. This ends up cloning the spec per target
+	FGameplayEffectSpec	TempSpecToApply(GameplayEffect, InEffectContext, Level);
 
-	for (TWeakObjectPtr<AActor> TargetActor : Actors)
-	{
-		if (TargetActor.IsValid())
-		{
-			UAbilitySystemComponent* TargetComponent = UAbilitySystemBlueprintLibrary::GetAbilitySystemComponent(TargetActor.Get());
-			if (TargetComponent && EffectContext.GetInstigatorAbilitySystemComponent())
-			{
-				AppliedHandles.Add(EffectContext.GetInstigatorAbilitySystemComponent()->ApplyGameplayEffectSpecToTarget(SpecToApply, TargetComponent, PredictionKey));
-			}
-		}
-	}
-
-	return AppliedHandles;
+	return ApplyGameplayEffectSpec(TempSpecToApply, PredictionKey);
 }
 
-TArray<FActiveGameplayEffectHandle> FGameplayAbilityTargetData::ApplyGameplayEffectSpec(FGameplayEffectSpec& Spec, FPredictionKey PredictionKey)
+TArray<FActiveGameplayEffectHandle> FGameplayAbilityTargetData::ApplyGameplayEffectSpec(const FGameplayEffectSpec& InSpec, FPredictionKey PredictionKey)
 {
 	TArray<FActiveGameplayEffectHandle>	AppliedHandles;
 
-	if (!ensure(Spec.GetContext().IsValid()))
+	if (!ensure(InSpec.GetContext().IsValid() && InSpec.GetContext().GetInstigatorAbilitySystemComponent()))
 	{
 		return AppliedHandles;
 	}
 
-	FGameplayEffectContextHandle EffectContext = Spec.GetContext();
-	AddTargetDataToContext(EffectContext, false);
-
 	TArray<TWeakObjectPtr<AActor> > Actors = GetActors();
 	
 	AppliedHandles.Reserve(Actors.Num());
 
 	for (TWeakObjectPtr<AActor> TargetActor : Actors)
 	{
-		if (TargetActor.IsValid())
+		UAbilitySystemComponent* TargetComponent = UAbilitySystemBlueprintLibrary::GetAbilitySystemComponent(TargetActor.Get());
+
+		if (TargetComponent)
 		{
-			UAbilitySystemComponent* TargetComponent = UAbilitySystemBlueprintLibrary::GetAbilitySystemComponent(TargetActor.Get());
-			if (TargetComponent && EffectContext.GetInstigatorAbilitySystemComponent())
-			{
-				AppliedHandles.Add(EffectContext.GetInstigatorAbilitySystemComponent()->ApplyGameplayEffectSpecToTarget(Spec, TargetComponent, PredictionKey));
-			}
+			// We have to make a new effect spec and context here, because otherwise the targeting info gets accumulated and things take damage multiple times
+			FGameplayEffectSpec	SpecToApply(InSpec);
+			FGameplayEffectContextHandle EffectContext = SpecToApply.GetContext().Duplicate();
+			SpecToApply.SetContext(EffectContext);
+
+			AddTargetDataToContext(EffectContext, false);
+
+			AppliedHandles.Add(EffectContext.GetInstigatorAbilitySystemComponent()->ApplyGameplayEffectSpecToTarget(SpecToApply, TargetComponent, PredictionKey));
 		}
 	}
 
diff --git a/Engine/Source/Runtime/GameplayAbilities/Private/GameplayEffect.cpp b/Engine/Source/Runtime/GameplayAbilities/Private/GameplayEffect.cpp
--- a/Engine/Source/Runtime/GameplayAbilities/Private/GameplayEffect.cpp
+++ b/Engine/Source/Runtime/GameplayAbilities/Private/GameplayEffect.cpp
@@ -9,7 +9,9 @@
 #include "AbilitySystemInterface.h"
 #include "GameplayModMagnitudeCalculation.h"
 #include "GameplayEffectExecutionCalculation.h"
-#include "VisualLogger.h"
+#if ENABLE_VISUAL_LOG
+#include "VisualLoggerTypes.h"
+#endif // ENABLE_VISUAL_LOG
 
 const float UGameplayEffect::INFINITE_DURATION = -1.f;
 const float UGameplayEffect::INSTANT_APPLICATION = 0.f;
@@ -324,6 +326,28 @@ void FGameplayEffectModifierMagnitude::GetAttributeCaptureDefinitions(OUT TArray
 	}
 }
 
+bool FGameplayEffectModifierMagnitude::GetStaticMagnitudeIfPossible(float InLevel, float& OutMagnitude) const
+{
+	if (MagnitudeCalculationType == EGameplayEffectMagnitudeCalculation::ScalableFloat)
+	{
+		OutMagnitude = ScalableFloatMagnitude.GetValueAtLevel(InLevel);
+		return true;
+	}
+
+	return false;
+}
+
+bool FGameplayEffectModifierMagnitude::GetSetByCallerDataNameIfPossible(FName& OutDataName) const
+{
+	if (MagnitudeCalculationType == EGameplayEffectMagnitudeCalculation::SetByCaller)
+	{
+		OutDataName = SetByCallerMagnitude.DataName;
+		return true;
+	}
+
+	return false;
+}
+
 #if WITH_EDITOR
 FText FGameplayEffectModifierMagnitude::GetValueForEditorDisplay() const
 {
@@ -503,6 +527,20 @@ FGameplayEffectSpec& FGameplayEffectSpec::operator=(const FGameplayEffectSpec& O
 	return *this;
 }
 
+FGameplayEffectSpecForRPC::FGameplayEffectSpecForRPC()
+	: Def(nullptr)
+	, Level(UGameplayEffect::INVALID_LEVEL)
+{
+}
+
+FGameplayEffectSpecForRPC::FGameplayEffectSpecForRPC(const FGameplayEffectSpec& InSpec)
+	: Def(InSpec.Def)
+	, ModifiedAttributes(MoveTemp(InSpec.ModifiedAttributes))
+	, EffectContext(InSpec.GetEffectContext())
+	, Level(InSpec.GetLevel())
+{
+}
+
 void FGameplayEffectSpec::SetupAttributeCaptureDefinitions()
 {
 	// Add duration if required
@@ -709,6 +747,18 @@ FGameplayEffectModifiedAttribute* FGameplayEffectSpec::GetModifiedAttribute(cons
 	return NULL;
 }
 
+const FGameplayEffectModifiedAttribute* FGameplayEffectSpecForRPC::GetModifiedAttribute(const FGameplayAttribute& Attribute) const
+{
+	for (const FGameplayEffectModifiedAttribute& ModifiedAttribute : ModifiedAttributes)
+	{
+		if (ModifiedAttribute.Attribute == Attribute)
+		{
+			return &ModifiedAttribute;
+		}
+	}
+	return NULL;
+}
+
 FGameplayEffectModifiedAttribute* FGameplayEffectSpec::AddModifiedAttribute(const FGameplayAttribute& Attribute)
 {
 	FGameplayEffectModifiedAttribute NewAttribute;
@@ -1164,6 +1214,8 @@ void FActiveGameplayEffect::PreReplicatedRemove(const struct FActiveGameplayEffe
 		return;
 	}
 
+	ABILITY_LOG(Verbose, TEXT("PreReplicatedRemove: %s %s Marked as Pending Remove: %s"), *Handle.ToString(), *Spec.Def->GetName(), IsPendingRemove ? TEXT("TRUE") : TEXT("FALSE"));
+
 	const_cast<FActiveGameplayEffectsContainer&>(InArray).InternalOnActiveGameplayEffectRemoved(*this);	// Const cast is ok. It is there to prevent mutation of the GameplayEffects array, which this wont do.
 	
 	InArray.Owner->InvokeGameplayCueEvent(Spec, EGameplayCueEvent::Removed);
@@ -1330,9 +1382,12 @@ void FActiveGameplayEffectsContainer::ExecuteActiveEffectsFrom(FGameplayEffectSp
 				// If successful, apply conditional specs
 				for (const TSubclassOf<UGameplayEffect>& TargetDefClass : CurExecDef.ConditionalGameplayEffectClasses)
 				{
-					const UGameplayEffect* TargetDef = TargetDefClass->GetDefaultObject<UGameplayEffect>();
+					if (*TargetDefClass != nullptr)
+					{
+						const UGameplayEffect* TargetDef = TargetDefClass->GetDefaultObject<UGameplayEffect>();
 
-					ConditionalEffectSpecs.Add(FGameplayEffectSpecHandle(new FGameplayEffectSpec(TargetDef, Spec.GetEffectContext(), Spec.GetLevel())));
+						ConditionalEffectSpecs.Add(FGameplayEffectSpecHandle(new FGameplayEffectSpec(TargetDef, Spec.GetEffectContext(), Spec.GetLevel())));
+					}
 				}
 			}
 
@@ -1405,6 +1460,9 @@ void FActiveGameplayEffectsContainer::ExecutePeriodicGameplayEffect(FActiveGamep
 			}
 		}
 
+		// Clear modified attributes before each periodic execution
+		ActiveEffect->Spec.ModifiedAttributes.Empty();
+
 		// Execute
 		ExecuteActiveEffectsFrom(ActiveEffect->Spec);
 	}
@@ -2235,6 +2293,8 @@ bool FActiveGameplayEffectsContainer::InternalRemoveActiveGameplayEffect(int32 I
 		FActiveGameplayEffect& Effect = *GetActiveGameplayEffect(Idx);
 		ensure(!Effect.IsPendingRemove);
 
+		ABILITY_LOG(Verbose, TEXT("InternalRemoveActiveGameplayEffect: Auth: %s Handle: %s Def: %s"), IsNetAuthority() ? TEXT("TRUE") : TEXT("FALSE"), *Effect.Handle.ToString(), Effect.Spec.Def ? *Effect.Spec.Def->GetName() : TEXT("NONE"));
+
 		if (StacksToRemove > 0 && Effect.Spec.StackCount > StacksToRemove)
 		{
 			// This won't be a full remove, only a change in StackCount.
@@ -2291,13 +2351,15 @@ bool FActiveGameplayEffectsContainer::InternalRemoveActiveGameplayEffect(int32 I
 		{
 			// We are locked, so this removal is now pending.
 			PendingRemoves++;
+
+			ABILITY_LOG(Verbose, TEXT("InternalRemoveActiveGameplayEffect while locked; Counting as a Pending Remove: Auth: %s Handle: %s Def: %s"), IsNetAuthority() ? TEXT("TRUE") : TEXT("FALSE"), *Effect.Handle.ToString(), Effect.Spec.Def ? *Effect.Spec.Def->GetName() : TEXT("NONE"));
 		}
 		else
 		{
 			// Not locked, so do the removal right away.
 
 			// If we are not scope locked, then there is no way this idx should be referring to something on the pending add list.
-			// It is possible to remove a GE that is pedning add, but it would happen while the scope lock is still in effect, resulting
+			// It is possible to remove a GE that is pending add, but it would happen while the scope lock is still in effect, resulting
 			// in a pending remove being set.
 			check(Idx < GameplayEffects_Internal.Num());
 
@@ -2336,6 +2398,7 @@ void FActiveGameplayEffectsContainer::InternalOnActiveGameplayEffectRemoved(cons
 	}
 
 	Effect.OnRemovedDelegate.Broadcast();
+	OnActiveGameplayEffectRemovedDelegate.Broadcast();
 }
 
 void FActiveGameplayEffectsContainer::RemoveActiveGameplayEffectGrantedTagsAndModifiers(const FActiveGameplayEffect& Effect)
@@ -2680,35 +2743,51 @@ bool FActiveGameplayEffectsContainer::HasPredictedEffectWithPredictedKey(FPredic
 	return false;
 }
 
-#if ENABLE_VISUAL_LOG
-void FActiveGameplayEffectsContainer::GrabDebugSnapshot(FVisualLogEntry* Snapshot) const
+void FActiveGameplayEffectsContainer::GetActiveGameplayEffectDataByAttribute(TMultiMap<FGameplayAttribute, FActiveGameplayEffectsContainer::DebugExecutedGameplayEffectData>& EffectMap) const
 {
-	FVisualLogStatusCategory ActiveEffectsCategory;
-	ActiveEffectsCategory.Category = TEXT("Effects");
-
-	TMultiMap<FGameplayAttribute, FActiveGameplayEffectsContainer::DebugExecutedGameplayEffectData> EffectMap;
+	EffectMap.Empty();
 
 	// Add all of the active gameplay effects
 	for (const FActiveGameplayEffect& Effect : this)
 	{
-		ensure(Effect.Spec.Modifiers.Num() == Effect.Spec.Def->Modifiers.Num());
-		for (int32 Idx = 0; Idx < Effect.Spec.Modifiers.Num(); ++Idx)
+		if (Effect.Spec.Modifiers.Num() == Effect.Spec.Def->Modifiers.Num())
 		{
-			FActiveGameplayEffectsContainer::DebugExecutedGameplayEffectData Data;
-			Data.Attribute = Effect.Spec.Def->Modifiers[Idx].Attribute;
-			Data.ActivationState = Effect.bIsInhibited ? TEXT("INHIBITED") : TEXT("ACTIVE");
-			Data.GameplayEffectName = Effect.Spec.Def->GetName();
-			Data.Magnitude = Effect.Spec.Modifiers[Idx].GetEvaluatedMagnitude();
+			for (int32 Idx = 0; Idx < Effect.Spec.Modifiers.Num(); ++Idx)
+			{
+				FActiveGameplayEffectsContainer::DebugExecutedGameplayEffectData Data;
+				Data.Attribute = Effect.Spec.Def->Modifiers[Idx].Attribute;
+				Data.ActivationState = Effect.bIsInhibited ? TEXT("INHIBITED") : TEXT("ACTIVE");
+				Data.GameplayEffectName = Effect.Spec.Def->GetName();
+				Data.ModifierOp = Effect.Spec.Def->Modifiers[Idx].ModifierOp;
+				Data.Magnitude = Effect.Spec.Modifiers[Idx].GetEvaluatedMagnitude();
+				if (Effect.Spec.StackCount > 1)
+				{
+					Data.Magnitude = GameplayEffectUtilities::ComputeStackedModifierMagnitude(Data.Magnitude, Effect.Spec.StackCount, Data.ModifierOp);
+				}
+				Data.StackCount = Effect.Spec.StackCount;
 
-			EffectMap.Add(Data.Attribute, Data);
+				EffectMap.Add(Data.Attribute, Data);
+			}
 		}
 	}
-
+#if ENABLE_VISUAL_LOG
 	// Add the executed gameplay effects if we recorded them
 	for (FActiveGameplayEffectsContainer::DebugExecutedGameplayEffectData Data : DebugExecutedGameplayEffects)
 	{
 		EffectMap.Add(Data.Attribute, Data);
 	}
+#endif // ENABLE_VISUAL_LOG
+}
+
+#if ENABLE_VISUAL_LOG
+void FActiveGameplayEffectsContainer::GrabDebugSnapshot(FVisualLogEntry* Snapshot) const
+{
+	FVisualLogStatusCategory ActiveEffectsCategory;
+	ActiveEffectsCategory.Category = TEXT("Effects");
+
+	TMultiMap<FGameplayAttribute, FActiveGameplayEffectsContainer::DebugExecutedGameplayEffectData> EffectMap;
+
+	GetActiveGameplayEffectDataByAttribute(EffectMap);
 
 	// For each attribute that was modified go through all of its modifiers and list them
 	TArray<FGameplayAttribute> AttributeKeys;
@@ -2964,6 +3043,7 @@ void FActiveGameplayEffectsContainer::DecrementLock()
 
 			if (Effect.IsPendingRemove)
 			{
+				ABILITY_LOG(Verbose, TEXT("DecrementLock decrementing a pending remove: Auth: %s Handle: %s Def: %s"), IsNetAuthority() ? TEXT("TRUE") : TEXT("FALSE"), *Effect.Handle.ToString(), Effect.Spec.Def ? *Effect.Spec.Def->GetName() : TEXT("NONE"));
 				GameplayEffects_Internal.RemoveAtSwap(idx, 1, false);
 				PendingRemoves--;
 			}
diff --git a/Engine/Source/Runtime/GameplayAbilities/Private/GameplayEffectTypes.cpp b/Engine/Source/Runtime/GameplayAbilities/Private/GameplayEffectTypes.cpp
--- a/Engine/Source/Runtime/GameplayAbilities/Private/GameplayEffectTypes.cpp
+++ b/Engine/Source/Runtime/GameplayAbilities/Private/GameplayEffectTypes.cpp
@@ -99,13 +99,19 @@ bool FGameplayEffectContext::NetSerialize(FArchive& Ar, class UPackageMap* Map,
 {
 	Ar << Instigator;
 	Ar << EffectCauser;
-	Ar << Actors;
 
-	bool HasHitResults = HitResult.IsValid();
-	Ar << HasHitResults;
+	uint8 HasActorByte = (Actors.Num() > 0);
+	Ar.SerializeBits(&HasActorByte, 1);
+	if (HasActorByte == 1)
+	{
+		Ar << Actors;
+	}
+
+	uint8 HasHitResultsByte = HitResult.IsValid();
+	Ar.SerializeBits(&HasHitResultsByte, 1);
 	if (Ar.IsLoading())
 	{
-		if (HasHitResults)
+		if (HasHitResultsByte)
 		{
 			if (!HitResult.IsValid())
 			{
@@ -115,14 +121,18 @@ bool FGameplayEffectContext::NetSerialize(FArchive& Ar, class UPackageMap* Map,
 		AddInstigator(Instigator.Get(), EffectCauser.Get()); // Just to initialize InstigatorAbilitySystemComponent
 	}
 
-	if (HasHitResults == 1)
+	if (HasHitResultsByte == 1)
 	{
 		HitResult->NetSerialize(Ar, Map, bOutSuccess);
 	}
-
-	Ar << bHasWorldOrigin;
-	Ar << WorldOrigin;
-
+	uint8 HasWorldOriginByte = bHasWorldOrigin;
+	Ar.SerializeBits(&HasWorldOriginByte, 1);
+	bHasWorldOrigin = HasWorldOriginByte & 1;
+	if (bHasWorldOrigin)
+	{
+		Ar << WorldOrigin;
+	}
+	
 	bOutSuccess = true;
 	return true;
 }
@@ -154,7 +164,7 @@ void FGameplayEffectContext::GetOwnedGameplayTags(OUT FGameplayTagContainer& Act
 	{
 		TagInterface->GetOwnedGameplayTags(ActorTagContainer);
 	}
-	else if (InstigatorAbilitySystemComponent)
+	else if (InstigatorAbilitySystemComponent.IsValid())
 	{
 		InstigatorAbilitySystemComponent->GetOwnedGameplayTags(ActorTagContainer);
 	}
@@ -163,7 +173,7 @@ void FGameplayEffectContext::GetOwnedGameplayTags(OUT FGameplayTagContainer& Act
 bool FGameplayEffectContextHandle::NetSerialize(FArchive& Ar, class UPackageMap* Map, bool& bOutSuccess)
 {
 	bool ValidData = Data.IsValid();
-	Ar << ValidData;
+	Ar.SerializeBits(&ValidData,1);
 
 	if (ValidData)
 	{
@@ -285,7 +295,11 @@ void FGameplayTagCountContainer::UpdateTagMap_Internal(const FGameplayTag& Tag,
 		// Block attempted reduction of non-explicit tags, as they were never truly added to the container directly
 		else
 		{
-			ABILITY_LOG(Warning, TEXT("Attempted to remove tag: %s from tag count container, but it is not explicitly in the container!"), *Tag.ToString());
+			// only warn about tags that are in the container but will not be removed because they aren't explicitly in the container
+			if (ExplicitTags.HasTag(Tag, EGameplayTagMatchType::IncludeParentTags, EGameplayTagMatchType::Explicit))
+			{
+				ABILITY_LOG(Warning, TEXT("Attempted to remove tag: %s from tag count container, but it is not explicitly in the container!"), *Tag.ToString());
+			}
 			return;
 		}
 	}
diff --git a/Engine/Source/Runtime/GameplayAbilities/Private/GameplayPrediction.cpp b/Engine/Source/Runtime/GameplayAbilities/Private/GameplayPrediction.cpp
--- a/Engine/Source/Runtime/GameplayAbilities/Private/GameplayPrediction.cpp
+++ b/Engine/Source/Runtime/GameplayAbilities/Private/GameplayPrediction.cpp
@@ -196,28 +196,34 @@ void FPredictionKeyDelegates::AddDependency(FPredictionKey::KeyType ThisKey, FPr
 
 FScopedPredictionWindow::FScopedPredictionWindow(UAbilitySystemComponent* AbilitySystemComponent, FPredictionKey InPredictionKey)
 {
+	if (!ensure(AbilitySystemComponent != NULL))
+	{
+		return;
+	}
+
 	// This is used to set an already generated prediction key as the current scoped prediction key.
 	// Should be called on the server for logical scopes where a given key is valid. E.g, "client gave me this key, we both are going to run Foo()".
-
-	// If you are hitting this, this FScopedPredictionWindow constructor should only be called from Server RPCs where the client has given us a PredictionKey.
-	ensure(AbilitySystemComponent->IsNetSimulating() == false);
-
-	Owner = AbilitySystemComponent;
-	Owner->ScopedPredictionKey = InPredictionKey;
-	ClearScopedPredictionKey = true;
-	SetReplicatedPredictionKey = true;
+	
+	if (AbilitySystemComponent->IsNetSimulating() == false)
+	{
+		Owner = AbilitySystemComponent;
+		check(Owner.IsValid());
+		Owner->ScopedPredictionKey = InPredictionKey;
+		ClearScopedPredictionKey = true;
+		SetReplicatedPredictionKey = true;
+	}
 }
 
 FScopedPredictionWindow::FScopedPredictionWindow(UAbilitySystemComponent* InAbilitySystemComponent, bool bCanGenerateNewKey)
 {
-	// On the server, this will do nothing since he is authorative and doesnt need a prediction key for anything.
+	// On the server, this will do nothing since he is authoritative and doesn't need a prediction key for anything.
 	// On the client, this will generate a new prediction key if bCanGenerateNewKey is true, and we have a invalid prediction key.
 
 	ClearScopedPredictionKey = false;
 	SetReplicatedPredictionKey = false;
 	Owner = InAbilitySystemComponent;
 
-	if (Owner->IsNetSimulating() == false)
+	if (!ensure(Owner.IsValid()) || Owner->IsNetSimulating() == false)
 	{
 		return;
 	}
@@ -225,6 +231,7 @@ FScopedPredictionWindow::FScopedPredictionWindow(UAbilitySystemComponent* InAbil
 	// InAbilitySystemComponent->GetPredictionKey().IsValidForMorePrediction() == false && 
 	if (bCanGenerateNewKey)
 	{
+		check(InAbilitySystemComponent != NULL); // Should have bailed above with ensure(Owner.IsValid())
 		ClearScopedPredictionKey = true;
 		RestoreKey = InAbilitySystemComponent->ScopedPredictionKey;
 		InAbilitySystemComponent->ScopedPredictionKey.GenerateDependentPredictionKey();
