Dev to Main integration from CL #2446290

diff --git a/Engine/Source/Runtime/GameplayAbilities/Public/Abilities/GameplayAbility.h b/Engine/Source/Runtime/GameplayAbilities/Public/Abilities/GameplayAbility.h
--- a/Engine/Source/Runtime/GameplayAbilities/Public/Abilities/GameplayAbility.h
+++ b/Engine/Source/Runtime/GameplayAbilities/Public/Abilities/GameplayAbility.h
@@ -123,7 +123,7 @@ public:
 		
 
 	/** Returns true if this ability can be activated right now. Has no side effects */
-	virtual bool CanActivateAbility(const FGameplayAbilitySpecHandle Handle, const FGameplayAbilityActorInfo* ActorInfo, OUT FGameplayTagContainer* OptionalRelevantTags = nullptr) const;
+	virtual bool CanActivateAbility(const FGameplayAbilitySpecHandle Handle, const FGameplayAbilityActorInfo* ActorInfo, FGameplayTagContainer SourceTags, FGameplayTagContainer TargetTags, OUT FGameplayTagContainer* OptionalRelevantTags = nullptr) const;
 
 	/** Returns true if this ability can be triggered right now. Has no side effects */
 	virtual bool ShouldAbilityRespondToEvent(const FGameplayAbilityActorInfo* ActorInfo, const FGameplayEventData* Payload) const;
@@ -135,7 +135,10 @@ public:
 	virtual void GetCooldownTimeRemainingAndDuration(FGameplayAbilitySpecHandle Handle, const FGameplayAbilityActorInfo* ActorInfo, float& TimeRemaining, float& CooldownDuration) const;
 
 	virtual const FGameplayTagContainer* GetCooldownTags() const;
-		
+	
+	/** Returns true if none of the ability's tags are blocked and if it doesn't have a "Blocking" tag and has all "Required" tags. */
+	bool DoesAbilitySatisfyTagRequirements(const UAbilitySystemComponent& AbilitySystemComponent, FGameplayTagContainer SourceTags, FGameplayTagContainer TargetTags, OUT FGameplayTagContainer* OptionalRelevantTags = nullptr) const;
+
 	EGameplayAbilityInstancingPolicy::Type GetInstancingPolicy() const
 	{
 		return InstancingPolicy;
@@ -306,6 +309,8 @@ protected:
 	UFUNCTION(BlueprintCallable, Category = "Abilities")
 	virtual void SendGameplayEvent(FGameplayTag EventTag, FGameplayEventData Payload);
 
+public:
+
 	// --------------------------------------
 	//	CommitAbility
 	// --------------------------------------
@@ -363,6 +368,7 @@ protected:
 	/** Destroys instanced-per-execution abilities. Instance-per-actor abilities should 'reset'. Any active ability state tasks receive the 'OnAbilityStateInterrupted' event. Non instance abilities - what can we do? */
 	virtual void CancelAbility(const FGameplayAbilitySpecHandle Handle, const FGameplayAbilityActorInfo* ActorInfo, const FGameplayAbilityActivationInfo ActivationInfo);
 
+protected:
 	/** Destroys instanced-per-execution abilities. Instance-per-actor abilities should 'reset'. Non instance abilities - what can we do? */
 	UFUNCTION(BlueprintCallable, Category = Ability)
 	void ConfirmTaskByInstanceName(FName InstanceName, bool bEndTask);
@@ -415,12 +421,12 @@ protected:
 	FActiveGameplayEffectHandle K2_ApplyGameplayEffectToOwner(const UGameplayEffect* GameplayEffect, int32 GameplayEffectLevel = 1);
 
 	/** Non blueprintcallable, safe to call on CDO/NonInstance abilities */
-	FActiveGameplayEffectHandle ApplyGameplayEffectToOwner(const FGameplayAbilitySpecHandle Handle, const FGameplayAbilityActorInfo* ActorInfo, const FGameplayAbilityActivationInfo ActivationInfo, const UGameplayEffect* GameplayEffect, float GameplayEffectLevel);
+	FActiveGameplayEffectHandle ApplyGameplayEffectToOwner(const FGameplayAbilitySpecHandle Handle, const FGameplayAbilityActorInfo* ActorInfo, const FGameplayAbilityActivationInfo ActivationInfo, const UGameplayEffect* GameplayEffect, float GameplayEffectLevel) const;
 
 	UFUNCTION(BlueprintCallable, Category = Ability, FriendlyName = "ApplyGameplayEffectSpecToOwner")
 	FActiveGameplayEffectHandle K2_ApplyGameplayEffectSpecToOwner(const FGameplayEffectSpecHandle EffectSpecHandle);
 
-	FActiveGameplayEffectHandle ApplyGameplayEffectSpecToOwner(const FGameplayAbilitySpecHandle AbilityHandle, const FGameplayAbilityActorInfo* ActorInfo, const FGameplayAbilityActivationInfo ActivationInfo, const FGameplayEffectSpecHandle SpecHandle);
+	FActiveGameplayEffectHandle ApplyGameplayEffectSpecToOwner(const FGameplayAbilitySpecHandle AbilityHandle, const FGameplayAbilityActorInfo* ActorInfo, const FGameplayAbilityActivationInfo ActivationInfo, const FGameplayEffectSpecHandle SpecHandle) const;
 
 	// ---------
 	// Apply Target
@@ -433,12 +439,12 @@ protected:
 	TArray<FActiveGameplayEffectHandle> K2_ApplyGameplayEffectToTarget(FGameplayAbilityTargetDataHandle TargetData, const UGameplayEffect* GameplayEffect, int32 GameplayEffectLevel = 1);
 
 	/** Non blueprintcallable, safe to call on CDO/NonInstance abilities */
-	TArray<FActiveGameplayEffectHandle> ApplyGameplayEffectToTarget(const FGameplayAbilitySpecHandle Handle, const FGameplayAbilityActorInfo* ActorInfo, const FGameplayAbilityActivationInfo ActivationInfo, FGameplayAbilityTargetDataHandle Target, TSubclassOf<UGameplayEffect> GameplayEffectClass, float GameplayEffectLevel);
+	TArray<FActiveGameplayEffectHandle> ApplyGameplayEffectToTarget(const FGameplayAbilitySpecHandle Handle, const FGameplayAbilityActorInfo* ActorInfo, const FGameplayAbilityActivationInfo ActivationInfo, FGameplayAbilityTargetDataHandle Target, TSubclassOf<UGameplayEffect> GameplayEffectClass, float GameplayEffectLevel) const;
 
 	UFUNCTION(BlueprintCallable, Category = Ability, FriendlyName = "ApplyGameplayEffectSpecToTarget")
 	TArray<FActiveGameplayEffectHandle> K2_ApplyGameplayEffectSpecToTarget(const FGameplayEffectSpecHandle EffectSpecHandle, FGameplayAbilityTargetDataHandle TargetData);
 
-	TArray<FActiveGameplayEffectHandle> ApplyGameplayEffectSpecToTarget(const FGameplayAbilitySpecHandle AbilityHandle, const FGameplayAbilityActorInfo* ActorInfo, const FGameplayAbilityActivationInfo ActivationInfo, const FGameplayEffectSpecHandle SpecHandle, FGameplayAbilityTargetDataHandle TargetData);
+	TArray<FActiveGameplayEffectHandle> ApplyGameplayEffectSpecToTarget(const FGameplayAbilitySpecHandle AbilityHandle, const FGameplayAbilityActorInfo* ActorInfo, const FGameplayAbilityActivationInfo ActivationInfo, const FGameplayEffectSpecHandle SpecHandle, FGameplayAbilityTargetDataHandle TargetData) const;
 
 	// ---------
 	// Remove Self
@@ -555,8 +561,6 @@ public:
 	UFUNCTION(BlueprintCallable, Category = Ability)
 	virtual void SetShouldBlockOtherAbilities(bool bShouldBlockAbilities);
 
-protected:
-
 	bool IsSupportedForNetworking() const override;
 
 	/** Returns the gameplay effect used to determine cooldown */
@@ -566,16 +570,18 @@ protected:
 	class UGameplayEffect* GetCostGameplayEffect() const;
 
 	/** Checks cooldown. returns true if we can be used again. False if not */
-	virtual bool CheckCooldown(const FGameplayAbilitySpecHandle Handle, const FGameplayAbilityActorInfo* ActorInfo) const;
+	virtual bool CheckCooldown(const FGameplayAbilitySpecHandle Handle, const FGameplayAbilityActorInfo* ActorInfo, OUT FGameplayTagContainer* OptionalRelevantTags = nullptr) const;
 
 	/** Applies CooldownGameplayEffect to the target */
-	virtual void ApplyCooldown(const FGameplayAbilitySpecHandle Handle, const FGameplayAbilityActorInfo* ActorInfo, const FGameplayAbilityActivationInfo ActivationInfo);
+	virtual void ApplyCooldown(const FGameplayAbilitySpecHandle Handle, const FGameplayAbilityActorInfo* ActorInfo, const FGameplayAbilityActivationInfo ActivationInfo) const;
 
 	/** Checks cost. returns true if we can pay for the abilty. False if not */
-	virtual bool CheckCost(const FGameplayAbilitySpecHandle Handle, const FGameplayAbilityActorInfo* ActorInfo) const;
+	virtual bool CheckCost(const FGameplayAbilitySpecHandle Handle, const FGameplayAbilityActorInfo* ActorInfo, OUT FGameplayTagContainer* OptionalRelevantTags = nullptr) const;
 
 	/** Applies the ability's cost to the target */
-	virtual void ApplyCost(const FGameplayAbilitySpecHandle Handle, const FGameplayAbilityActorInfo* ActorInfo, const FGameplayAbilityActivationInfo ActivationInfo);
+	virtual void ApplyCost(const FGameplayAbilitySpecHandle Handle, const FGameplayAbilityActorInfo* ActorInfo, const FGameplayAbilityActivationInfo ActivationInfo) const;
+
+protected:
 
 	// -----------------------------------------------	
 
@@ -655,6 +661,22 @@ protected:
 	UPROPERTY(EditDefaultsOnly, Category = Tags)
 	FGameplayTagContainer ActivationBlockedTags;
 
+	/** This ability can only be activated if the source actor/component has all of these tags */
+	UPROPERTY(EditDefaultsOnly, Category = Tags)
+	FGameplayTagContainer SourceRequiredTags;
+
+	/** This ability is blocked if the source actor/component has any of these tags */
+	UPROPERTY(EditDefaultsOnly, Category = Tags)
+	FGameplayTagContainer SourceBlockedTags;
+
+	/** This ability can only be activated if the target actor/component has all of these tags */
+	UPROPERTY(EditDefaultsOnly, Category = Tags)
+	FGameplayTagContainer TargetRequiredTags;
+
+	/** This ability is blocked if the target actor/component has any of these tags */
+	UPROPERTY(EditDefaultsOnly, Category = Tags)
+	FGameplayTagContainer TargetBlockedTags;
+
 
 	// ----------------------------------------------------------------------------------------------------------------
 	//
diff --git a/Engine/Source/Runtime/GameplayAbilities/Public/Abilities/GameplayAbilityTargetActor.h b/Engine/Source/Runtime/GameplayAbilities/Public/Abilities/GameplayAbilityTargetActor.h
--- a/Engine/Source/Runtime/GameplayAbilities/Public/Abilities/GameplayAbilityTargetActor.h
+++ b/Engine/Source/Runtime/GameplayAbilities/Public/Abilities/GameplayAbilityTargetActor.h
@@ -64,7 +64,7 @@ public:
 	FAbilityTargetData	TargetDataReadyDelegate;
 	FAbilityTargetData	CanceledDelegate;
 
-	virtual bool IsNetRelevantFor(const APlayerController* RealViewer, const AActor* Viewer, const FVector& SrcLocation) const override;
+	virtual bool IsNetRelevantFor(const APlayerController* RealViewer, const AActor* ViewTarget, const FVector& SrcLocation) const override;
 
 	UPROPERTY(BlueprintReadOnly, Category = "Targeting")
 	APlayerController* MasterPC;
diff --git a/Engine/Source/Runtime/GameplayAbilities/Public/Abilities/GameplayAbilityTargetTypes.h b/Engine/Source/Runtime/GameplayAbilities/Public/Abilities/GameplayAbilityTargetTypes.h
--- a/Engine/Source/Runtime/GameplayAbilities/Public/Abilities/GameplayAbilityTargetTypes.h
+++ b/Engine/Source/Runtime/GameplayAbilities/Public/Abilities/GameplayAbilityTargetTypes.h
@@ -65,7 +65,7 @@ struct GAMEPLAYABILITIES_API FGameplayAbilityTargetData
 
 	TArray<FActiveGameplayEffectHandle> ApplyGameplayEffect(const UGameplayEffect* GameplayEffect, const FGameplayEffectContextHandle& InEffectContext, float Level, FPredictionKey PredictionKey = FPredictionKey());
 
-	TArray<FActiveGameplayEffectHandle> ApplyGameplayEffectSpec(FGameplayEffectSpec& Spec, FPredictionKey PredictionKey = FPredictionKey());
+	TArray<FActiveGameplayEffectHandle> ApplyGameplayEffectSpec(const FGameplayEffectSpec& Spec, FPredictionKey PredictionKey = FPredictionKey());
 
 	virtual void AddTargetDataToContext(FGameplayEffectContextHandle& Context, bool bIncludeActorArray);
 
diff --git a/Engine/Source/Runtime/GameplayAbilities/Public/Abilities/GameplayAbilityTypes.h b/Engine/Source/Runtime/GameplayAbilities/Public/Abilities/GameplayAbilityTypes.h
--- a/Engine/Source/Runtime/GameplayAbilities/Public/Abilities/GameplayAbilityTypes.h
+++ b/Engine/Source/Runtime/GameplayAbilities/Public/Abilities/GameplayAbilityTypes.h
@@ -290,13 +290,13 @@ struct GAMEPLAYABILITIES_API FGameplayAbilitySpec : public FFastArraySerializerI
 	GENERATED_USTRUCT_BODY()
 
 	FGameplayAbilitySpec()
-	: Ability(nullptr), Level(1), InputID(INDEX_NONE), SourceObject(nullptr), InputPressed(false), ActiveCount(0)
+	: Ability(nullptr), Level(1), InputID(INDEX_NONE), SourceObject(nullptr), InputPressed(false), RemoveAfterActivation(false), ActiveCount(0)
 	{
 		
 	}
 
 	FGameplayAbilitySpec(UGameplayAbility* InAbility, int32 InLevel=1, int32 InInputID=INDEX_NONE, UObject* InSourceObject=nullptr)
-		: Ability(InAbility), Level(InLevel), InputID(InInputID), SourceObject(InSourceObject), InputPressed(false), ActiveCount(0)
+		: Ability(InAbility), Level(InLevel), InputID(InInputID), SourceObject(InSourceObject), InputPressed(false), RemoveAfterActivation(false), ActiveCount(0)
 	{
 		Handle.GenerateNewHandle();
 	}
@@ -325,6 +325,10 @@ struct GAMEPLAYABILITIES_API FGameplayAbilitySpec : public FFastArraySerializerI
 	UPROPERTY(NotReplicated)
 	bool	InputPressed;
 
+	/** If true, this ability should be removed as soon as it finishes executing */
+	UPROPERTY(NotReplicated)
+	bool	RemoveAfterActivation;
+
 	/** A count of the number of times this ability has been activated minus the number of times it has been ended. For instanced abilities this will be the number of currently active instances. Can't replicate until prediction accurately handles this.*/
 	UPROPERTY(NotReplicated)
 	uint8	ActiveCount;
@@ -515,9 +519,10 @@ struct GAMEPLAYABILITIES_API FGameplayEventData
 	GENERATED_USTRUCT_BODY()
 
 	FGameplayEventData()
-	: Instigator(NULL)
-	, Target(NULL)
-	, OptionalObject(NULL)
+	: Instigator(nullptr)
+	, Target(nullptr)
+	, OptionalObject(nullptr)
+	, OptionalObject2(nullptr)
 	, EventMagnitude(0.f)
 	{
 	}
@@ -538,6 +543,10 @@ struct GAMEPLAYABILITIES_API FGameplayEventData
 	UPROPERTY(EditAnywhere, BlueprintReadWrite, Category = GameplayAbilityTriggerPayload)
 	UObject* OptionalObject;
 
+	// A second optional ability-specific object to be passed though the event
+	UPROPERTY(EditAnywhere, BlueprintReadWrite, Category = GameplayAbilityTriggerPayload)
+	UObject* OptionalObject2;
+
 	// Polymorphic context information
 	UPROPERTY(EditAnywhere, BlueprintReadWrite, Category = GameplayAbilityTriggerPayload)
 	FGameplayEffectContextHandle ContextHandle;
@@ -553,14 +562,6 @@ struct GAMEPLAYABILITIES_API FGameplayEventData
 	// The magnitude of the triggering event
 	UPROPERTY(EditAnywhere, BlueprintReadWrite, Category = GameplayAbilityTriggerPayload)
 	float EventMagnitude;
-
-	// DEPRECATED
-	UPROPERTY()
-	float Var1_DEPRECATED;
-	UPROPERTY()
-	float Var2_DEPRECATED;
-	UPROPERTY()
-	float Var3_DEPRECATED;
 };
 
 /** 
diff --git a/Engine/Source/Runtime/GameplayAbilities/Public/Abilities/GameplayAbilityWorldReticle.h b/Engine/Source/Runtime/GameplayAbilities/Public/Abilities/GameplayAbilityWorldReticle.h
--- a/Engine/Source/Runtime/GameplayAbilities/Public/Abilities/GameplayAbilityWorldReticle.h
+++ b/Engine/Source/Runtime/GameplayAbilities/Public/Abilities/GameplayAbilityWorldReticle.h
@@ -35,7 +35,7 @@ public:
 
 	// ------------------------------
 
-	virtual bool IsNetRelevantFor(const APlayerController* RealViewer, const AActor* Viewer, const FVector& SrcLocation) const override;
+	virtual bool IsNetRelevantFor(const APlayerController* RealViewer, const AActor* ViewTarget, const FVector& SrcLocation) const override;
 
 	void InitializeReticle(AGameplayAbilityTargetActor* InTargetingActor, FWorldReticleParameters InParameters);
 
diff --git a/Engine/Source/Runtime/GameplayAbilities/Public/Abilities/GameplayAbility_CharacterJump.h b/Engine/Source/Runtime/GameplayAbilities/Public/Abilities/GameplayAbility_CharacterJump.h
--- a/Engine/Source/Runtime/GameplayAbilities/Public/Abilities/GameplayAbility_CharacterJump.h
+++ b/Engine/Source/Runtime/GameplayAbilities/Public/Abilities/GameplayAbility_CharacterJump.h
@@ -15,7 +15,7 @@ class GAMEPLAYABILITIES_API UGameplayAbility_CharacterJump : public UGameplayAbi
 
 public:
 
-	virtual bool CanActivateAbility(const FGameplayAbilitySpecHandle Handle, const FGameplayAbilityActorInfo* ActorInfo, OUT FGameplayTagContainer* OptionalRelevantTags = nullptr) const override;
+	virtual bool CanActivateAbility(const FGameplayAbilitySpecHandle Handle, const FGameplayAbilityActorInfo* ActorInfo, FGameplayTagContainer SourceTags, FGameplayTagContainer TargetTags, OUT FGameplayTagContainer* OptionalRelevantTags = nullptr) const override;
 	
 	virtual void ActivateAbility(const FGameplayAbilitySpecHandle Handle, const FGameplayAbilityActorInfo* OwnerInfo, const FGameplayAbilityActivationInfo ActivationInfo, const FGameplayEventData* TriggerEventData) override;
 
diff --git a/Engine/Source/Runtime/GameplayAbilities/Public/AbilitySystemBlueprintLibrary.h b/Engine/Source/Runtime/GameplayAbilities/Public/AbilitySystemBlueprintLibrary.h
--- a/Engine/Source/Runtime/GameplayAbilities/Public/AbilitySystemBlueprintLibrary.h
+++ b/Engine/Source/Runtime/GameplayAbilities/Public/AbilitySystemBlueprintLibrary.h
@@ -33,7 +33,7 @@ class GAMEPLAYABILITIES_API UAbilitySystemBlueprintLibrary : public UBlueprintFu
 	// -------------------------------------------------------------------------------
 
 	UFUNCTION(BlueprintPure, Category = "Ability|Attribute")
-	static float GetFloatAttribute(const class AActor* Actor, FGameplayAttribute Attribute);
+	static float GetFloatAttribute(const class AActor* Actor, FGameplayAttribute Attribute, bool& bSuccessfullyFoundAttribute);
 
 	// -------------------------------------------------------------------------------
 	//		TargetData
@@ -72,6 +72,10 @@ class GAMEPLAYABILITIES_API UAbilitySystemBlueprintLibrary : public UBlueprintFu
 	UFUNCTION(BlueprintPure, Category = "Ability|TargetData")
 	static TArray<AActor*> GetActorsFromTargetData(FGameplayAbilityTargetDataHandle TargetData, int32 Index);
 
+	/** Returns true if the given TargetData has the actor passed in targeted */
+	UFUNCTION(BlueprintPure, Category = "Ability|TargetData")
+	static bool DoesTargetDataContainActor(FGameplayAbilityTargetDataHandle TargetData, int32 Index, AActor* Actor);
+
 	/** Returns true if the given TargetData has at least 1 actor targeted */
 	UFUNCTION(BlueprintPure, Category = "Ability|TargetData")
 	static bool TargetDataHasActor(FGameplayAbilityTargetDataHandle TargetData, int32 Index);
diff --git a/Engine/Source/Runtime/GameplayAbilities/Public/AbilitySystemComponent.h b/Engine/Source/Runtime/GameplayAbilities/Public/AbilitySystemComponent.h
--- a/Engine/Source/Runtime/GameplayAbilities/Public/AbilitySystemComponent.h
+++ b/Engine/Source/Runtime/GameplayAbilities/Public/AbilitySystemComponent.h
@@ -38,9 +38,12 @@
  * 
  */
 
-
+/** Called when a targeting actor rejects target confirmation */
 DECLARE_MULTICAST_DELEGATE_OneParam(FTargetingRejectedConfirmation, int32);
 
+/** Called when ability fails to activate, passes along the failed ability and a tag explaining why */
+DECLARE_MULTICAST_DELEGATE_TwoParams(FAbilityFailedDelegate, const UGameplayAbility*, FGameplayTagContainer);
+
 /**
  *	The core ActorComponent for interfacing with the GameplayAbilities System
  *	NOTE: This feature is EXPERIMENTAL. Use at your own risk!
@@ -184,20 +187,26 @@ class GAMEPLAYABILITIES_API UAbilitySystemComponent : public UActorComponent, pu
 	UFUNCTION()
 	void OnRep_PredictionKey();
 
+	// A pending activation that cannot be activated yet, will be rechecked at a later point
 	struct FPendingAbilityInfo
 	{
 		bool operator==(const FPendingAbilityInfo& Other) const
 		{
+			// Don't compare event data, not valid to have multiple activations in flight with same key and handle but different event data
 			return PredictionKey == Other.PredictionKey	&& Handle == Other.Handle;
 		}
 
 		FPredictionKey	PredictionKey;
 		FGameplayAbilitySpecHandle Handle;
+		FGameplayEventData TriggerEventData;
 	};
 
 	// This is a list of GameplayAbilities that are predicted by the client and were triggered by abilities that were also predicted by the client
 	// When the server version of the predicted ability executes it should trigger copies of these and the copies will be associated with the correct prediction keys
-	TArray<FPendingAbilityInfo> PendingClientAbilities;
+	TArray<FPendingAbilityInfo> PendingClientActivatedAbilities;
+
+	// This is a list of GameplayAbilities that were activated on the server and can't yet execute on the client. It will try to execute these at a later point
+	TArray<FPendingAbilityInfo> PendingServerActivatedAbilities;
 
 	enum class EAbilityExecutionState : uint8
 	{
@@ -370,6 +379,7 @@ class GAMEPLAYABILITIES_API UAbilitySystemComponent : public UActorComponent, pu
 	// --------------------------------------------
 	
 	FOnActiveGameplayEffectRemoved* OnGameplayEffectRemovedDelegate(FActiveGameplayEffectHandle Handle);
+	FOnActiveGameplayEffectRemoved& OnAnyGameplayEffectRemovedDelegate();
 
 	UFUNCTION(BlueprintCallable, Category = GameplayEffects, meta=(FriendlyName = "ApplyGameplayEffectToTarget"))
 	FActiveGameplayEffectHandle BP_ApplyGameplayEffectToTarget(TSubclassOf<UGameplayEffect> GameplayEffectClass, UAbilitySystemComponent *Target, float Level, FGameplayEffectContextHandle Context);
@@ -428,7 +438,7 @@ class GAMEPLAYABILITIES_API UAbilitySystemComponent : public UActorComponent, pu
 	// GameplayCues can come from GameplayEffectSpecs
 
 	UFUNCTION(NetMulticast, unreliable)
-	void NetMulticast_InvokeGameplayCueExecuted_FromSpec(const FGameplayEffectSpec Spec, FPredictionKey PredictionKey);
+	void NetMulticast_InvokeGameplayCueExecuted_FromSpec(const FGameplayEffectSpecForRPC Spec, FPredictionKey PredictionKey);
 
 	// GameplayCues can also come on their own. These take an optional effect context to pass through hit result, etc
 	void ExecuteGameplayCue(const FGameplayTag GameplayCueTag, FGameplayEffectContextHandle EffectContext = FGameplayEffectContextHandle());
@@ -455,7 +465,7 @@ class GAMEPLAYABILITIES_API UAbilitySystemComponent : public UActorComponent, pu
 	UFUNCTION(NetMulticast, unreliable)
 	void NetMulticast_InvokeGameplayCueRemoved(const FGameplayTag GameplayCueTag, FPredictionKey PredictionKey);
 
-	void InvokeGameplayCueEvent(const FGameplayEffectSpec &Spec, EGameplayCueEvent::Type EventType);
+	void InvokeGameplayCueEvent(const FGameplayEffectSpecForRPC &Spec, EGameplayCueEvent::Type EventType);
 
 	void InvokeGameplayCueEvent(const FGameplayTag GameplayCueTag, EGameplayCueEvent::Type EventType, FGameplayEffectContextHandle EffectContext = FGameplayEffectContextHandle());
 
@@ -471,7 +481,7 @@ class GAMEPLAYABILITIES_API UAbilitySystemComponent : public UActorComponent, pu
 	/**
 	 *	GameplayAbilities
 	 *	
-	 *	The role of the AbilitySystemComponent wrt Abilities is to provide:
+	 *	The role of the AbilitySystemComponent with respect to Abilities is to provide:
 	 *		-Management of ability instances (whether per actor or per execution instance).
 	 *			-Someone *has* to keep track of these instances.
 	 *			-Non instanced abilities *could* be executed without any ability stuff in AbilitySystemComponent.
@@ -486,12 +496,27 @@ class GAMEPLAYABILITIES_API UAbilitySystemComponent : public UActorComponent, pu
 	/** Grants Ability. Returns handle that can be used in TryActivateAbility, etc. */
 	FGameplayAbilitySpecHandle GiveAbility(FGameplayAbilitySpec AbilitySpec);
 
-	void GetActivateableGameplayAbilitySpecsByAllMatchingTags(const FGameplayTagContainer& GameplayTagContainer, TArray < struct FGameplayAbilitySpec* >& MatchingGameplayAbilities) const;
-
-	/** Attempts to activate a gameplay ability that matches the given tag. Returns true if anything is activated */
+	/** Grants an ability and attempts to activate it exactly one time, which will cause it to be removed. Only valid on the server! */
+	FGameplayAbilitySpecHandle GiveAbilityAndActivateOnce(FGameplayAbilitySpec AbilitySpec);
+
+	// Gets all Activatable Gameplay Abilities that match all tags in GameplayTagContainer AND for which
+	// DoesAbilitySatisfyTagRequirements() is true.  The latter requirement allows this function to find the correct
+	// ability without requiring advanced knowledge.  For example, if there are two "Melee" abilities, one of which
+	// requires a weapon and one of which requires being unarmed, then those abilities can use Blocking and Required
+	// tags to determine when they can fire.  Using the Satisfying Tags requirements simplifies a lot of usage cases.
+	// For example, Behavior Trees can use various decorators to test an ability fetched using this mechanism as well
+	// as the Task to execute the ability without needing to know that there even is more than one such ability.
+	void GetActivatableGameplayAbilitySpecsByAllMatchingTags(const FGameplayTagContainer& GameplayTagContainer,
+											TArray < struct FGameplayAbilitySpec* >& MatchingGameplayAbilities) const;
+
+	/** Attempts to activate every gameplay ability that matches the given tag and DoesAbilitySatisfyTagRequirements().
+	  * Returns true if anything is activated.  Can activate more than one ability.*/
 	UFUNCTION(BlueprintCallable, Category = "Abilities")
 	bool TryActivateAbilityByTag(const FGameplayTagContainer& GameplayTagContainer);
 
+	/** Checks if the ability system is currently blocking InputID. Returns true if InputID is blocked, false otherwise.  */
+	bool IsAbilityInputBlocked(int32 InputID) const;
+
 	/** Attempts to activate the given ability */
 	bool TryActivateAbility(FGameplayAbilitySpecHandle AbilityToActivate, FPredictionKey InPredictionKey = FPredictionKey(), UGameplayAbility ** OutInstancedAbility = nullptr, FOnGameplayAbilityEnded* OnGameplayAbilityEndedDelegate = nullptr, const FGameplayEventData* TriggerEventData = nullptr);
 
@@ -588,7 +613,7 @@ protected:
 public:
 
 	/** Returns the list of all activatable abilities */
-	const TArray<FGameplayAbilitySpec>& GetActivatableAbilities();
+	const TArray<FGameplayAbilitySpec>& GetActivatableAbilities() const;
 
 	/** Returns an ability spec from a handle. If modifying call MarkAbilitySpecDirty */
 	FGameplayAbilitySpec* FindAbilitySpecFromHandle(FGameplayAbilitySpecHandle Handle);
@@ -615,7 +640,7 @@ public:
 	void	RemoteEndAbility(FGameplayAbilitySpecHandle AbilityToEnd, FGameplayAbilityActivationInfo ActivationInfo);
 
 	UFUNCTION(Server, reliable, WithValidation)
-	void	ServerEndAbility(FGameplayAbilitySpecHandle AbilityToEnd, FGameplayAbilityActivationInfo ActivationInfo);
+	void	ServerEndAbility(FGameplayAbilitySpecHandle AbilityToEnd, FGameplayAbilityActivationInfo ActivationInfo, FPredictionKey PredictionKey);
 
 	UFUNCTION(Client, reliable)
 	void	ClientEndAbility(FGameplayAbilitySpecHandle AbilityToEnd, FGameplayAbilityActivationInfo ActivationInfo);
@@ -626,7 +651,10 @@ public:
 	void	OnClientActivateAbilityFailed(FGameplayAbilitySpecHandle AbilityToActivate, FPredictionKey::KeyType PredictionKey);
 
 	UFUNCTION(Client, Reliable)
-	void	ClientActivateAbilitySucceed(FGameplayAbilitySpecHandle AbilityToActivate, int16 PredictionKey, FGameplayEventData TriggerEventData);
+	void	ClientActivateAbilitySucceed(FGameplayAbilitySpecHandle AbilityToActivate, int16 PredictionKey);
+
+	UFUNCTION(Client, Reliable)
+	void	ClientActivateAbilitySucceedWithEventData(FGameplayAbilitySpecHandle AbilityToActivate, int16 PredictionKey, FGameplayEventData TriggerEventData);
 
 	// ----------------------------------------------------------------------------------------------------------------
 
@@ -672,12 +700,14 @@ public:
 
 	/** A generic callback anytime an ability is commited (cost/cooldown applied) */
 	FGenericAbilityDelegate AbilityCommitedCallbacks;
+	FAbilityFailedDelegate AbilityFailedCallbacks;
 
 	/** Executes a gameplay event. Returns the number of successful ability activations triggered by the event */
 	int32 HandleGameplayEvent(FGameplayTag EventTag, const FGameplayEventData* Payload);
 
-	void NotifyAbilityCommit(UGameplayAbility* Ability);
-	void NotifyAbilityActivated(const FGameplayAbilitySpecHandle Handle, UGameplayAbility* Ability);
+	virtual void NotifyAbilityCommit(UGameplayAbility* Ability);
+	virtual void NotifyAbilityActivated(const FGameplayAbilitySpecHandle Handle, UGameplayAbility* Ability);
+	virtual void NotifyAbilityFailed(const FGameplayAbilitySpecHandle Handle, UGameplayAbility* Ability, FGameplayTagContainer FailureReason);
 
 	UPROPERTY()
 	TArray<AGameplayAbilityTargetActor*>	SpawnedTargetActors;
@@ -926,6 +956,10 @@ public:
 	UFUNCTION()
 	void OnRep_SimulatedTasks();
 
+#if ENABLE_VISUAL_LOG
+	void ClearDebugInstantEffects();
+#endif // ENABLE_VISUAL_LOG
+
 protected:
 
 	/** Actually pushes the final attribute value to the attribute set's property. Should not be called by outside code since this does not go through the attribute aggregator system. */
diff --git a/Engine/Source/Runtime/GameplayAbilities/Public/AbilitySystemGlobals.h b/Engine/Source/Runtime/GameplayAbilities/Public/AbilitySystemGlobals.h
--- a/Engine/Source/Runtime/GameplayAbilities/Public/AbilitySystemGlobals.h
+++ b/Engine/Source/Runtime/GameplayAbilities/Public/AbilitySystemGlobals.h
@@ -2,6 +2,7 @@
 
 #pragma once
 
+#include "GameplayTagsModule.h"
 #include "AbilitySystemGlobals.generated.h"
 
 class AActor;
@@ -16,56 +17,63 @@
 struct FAttributeSetInitter;
 struct FGameplayEffectSpec;
 
-/** Holds global data for the skill system. Can be configured per project via config file */
+/** Holds global data for the ability system. Can be configured per project via config file */
 UCLASS(config=Game)
 class GAMEPLAYABILITIES_API UAbilitySystemGlobals : public UObject
 {
 	GENERATED_UCLASS_BODY()
 
+	/** Gets the single instance of the globals object, will create it as necessary */
+	static UAbilitySystemGlobals& Get();
+
+	/** Should be called once as part of project setup to load global data tables and tags */
 	virtual void InitGlobalData();
 
-	/** Holds all of the valid gameplay-related tags that can be applied to assets */
-	UPROPERTY(config)
-	FString GlobalCurveTableName;
+	/** Returns true if InitGlobalData has been called */
+	bool IsAbilitySystemGlobalsInitialized()
+	{
+		return GlobalAttributeSetInitter.IsValid();
+	}
 
-	/** Holds information about the valid attributes' min and max values and stacking rules */
-	UPROPERTY(config)
-	FString GlobalAttributeMetaDataTableName;
-
-	/** Holds default values for attribute sets, keyed off of Name/Levels. */
-	UPROPERTY(config)
-	FString GlobalAttributeSetDefaultsTableName;
-
-	/** The class to instantiate as the global GameplayCue manager. This class is responsible for directing GameplayCue events, loading and cooking assets related to GameplayCues. */
-	UPROPERTY(config)
-	FStringAssetReference GlobalGameplayCueManagerName;
-
-	/** Look in these paths for GameplayCueNotifies */
-	UPROPERTY(config)
-	TArray<FString>	GameplayCueNotifyPaths;
-
-	/** Fully load all gameplay cues from the object library */
-	UPROPERTY(config)
-	bool	GameplayCueNotifyFullyLoad;
-
-	/** Set to true if you want clients to try to predict gameply effects done to targets. If false it will only predict self effects */
-	UPROPERTY(config)
-	bool	PredictTargetGameplayEffects;
-
-	/** The class to instantiate as the globals object. Defaults to this class but can be overridden */
-	UPROPERTY(config)
-	FStringClassReference AbilitySystemGlobalsClassName;
-
+	/** Returns the curvetable used as the default for scalable floats that don't specify a curve table */
 	UCurveTable* GetGlobalCurveTable();
 
+	/** Returns the data table defining attribute metadata (NOTE: Currently not in use) */
 	UDataTable* GetGlobalAttributeMetaDataTable();
 
+	/** Returns data used to initialize attributes to their default values */
+	FAttributeSetInitter* GetAttributeSetInitter() const;
+
+	/** Searches the passed in actor for an ability system component, will use the AbilitySystemInterface */
+	static UAbilitySystemComponent* GetAbilitySystemComponentFromActor(const AActor* Actor, bool LookForComponent=false);
+
+	/** Should allocate a project specific AbilityActorInfo struct. Caller is responsible for deallocation */
+	virtual FGameplayAbilityActorInfo * AllocAbilityActorInfo() const;
+
+	/** Should allocate a project specific GameplayEffectContext struct. Caller is responsible for deallocation */
+	virtual FGameplayEffectContext* AllocGameplayEffectContext() const;
+
+	/** Global callback that can handle game-specific code that needs to run before applying a gameplay effect spec */
+	virtual void GlobalPreGameplayEffectSpecApply(FGameplayEffectSpec& Spec, UAbilitySystemComponent* AbilitySystemComponent);
+
+	/** Returns true if the ability system should try to predict gameplay effects applied to non local targets */
+	bool ShouldPredictTargetGameplayEffects() const
+	{
+		return PredictTargetGameplayEffects;
+	}
+
+	/** Searches the passed in class to look for a UFunction implementing the gameplay cue tag, sets MatchedTag to the exact tag found */
+	UFunction* GetGameplayCueFunction(const FGameplayTag &Tag, UClass* Class, FName &MatchedTag);
+
+	/** Returns the gameplay cue manager singleton object, creating if necessary */
+	virtual UGameplayCueManager* GetGameplayCueManager();
+
+	/** Sets a default gameplay cue tag using the asset's name */
 	static void DeriveGameplayCueTagFromAssetName(FString AssetName, FGameplayTag& GameplayCueTag, FName& GameplayCueName);
 
-	bool IsAbilitySystemGlobalsInitialized()
-	{
-		return GlobalAttributeSetInitter.IsValid();
-	}
+	/** The class to instantiate as the globals object. Defaults to this class but can be overridden */
+	UPROPERTY(config)
+	FStringClassReference AbilitySystemGlobalsClassName;
 
 	void AutomationTestOnly_SetGlobalCurveTable(class UCurveTable *InTable)
 	{
@@ -77,23 +85,66 @@
 		GlobalAttributeMetaDataTable = InTable;
 	}
 
-	static UAbilitySystemGlobals& Get();
+	// Cheat functions
 
-	static UAbilitySystemComponent* GetAbilitySystemComponentFromActor(const AActor* Actor, bool LookForComponent=false);
+	/** Toggles whether we should ignore ability cooldowns. Does nothing in shipping builds */
+	UFUNCTION(exec)
+	virtual void ToggleIgnoreAbilitySystemCooldowns();
 
-	/** Should allocate a project specific AbilityActorInfo struct. Caller is responsible for deallocation */
-	virtual FGameplayAbilityActorInfo * AllocAbilityActorInfo() const;
+	/** Toggles whether we should ignore ability costs. Does nothing in shipping builds */
+	UFUNCTION(exec)
+	virtual void ToggleIgnoreAbilitySystemCosts();
 
-	/** Should allocate a project specific GameplayEffectContext struct. Caller is responsible for deallocation */
-	virtual FGameplayEffectContext* AllocGameplayEffectContext() const;
+	/** Returns true if ability cooldowns are ignored, returns false otherwise. Always returns false in shipping builds. */
+	bool ShouldIgnoreCooldowns() const;
 
-	UFunction* GetGameplayCueFunction(const FGameplayTag &Tag, UClass* Class, FName &MatchedTag);
+	/** Returns true if ability costs are ignored, returns false otherwise. Always returns false in shipping builds. */
+	bool ShouldIgnoreCosts() const;
 
-	FAttributeSetInitter* GetAttributeSetInitter() const;
+	// Global Tags
 
-	virtual UGameplayCueManager* GetGameplayCueManager();
+	UPROPERTY()
+	FGameplayTag ActivateFailCooldownTag; // TryActivate failed due to being on cooldown
+	UPROPERTY(config)
+	FName ActivateFailCooldownName;
 
-	virtual void GlobalPreGameplayEffectSpecApply(FGameplayEffectSpec& Spec, UAbilitySystemComponent* AbilitySystemComponent);
+	UPROPERTY()
+	FGameplayTag ActivateFailCostTag; // TryActivate failed due to not being able to spend costs
+	UPROPERTY(config)
+	FName ActivateFailCostName;
 
+	UPROPERTY()
+	FGameplayTag ActivateFailTagsBlockedTag; // TryActivate failed due to being blocked by other abilities
+	UPROPERTY(config)
+	FName ActivateFailTagsBlockedName;
+
+	UPROPERTY()
+	FGameplayTag ActivateFailTagsMissingTag; // TryActivate failed due to missing required tags
+	UPROPERTY(config)
+	FName ActivateFailTagsMissingName;
+
+	virtual void InitGlobalTags()
+	{
+		if (ActivateFailCooldownName != NAME_None)
+		{
+			ActivateFailCooldownTag = IGameplayTagsModule::RequestGameplayTag(ActivateFailCooldownName);
+		}
+
+		if (ActivateFailCostName != NAME_None)
+		{
+			ActivateFailCostTag = IGameplayTagsModule::RequestGameplayTag(ActivateFailCostName);
+		}
+
+		if (ActivateFailTagsBlockedName != NAME_None)
+		{
+			ActivateFailTagsBlockedTag = IGameplayTagsModule::RequestGameplayTag(ActivateFailTagsBlockedName);
+		}
+
+		if (ActivateFailTagsMissingName != NAME_None)
+		{
+			ActivateFailTagsMissingTag = IGameplayTagsModule::RequestGameplayTag(ActivateFailTagsMissingName);
+		}
+	}
+
 protected:
 
@@ -101,8 +152,46 @@
 	virtual void ReloadAttributeDefaults();
 	virtual void AllocAttributeSetInitter();
 
+#if !(UE_BUILD_SHIPPING || UE_BUILD_TEST)
+	// data used for ability system cheat commands
+
+	/** If we should ignore the cooldowns when activating abilities in the ability system. Set with ToggleIgnoreAbilitySystemCooldowns() */
+	bool bIgnoreAbilitySystemCooldowns;
+
+	/** If we should ignore the costs when activating abilities in the ability system. Set with ToggleIgnoreAbilitySystemCosts() */
+	bool bIgnoreAbilitySystemCosts;
+#endif // #if !(UE_BUILD_SHIPPING || UE_BUILD_TEST)
+
 private:
 
+	/** Holds all of the valid gameplay-related tags that can be applied to assets */
+	UPROPERTY(config)
+	FString GlobalCurveTableName;
+
+	/** Holds information about the valid attributes' min and max values and stacking rules */
+	UPROPERTY(config)
+	FString GlobalAttributeMetaDataTableName;
+
+	/** Holds default values for attribute sets, keyed off of Name/Levels. */
+	UPROPERTY(config)
+	FString GlobalAttributeSetDefaultsTableName;
+
+	/** The class to instantiate as the global GameplayCue manager. This class is responsible for directing GameplayCue events, loading and cooking assets related to GameplayCues. */
+	UPROPERTY(config)
+	FStringAssetReference GlobalGameplayCueManagerName;
+
+	/** Look in these paths for GameplayCueNotifies */
+	UPROPERTY(config)
+	TArray<FString>	GameplayCueNotifyPaths;
+
+	/** Fully load all gameplay cues from the object library */
+	UPROPERTY(config)
+	bool	GameplayCueNotifyFullyLoad;
+
+	/** Set to true if you want clients to try to predict gameplay effects done to targets. If false it will only predict self effects */
+	UPROPERTY(config)
+	bool	PredictTargetGameplayEffects;
+
 	UPROPERTY()
 	UCurveTable* GlobalCurveTable;
 
diff --git a/Engine/Source/Runtime/GameplayAbilities/Public/AbilitySystemLog.h b/Engine/Source/Runtime/GameplayAbilities/Public/AbilitySystemLog.h
--- a/Engine/Source/Runtime/GameplayAbilities/Public/AbilitySystemLog.h
+++ b/Engine/Source/Runtime/GameplayAbilities/Public/AbilitySystemLog.h
@@ -2,6 +2,8 @@
 
 #pragma once
 
+#include "VisualLogger.h"
+
 // Intended categories:
 //	Log - This happened. What gameplay programmers may care about to debug
 //	Verbose - This is why this happened. What you may turn on to debug the skill system code.
diff --git a/Engine/Source/Runtime/GameplayAbilities/Public/GameplayEffect.h b/Engine/Source/Runtime/GameplayAbilities/Public/GameplayEffect.h
--- a/Engine/Source/Runtime/GameplayAbilities/Public/GameplayEffect.h
+++ b/Engine/Source/Runtime/GameplayAbilities/Public/GameplayEffect.h
@@ -166,7 +166,7 @@ struct FSetByCallerFloat
 
 /** Struct representing the magnitude of a gameplay effect modifier, potentially calculated in numerous different ways */
 USTRUCT()
-struct FGameplayEffectModifierMagnitude
+struct GAMEPLAYABILITIES_API FGameplayEffectModifierMagnitude
 {
 	GENERATED_USTRUCT_BODY()
 
@@ -232,8 +232,14 @@ public:
 
 	EGameplayEffectMagnitudeCalculation GetMagnitudeCalculationType() const { return MagnitudeCalculationType; }
 
+	/** Returns the magnitude as it was entered in data. Only applies to ScalableFloat or any other type that can return data without context */
+	bool GetStaticMagnitudeIfPossible(float InLevel, float& OutMagnitude) const;
+
+	/** Returns the DataName associated with this magnitude if it is set by caller */
+	bool GetSetByCallerDataNameIfPossible(FName& OutDataName) const;
+
 #if WITH_EDITOR
-	GAMEPLAYABILITIES_API FText GetValueForEditorDisplay() const;
+	FText GetValueForEditorDisplay() const;
 #endif
 
 protected:
@@ -1026,6 +1032,7 @@ public:
 	UPROPERTY()
 	const UGameplayEffect* Def;
 	
+	/** A list of attributes that were modified during the application of this spec */
 	UPROPERTY()
 	TArray<FGameplayEffectModifiedAttribute> ModifiedAttributes;
 	
@@ -1093,6 +1100,49 @@ private:
 	float Level;	
 };
 
+
+/** This is a cut down version of the gameplay effect spec used for RPCs. */
+USTRUCT()
+struct GAMEPLAYABILITIES_API FGameplayEffectSpecForRPC
+{
+	GENERATED_USTRUCT_BODY()
+
+	FGameplayEffectSpecForRPC();
+
+	FGameplayEffectSpecForRPC(const FGameplayEffectSpec& InSpec);
+
+	/** GameplayEfect definition. The static data that this spec points to. */
+	UPROPERTY()
+	const UGameplayEffect* Def;
+
+	UPROPERTY()
+	TArray<FGameplayEffectModifiedAttribute> ModifiedAttributes;
+
+	UPROPERTY()
+	FGameplayEffectContextHandle EffectContext; // This tells us how we got here (who / what applied us)
+
+	UPROPERTY()
+	float Level;
+
+	FGameplayEffectContextHandle GetContext() const
+	{
+		return EffectContext;
+	}
+
+	float GetLevel() const
+	{
+		return Level;
+	}
+
+	FString ToSimpleString() const
+	{
+		return FString::Printf(TEXT("%s"), *Def->GetName());
+	}
+
+	const FGameplayEffectModifiedAttribute* GetModifiedAttribute(const FGameplayAttribute& Attribute) const;
+};
+
+
 /**
  * Active GameplayEffect instance
  *	-What GameplayEffect Spec
@@ -1269,7 +1319,7 @@ struct FActiveGameplayEffectQuery
  *
  */
 USTRUCT()
-struct FActiveGameplayEffectsContainer : public FFastArraySerializer
+struct GAMEPLAYABILITIES_API FActiveGameplayEffectsContainer : public FFastArraySerializer
 {
 	GENERATED_USTRUCT_BODY();
 
@@ -1288,26 +1338,32 @@ struct FActiveGameplayEffectsContainer : public FFastArraySerializer
 
 	UAbilitySystemComponent* Owner;
 
-#if ENABLE_VISUAL_LOG
+	FOnActiveGameplayEffectRemoved	OnActiveGameplayEffectRemovedDelegate;
+
 	struct DebugExecutedGameplayEffectData
 	{
 		FString GameplayEffectName;
 		FString ActivationState;
 		FGameplayAttribute Attribute;
+		TEnumAsByte<EGameplayModOp::Type> ModifierOp;
 		float Magnitude;
+		int32 StackCount;
 	};
-
+#if ENABLE_VISUAL_LOG
 	// Stores a record of gameplay effects that have executed and their results. Useful for debugging.
 	TArray<DebugExecutedGameplayEffectData> DebugExecutedGameplayEffects;
 
-	GAMEPLAYABILITIES_API void GrabDebugSnapshot(FVisualLogEntry* Snapshot) const;
+	void GrabDebugSnapshot(FVisualLogEntry* Snapshot) const;
 #endif // ENABLE_VISUAL_LOG
 
+	void GetActiveGameplayEffectDataByAttribute(TMultiMap<FGameplayAttribute, FActiveGameplayEffectsContainer::DebugExecutedGameplayEffectData>& EffectMap) const;
+
 	void RegisterWithOwner(UAbilitySystemComponent* Owner);	
 	
 	FActiveGameplayEffect* ApplyGameplayEffectSpec(const FGameplayEffectSpec& Spec, FPredictionKey InPredictionKey);
 
 	FActiveGameplayEffect* GetActiveGameplayEffect(const FActiveGameplayEffectHandle Handle);
+
 	const FActiveGameplayEffect* GetActiveGameplayEffect(const FActiveGameplayEffectHandle Handle) const;
 		
 	void ExecuteActiveEffectsFrom(FGameplayEffectSpec &Spec, FPredictionKey PredictionKey = FPredictionKey() );
diff --git a/Engine/Source/Runtime/GameplayAbilities/Public/GameplayEffectTypes.h b/Engine/Source/Runtime/GameplayAbilities/Public/GameplayEffectTypes.h
--- a/Engine/Source/Runtime/GameplayAbilities/Public/GameplayEffectTypes.h
+++ b/Engine/Source/Runtime/GameplayAbilities/Public/GameplayEffectTypes.h
@@ -258,19 +258,12 @@ struct GAMEPLAYABILITIES_API FGameplayEffectContext
 	GENERATED_USTRUCT_BODY()
 
 	FGameplayEffectContext()
-		: Instigator(NULL)
-		, EffectCauser(NULL)
-		, SourceObject(NULL)
-		, InstigatorAbilitySystemComponent(NULL)
-		, bHasWorldOrigin(false)
+		: bHasWorldOrigin(false)
 	{
 	}
 
 	FGameplayEffectContext(AActor* InInstigator, AActor* InEffectCauser)
-		: Instigator(NULL)
-		, EffectCauser(NULL)
-		, SourceObject(NULL)
-		, InstigatorAbilitySystemComponent(NULL)
+		: bHasWorldOrigin(false)
 	{
 		AddInstigator(InInstigator, InEffectCauser);
 	}
@@ -294,7 +287,7 @@ struct GAMEPLAYABILITIES_API FGameplayEffectContext
 	/** Returns the ability system component of the instigator of this effect */
 	virtual UAbilitySystemComponent* GetInstigatorAbilitySystemComponent() const
 	{
-		return InstigatorAbilitySystemComponent;
+		return InstigatorAbilitySystemComponent.Get();
 	}
 
 	/** Returns the physical actor tied to the application of this effect */
@@ -312,7 +305,7 @@ struct GAMEPLAYABILITIES_API FGameplayEffectContext
 	/** Returns the ability system component of the instigator that started the whole chain */
 	virtual UAbilitySystemComponent* GetOriginalInstigatorAbilitySystemComponent() const
 	{
-		return InstigatorAbilitySystemComponent;
+		return InstigatorAbilitySystemComponent.Get();
 	}
 
 	/** Sets the object this effect was created from. */
@@ -324,7 +317,7 @@ struct GAMEPLAYABILITIES_API FGameplayEffectContext
 	/** Returns the object this effect was created from. */
 	virtual const UObject* GetSourceObject() const
 	{
-		return SourceObject;
+		return SourceObject.Get();
 	}
 
 	virtual void AddActors(TArray<TWeakObjectPtr<AActor>> InActor, bool bReset = false);
@@ -386,6 +379,7 @@ struct GAMEPLAYABILITIES_API FGameplayEffectContext
 	virtual bool NetSerialize(FArchive& Ar, class UPackageMap* Map, bool& bOutSuccess);
 
 protected:
+	// The object pointers here have to be weak because contexts aren't necessarily tracked by GC in all cases
 
 	/** Instigator actor, the actor that owns the ability system component */
 	UPROPERTY()
@@ -397,11 +391,11 @@ protected:
 
 	/** Object this effect was created from, can be an actor or static object. Useful to bind an effect to a gameplay object */
 	UPROPERTY()
-	const UObject* SourceObject;
+	TWeakObjectPtr<const UObject> SourceObject;
 
 	/** The ability system component that's bound to instigator */
 	UPROPERTY(NotReplicated)
-	UAbilitySystemComponent* InstigatorAbilitySystemComponent;
+	TWeakObjectPtr<UAbilitySystemComponent> InstigatorAbilitySystemComponent;
 
 	UPROPERTY()
 	TArray<TWeakObjectPtr<AActor>> Actors;
diff --git a/Engine/Source/Runtime/GameplayAbilities/Public/TickableAttributeSetInterface.h b/Engine/Source/Runtime/GameplayAbilities/Public/TickableAttributeSetInterface.h
--- a/Engine/Source/Runtime/GameplayAbilities/Public/TickableAttributeSetInterface.h
+++ b/Engine/Source/Runtime/GameplayAbilities/Public/TickableAttributeSetInterface.h
@@ -5,7 +5,7 @@
 #include "TickableAttributeSetInterface.generated.h"
 
 /** Interface for actors which can be "spotted" by a player */
-UINTERFACE()
+UINTERFACE(MinimalAPI, meta = (CannotImplementInterfaceInBlueprint))
 class UTickableAttributeSetInterface : public UInterface
 {
 	GENERATED_UINTERFACE_BODY()
diff --git a/Engine/Source/Runtime/GameplayTags/Classes/BlueprintGameplayTagLibrary.h b/Engine/Source/Runtime/GameplayTags/Classes/BlueprintGameplayTagLibrary.h
--- a/Engine/Source/Runtime/GameplayTags/Classes/BlueprintGameplayTagLibrary.h
+++ b/Engine/Source/Runtime/GameplayTags/Classes/BlueprintGameplayTagLibrary.h
@@ -96,4 +96,17 @@ class UBlueprintGameplayTagLibrary : public UBlueprintFunctionLibrary
 	UFUNCTION(BlueprintPure, meta = (BlueprintInternalUseOnly = "TRUE"))
 	static bool HasAllMatchingGameplayTags(TScriptInterface<IGameplayTagAssetInterface> TagContainerInterface, const FGameplayTagContainer& OtherContainer, bool bCountEmptyAsMatch);
 
+	/**
+	* Check if the specified tag container has the specified tag, using the specified tag matching types
+	*
+	* @param TagContainerInterface		An Interface to a tag container
+	* @param ContainerTagsMatchType		Matching options to use for tags inside the container
+	* @param Tag						Tag to check for in the container
+	* @param TagMatchType				Matching option to use for the tag
+	*
+	* @return True if the container has the specified tag, false if it does not
+	*/
+	UFUNCTION(BlueprintPure, meta = (BlueprintInternalUseOnly = "TRUE"))
+	static bool DoesTagAssetInterfaceHaveTag(TScriptInterface<IGameplayTagAssetInterface> TagContainerInterface, TEnumAsByte<EGameplayTagMatchType::Type> ContainerTagsMatchType, const FGameplayTag& Tag, TEnumAsByte<EGameplayTagMatchType::Type> TagMatchType);
+
 };
diff --git a/Engine/Source/Runtime/GameplayTags/Classes/GameplayTagContainer.h b/Engine/Source/Runtime/GameplayTags/Classes/GameplayTagContainer.h
--- a/Engine/Source/Runtime/GameplayTags/Classes/GameplayTagContainer.h
+++ b/Engine/Source/Runtime/GameplayTags/Classes/GameplayTagContainer.h
@@ -21,6 +21,9 @@ enum class EGameplayContainerMatchType : uint8
 	All		//	Means the filter is only populated if all of the tags in this container match.
 };
 
+typedef uint16 FGameplayTagNetIndex;
+#define INVALID_TAGNETINDEX MAX_uint16
+
 USTRUCT(BlueprintType)
 struct GAMEPLAYTAGS_API FGameplayTag
 {
@@ -86,6 +89,9 @@ struct GAMEPLAYTAGS_API FGameplayTag
 	/** Returns direct parent GameplayTag of this GameplayTag */
 	FGameplayTag RequestDirectParent() const;
 
+	/** Overridden for fast serialize */
+	bool NetSerialize(FArchive& Ar, class UPackageMap* Map, bool& bOutSuccess);
+
 private:
 
 	/** Intentionally private so only the tag manager can use */
@@ -98,6 +104,15 @@ private:
 	friend class UGameplayTagsManager;
 };
 
+template<>
+struct TStructOpsTypeTraits< FGameplayTag > : public TStructOpsTypeTraitsBase
+{
+	enum
+	{
+		WithNetSerializer = true,
+	};
+};
+
 /** Simple struct for a gameplay tag container */
 USTRUCT(BlueprintType)
 struct GAMEPLAYTAGS_API FGameplayTagContainer
diff --git a/Engine/Source/Runtime/GameplayTags/Classes/GameplayTagsManager.h b/Engine/Source/Runtime/GameplayTags/Classes/GameplayTagsManager.h
--- a/Engine/Source/Runtime/GameplayTags/Classes/GameplayTagsManager.h
+++ b/Engine/Source/Runtime/GameplayTags/Classes/GameplayTagsManager.h
@@ -81,6 +81,13 @@ struct FGameplayTagNode
 	 */
 	GAMEPLAYTAGS_API TWeakPtr<FGameplayTagNode> GetParentTagNode() const;
 
+	/**
+	* Get the net index of this node
+	*
+	* @return The net index of this node
+	*/
+	GAMEPLAYTAGS_API FGameplayTagNetIndex GetNetIndex() const;
+
 	/** Reset the node of all of its values */
 	GAMEPLAYTAGS_API void ResetNode();
 
@@ -99,6 +106,11 @@ private:
 
 	/** Owner gameplay tag node, if any */
 	TWeakPtr<FGameplayTagNode> ParentNode;
+	
+	/** Net Index of this node */
+	FGameplayTagNetIndex NetIndex;
+
+	friend class UGameplayTagsManager;
 };
 
 /** Holds global data loaded at startup, is in a singleton UObject so it works properly with hot reload */
@@ -241,8 +253,15 @@ class GAMEPLAYTAGS_API UGameplayTagsManager : public UObject
 	/** Returns true if we should import tags from UGameplayTagsSettings objects (configured by INI files) */
 	static bool ShouldImportTagsFromINI();
 
+	/** Should use fast replication */
+	static bool ShouldUseFastReplication();
+
 	void RedirectTagsForContainer(FGameplayTagContainer& Container, TArray<FName>& DeprecatedTagNamesNotFoundInTagMap);
 
+	/** Gets a tag name from net index and vice versa, used for replication efficiency */
+	FName GetTagNameFromNetIndex(FGameplayTagNetIndex Index);
+	FGameplayTagNetIndex GetNetIndexFromTag(const FGameplayTag &InTag);
+
 private:
 
 	friend class FGameplayTagTest;
@@ -272,15 +291,21 @@ private:
 	 */
 	void GetAllParentNodeNames(TSet<FName>& NamesList, const TSharedPtr<FGameplayTagNode> GameplayTag) const;
 
+	/** Constructs the net indices for each tag */
+	void ConstructNetIndex();
+
 	/** Roots of gameplay tag nodes */
 	TSharedPtr<FGameplayTagNode> GameplayRootTag;
 
-	/** Map of Tags to Nodes - Internal use only*/
+	/** Map of Tags to Nodes - Internal use only */
 	TMap<FGameplayTag, TSharedPtr<FGameplayTagNode>> GameplayTagNodeMap;
 
-	/** Map of Names to tags - Internal use only*/
+	/** Map of Names to tags - Internal use only */
 	TMap<FName, FGameplayTag> GameplayTagMap;
 
+	/** Sorted list of nodes, used for network replication */
+	TArray<TSharedPtr<FGameplayTagNode>> NetworkGameplayTagNodeIndex;
+
 	/** Holds all of the valid gameplay-related tags that can be applied to assets */
 	UPROPERTY()
 	TArray<UDataTable*> GameplayTagTables;
@@ -288,6 +313,10 @@ private:
 	/** The delegate to execute when the tag tree changes */
 	FGameplayTagTreeChanged GameplayTagTreeChangedEvent;
 
+	// Prevent spamming about bad tag redirectors.  This flag is used to make all warnings only happen once, rather
+	// than once per every TagContainer that loads data!
+	bool bHasHandledRedirectors;
+
 #if WITH_EDITOR
 	/** Flag to say if we have registered the ObjectReimport, this is needed as use of Tags manager can happen before Editor is ready*/
 	bool RegisteredObjectReimport;
diff --git a/Engine/Source/Runtime/GameplayTags/Private/BlueprintGameplayTagLibrary.cpp b/Engine/Source/Runtime/GameplayTags/Private/BlueprintGameplayTagLibrary.cpp
--- a/Engine/Source/Runtime/GameplayTags/Private/BlueprintGameplayTagLibrary.cpp
+++ b/Engine/Source/Runtime/GameplayTags/Private/BlueprintGameplayTagLibrary.cpp
@@ -39,5 +39,26 @@ FGameplayTag UBlueprintGameplayTagLibrary::MakeLiteralGameplayTag(FGameplayTag V
 
 bool UBlueprintGameplayTagLibrary::HasAllMatchingGameplayTags(TScriptInterface<IGameplayTagAssetInterface> TagContainerInterface, const FGameplayTagContainer& OtherContainer, bool bCountEmptyAsMatch)
 {
+	if (TagContainerInterface.GetInterface() == NULL)
+	{
+		if (bCountEmptyAsMatch)
+		{
+			return (OtherContainer.Num() == 0);
+		}
+		return false;
+	}
+
 	return TagContainerInterface->HasAllMatchingGameplayTags(OtherContainer, bCountEmptyAsMatch);
-}
\ No newline at end of file
+}
+
+bool UBlueprintGameplayTagLibrary::DoesTagAssetInterfaceHaveTag(TScriptInterface<IGameplayTagAssetInterface> TagContainerInterface, TEnumAsByte<EGameplayTagMatchType::Type> ContainerTagsMatchType, const FGameplayTag& Tag, TEnumAsByte<EGameplayTagMatchType::Type> TagMatchType)
+{
+	if (TagContainerInterface.GetInterface() == NULL)
+	{
+		return false;
+	}
+
+	FGameplayTagContainer OwnedTags;
+	TagContainerInterface->GetOwnedGameplayTags(OwnedTags);
+	return (OwnedTags.HasTag(Tag, ContainerTagsMatchType, TagMatchType));
+}
diff --git a/Engine/Source/Runtime/GameplayTags/Private/GameplayTagContainer.cpp b/Engine/Source/Runtime/GameplayTags/Private/GameplayTagContainer.cpp
--- a/Engine/Source/Runtime/GameplayTags/Private/GameplayTagContainer.cpp
+++ b/Engine/Source/Runtime/GameplayTags/Private/GameplayTagContainer.cpp
@@ -361,3 +361,51 @@ FGameplayTag FGameplayTag::RequestDirectParent() const
 {
 	return IGameplayTagsModule::Get().GetGameplayTagsManager().RequestGameplayTagDirectParent(*this);
 }
+
+bool FGameplayTag::NetSerialize(FArchive& Ar, class UPackageMap* Map, bool& bOutSuccess)
+{
+	UGameplayTagsManager& TagManager = IGameplayTagsModule::Get().GetGameplayTagsManager();
+
+	uint8 bHasName = (TagName != NAME_None);
+	uint8 bHasNetIndex = 0;
+	FGameplayTagNetIndex NetIndex = INVALID_TAGNETINDEX;
+
+	if (Ar.IsSaving())
+	{
+		NetIndex = TagManager.GetNetIndexFromTag(*this);
+		if (NetIndex != INVALID_TAGNETINDEX)
+		{
+			// If we have a valid net index, serialize with that
+			bHasNetIndex = true;
+		}
+	}
+
+	// Serialize if we have a name at all or are empty
+	Ar.SerializeBits(&bHasName, 1);
+
+	if (bHasName)
+	{
+		Ar.SerializeBits(&bHasNetIndex, 1);
+		// If we have a net index serialize that, otherwise serialize as a name
+		if (bHasNetIndex)
+		{
+			Ar << NetIndex;
+		}
+		else
+		{
+			Ar << TagName;
+		}
+
+		if (Ar.IsLoading() && bHasNetIndex)
+		{
+			TagName = TagManager.GetTagNameFromNetIndex(NetIndex);
+		}
+	}
+	else
+	{
+		TagName = NAME_None;
+	}
+
+	bOutSuccess = true;
+	return true;
+}
\ No newline at end of file
diff --git a/Engine/Source/Runtime/GameplayTags/Private/GameplayTagsManager.cpp b/Engine/Source/Runtime/GameplayTags/Private/GameplayTagsManager.cpp
--- a/Engine/Source/Runtime/GameplayTags/Private/GameplayTagsManager.cpp
+++ b/Engine/Source/Runtime/GameplayTags/Private/GameplayTagsManager.cpp
@@ -10,7 +10,8 @@
 #endif
 
 UGameplayTagsManager::UGameplayTagsManager(const FObjectInitializer& ObjectInitializer)
-	: Super(ObjectInitializer)
+	: Super(ObjectInitializer),
+	bHasHandledRedirectors(false)
 {
 #if WITH_EDITOR
 	RegisteredObjectReimport = false;
@@ -147,10 +148,63 @@ void UGameplayTagsManager::ConstructGameplayTagTree()
 			GameplayRootTag->GetChildTagNodes().Sort(FCompareFGameplayTagNodeByTag());
 		}
 
+		if (ShouldUseFastReplication())
+		{
+			ConstructNetIndex();
+		}
+
 		GameplayTagTreeChangedEvent.Broadcast();
 	}
 }
 
+void UGameplayTagsManager::ConstructNetIndex()
+{
+	NetworkGameplayTagNodeIndex.Empty();
+
+	GameplayTagNodeMap.GenerateValueArray(NetworkGameplayTagNodeIndex);
+
+	NetworkGameplayTagNodeIndex.Sort(FCompareFGameplayTagNodeByTag());
+
+	// This is now sorted and it should be the same on both client and server
+
+	if (NetworkGameplayTagNodeIndex.Num() >= INVALID_TAGNETINDEX)
+	{
+		ensureMsgf(false, TEXT("Too many tags in dictionary for networking! Remove tags or increase tag net index size"));
+
+		NetworkGameplayTagNodeIndex.SetNum(INVALID_TAGNETINDEX - 1);
+	}
+
+	for (FGameplayTagNetIndex i = 0; i < NetworkGameplayTagNodeIndex.Num(); i++)
+	{
+		if (NetworkGameplayTagNodeIndex[i].IsValid())
+		{
+			NetworkGameplayTagNodeIndex[i]->NetIndex = i;
+		}
+	}
+}
+
+FName UGameplayTagsManager::GetTagNameFromNetIndex(FGameplayTagNetIndex Index)
+{
+	if (Index >= NetworkGameplayTagNodeIndex.Num())
+	{
+		ensureMsgf(false, TEXT("Received invalid tag net index %d! Tag index is out of sync on client!"), Index);
+
+		return NAME_None;
+	}
+	return NetworkGameplayTagNodeIndex[Index]->GetCompleteTag();
+}
+
+FGameplayTagNetIndex UGameplayTagsManager::GetNetIndexFromTag(const FGameplayTag &InTag)
+{
+	const TSharedPtr<FGameplayTagNode>* GameplayTagNode = GameplayTagNodeMap.Find(InTag);
+
+	if (GameplayTagNode && GameplayTagNode->IsValid())
+	{
+		return (*GameplayTagNode)->GetNetIndex();
+	}
+	return INVALID_TAGNETINDEX;
+}
+
 bool UGameplayTagsManager::ShouldImportTagsFromINI()
 {
 	bool ImportFromINI = false;
@@ -158,6 +212,13 @@ bool UGameplayTagsManager::ShouldImportTagsFromINI()
 	return ImportFromINI;
 }
 
+bool UGameplayTagsManager::ShouldUseFastReplication()
+{
+	bool ImportFromINI = false;
+	GConfig->GetBool(TEXT("GameplayTags"), TEXT("FastReplication"), ImportFromINI, GEngineIni);
+	return ImportFromINI;
+}
+
 void UGameplayTagsManager::RedirectTagsForContainer(FGameplayTagContainer& Container, TArray<FName>& DeprecatedTagNamesNotFoundInTagMap)
 {
 	FConfigSection* PackageRedirects = GConfig->GetSectionPrivate( TEXT("/Script/Engine.Engine"), false, true, GEngineIni );
@@ -176,11 +237,14 @@ void UGameplayTagsManager::RedirectTagsForContainer(FGameplayTagContainer& Conta
 			FGameplayTag OldTag = RequestGameplayTag(OldTagName, bEnsureIfNotFound); //< This only succeeds if OldTag is in the Table!
 			if (OldTag.IsValid())
 			{
-				UE_LOG(LogGameplayTags, Warning,
-					TEXT("Old tag (%s) which is being redirected still exists in the table!  Generally you should "
-						 TEXT("remove the old tags from the table when you are redirecting to new tags, or else users will ")
-						 TEXT("still be able to add the old tags to containers.")), *OldTagName.ToString()
-				      );
+				if (!bHasHandledRedirectors)
+				{
+					UE_LOG(LogGameplayTags, Warning,
+						TEXT("Old tag (%s) which is being redirected still exists in the table!  Generally you should "
+							 TEXT("remove the old tags from the table when you are redirecting to new tags, or else users will ")
+							 TEXT("still be able to add the old tags to containers.")), *OldTagName.ToString()
+						  );
+				}
 			}
 			else
 			{	// Create the old tag from scratch in order to be able to find and remove it.
@@ -223,11 +287,17 @@ void UGameplayTagsManager::RedirectTagsForContainer(FGameplayTagContainer& Conta
 			}
 			else // !NewTag.IsValid(), and NewTagName is not NAME_None
 			{
-				UE_LOG(LogGameplayTags, Warning, TEXT("Invalid new tag %s!  Cannot replace old tag %s."),
-					*NewTagName.ToString(), *OldTagName.ToString());
+				if (!bHasHandledRedirectors)
+				{
+					UE_LOG(LogGameplayTags, Warning, TEXT("Invalid new tag %s!  Cannot replace old tag %s."),
+						*NewTagName.ToString(), *OldTagName.ToString());
+				}
 			}
 		}
 	}
+
+	// Cache that this has run once so we don't spam warnings for every Tag Container.
+	bHasHandledRedirectors = true;
 }
 
 void UGameplayTagsManager::PopulateTreeFromDataTable(class UDataTable* InTable)
@@ -632,6 +702,7 @@ FGameplayTagNode::FGameplayTagNode(FName InTag, TWeakPtr<FGameplayTagNode> InPar
 	, CompleteTag(NAME_None)
 	, CategoryDescription(InCategoryDescription)
 	, ParentNode(InParentNode)
+	, NetIndex(INVALID_TAGNETINDEX)
 {
 	TArray<FName> Tags;
 
@@ -687,10 +758,16 @@ TWeakPtr<FGameplayTagNode> FGameplayTagNode::GetParentTagNode() const
 	return ParentNode;
 }
 
+FGameplayTagNetIndex FGameplayTagNode::GetNetIndex() const
+{
+	return NetIndex;
+}
+
 void FGameplayTagNode::ResetNode()
 {
 	Tag = NAME_None;
 	CompleteTag = NAME_None;
+	NetIndex = INVALID_TAGNETINDEX;
 
 	for (int32 ChildIdx = 0; ChildIdx < ChildTags.Num(); ++ChildIdx)
 	{
diff --git a/Engine/Source/Runtime/Launch/Private/LaunchEngineLoop.cpp b/Engine/Source/Runtime/Launch/Private/LaunchEngineLoop.cpp
--- a/Engine/Source/Runtime/Launch/Private/LaunchEngineLoop.cpp
+++ b/Engine/Source/Runtime/Launch/Private/LaunchEngineLoop.cpp
@@ -48,6 +48,7 @@
 	#include "HotReloadInterface.h"
 	#include "ISessionService.h"
 	#include "ISessionServicesModule.h"
+	#include "Engine/GameInstance.h"
 
 #if !UE_SERVER
 	#include "HeadMountedDisplay.h"
@@ -2110,34 +2111,38 @@ void FEngineLoop::Exit()
 	FIOSystem::Shutdown();
 }
 
-void FEngineLoop::ProcessPlayerControllersSlateOperations() const
+void FEngineLoop::ProcessLocalPlayerSlateOperations() const
 {
 	FSlateApplication& SlateApp = FSlateApplication::Get();
 
 	// For all the game worlds drill down to the player controller for each game viewport and process it's slate operation
-	for (const FWorldContext& Context : GEngine->GetWorldContexts())
+	for ( const FWorldContext& Context : GEngine->GetWorldContexts() )
 	{
 		UWorld* CurWorld = Context.World();
-		if (CurWorld != nullptr && CurWorld->IsGameWorld())
+		if ( CurWorld && CurWorld->IsGameWorld() )
 		{
-			ULocalPlayer* LocalPlayer = CurWorld->GetFirstLocalPlayerFromController();
 			UGameViewportClient* GameViewportClient = CurWorld->GetGameViewport();
-			if (LocalPlayer != nullptr && GameViewportClient != nullptr)
+			TSharedPtr< SViewport > ViewportWidget = GameViewportClient ? GameViewportClient->GetGameViewportWidget() : nullptr;
+
+			if ( ViewportWidget.IsValid() )
 			{
-				TSharedPtr<SViewport> ViewportWidget = GameViewportClient->GetGameViewportWidget();
-				if (ViewportWidget.IsValid())
+				for( FConstPlayerControllerIterator Iterator = CurWorld->GetPlayerControllerIterator(); Iterator; ++Iterator )
 				{
-					APlayerController* PlayerController = LocalPlayer->PlayerController;
-					if (PlayerController != nullptr)
+					APlayerController* PlayerController = *Iterator;
+					if( PlayerController )
 					{
-						FReply& TheReply = *PlayerController->SlateOperations;
+						ULocalPlayer* LocalPlayer = Cast< ULocalPlayer >( PlayerController->Player );
+						if( LocalPlayer )
+						{
+							FReply& TheReply = LocalPlayer->GetSlateOperations();
 
-						FWidgetPath PathToWidget;
-						SlateApp.GeneratePathToWidgetUnchecked(ViewportWidget.ToSharedRef(), PathToWidget);
+							FWidgetPath PathToWidget;
+							SlateApp.GeneratePathToWidgetUnchecked( ViewportWidget.ToSharedRef(), PathToWidget );
 
-						SlateApp.ProcessReply(PathToWidget, TheReply, NULL, NULL);
+							SlateApp.ProcessReply( PathToWidget, TheReply, nullptr, nullptr, LocalPlayer->GetControllerId() );
 
-						TheReply = FReply::Unhandled();
+							TheReply = FReply::Unhandled();
+						}
 					}
 				}
 			}
@@ -2278,7 +2283,7 @@ void FEngineLoop::Tick()
 			check(!IsRunningDedicatedServer());
 			
 			// Process slate operations accumulated in the world ticks.
-			ProcessPlayerControllersSlateOperations();
+			ProcessLocalPlayerSlateOperations();
 
 			// Tick Slate application
 			FSlateApplication::Get().Tick();
diff --git a/Engine/Source/Runtime/Launch/Public/LaunchEngineLoop.h b/Engine/Source/Runtime/Launch/Public/LaunchEngineLoop.h
--- a/Engine/Source/Runtime/Launch/Public/LaunchEngineLoop.h
+++ b/Engine/Source/Runtime/Launch/Public/LaunchEngineLoop.h
@@ -105,7 +105,7 @@ public:
 private:
 
 	/** Utility function that processes Slate operations. */
-	void ProcessPlayerControllersSlateOperations() const;
+	void ProcessLocalPlayerSlateOperations() const;
 
 protected:
 
diff --git a/Engine/Source/Runtime/Navmesh/Private/Detour/DetourNavMesh.cpp b/Engine/Source/Runtime/Navmesh/Private/Detour/DetourNavMesh.cpp
--- a/Engine/Source/Runtime/Navmesh/Private/Detour/DetourNavMesh.cpp
+++ b/Engine/Source/Runtime/Navmesh/Private/Detour/DetourNavMesh.cpp
@@ -1068,7 +1068,7 @@ void dtNavMesh::connectExtOffMeshLinks(dtMeshTile* tile, dtMeshTile* target, int
 		
 		const dtLink& targetLink = getLink(target, targetPoly->firstLink);
 		const dtPolyRef targetLandPoly = targetLink.ref;
-		const float ext[3] = { targetCon->rad, (targetCon->height < 0) ? target->header->walkableClimb : targetCon->height, targetCon->rad };
+		const float ext[3] = { targetCon->rad, targetCon->height, targetCon->rad };
 
 		// Find polygon to connect to.
 		const float* p = &targetCon->pos[3];
@@ -1182,7 +1182,7 @@ void dtNavMesh::baseOffMeshLinks(dtMeshTile* tile)
 		dtOffMeshConnection* con = &tile->offMeshCons[i];
 		dtPoly* poly = &tile->polys[con->poly];
 	
-		const float ext[3] = { con->rad, (con->height < 0) ? tile->header->walkableClimb : con->height, con->rad };
+		const float ext[3] = { con->rad, con->height, con->rad };
 		
 		// Find polygon to connect to.
 		const float* p = &con->pos[0]; // First vertex
diff --git a/Engine/Source/Runtime/Navmesh/Private/Detour/DetourNavMeshBuilder.cpp b/Engine/Source/Runtime/Navmesh/Private/Detour/DetourNavMeshBuilder.cpp
--- a/Engine/Source/Runtime/Navmesh/Private/Detour/DetourNavMeshBuilder.cpp
+++ b/Engine/Source/Runtime/Navmesh/Private/Detour/DetourNavMeshBuilder.cpp
@@ -359,8 +359,11 @@ bool dtCreateNavMeshData(dtNavMeshCreateParams* params, unsigned char** outData,
 				// Zero out off-mesh start positions which are not even potentially touching the mesh.
 				if (offMeshConClass[i*2+0] == 0xff)
 				{
-					if (offMeshCon.vertsA0[1] < bmin[1] || offMeshCon.vertsA0[1] > bmax[1])
-						offMeshConClass[i*2+0] = 0;
+					if ((offMeshCon.vertsA0[1] - offMeshCon.snapHeight) < bmin[1] &&
+						(offMeshCon.vertsA0[1] + offMeshCon.snapHeight) > bmax[1])
+					{
+						offMeshConClass[i * 2 + 0] = 0;
+					}
 				}
 
 				// Count how many links should be allocated for off-mesh connections.
diff --git a/Engine/Source/Runtime/Navmesh/Private/DetourCrowd/DetourCrowd.cpp b/Engine/Source/Runtime/Navmesh/Private/DetourCrowd/DetourCrowd.cpp
--- a/Engine/Source/Runtime/Navmesh/Private/DetourCrowd/DetourCrowd.cpp
+++ b/Engine/Source/Runtime/Navmesh/Private/DetourCrowd/DetourCrowd.cpp
@@ -1730,6 +1730,9 @@ void dtCrowd::updateStepOffMeshAnim(const float dt, dtCrowdAgentDebugInfo*)
 
 void dtCrowd::updateStepOffMeshVelocity(const float dt, dtCrowdAgentDebugInfo*)
 {
+	float DirLink[3];
+	float DirToEnd[3];
+
 	// UE4 version of offmesh anims, updates velocity and checks distance instead of fixed time
 	for (int i = 0; i < m_numActiveAgents; ++i)
 	{
@@ -1742,11 +1745,16 @@ void dtCrowd::updateStepOffMeshVelocity(const float dt, dtCrowdAgentDebugInfo*)
 
 		anim->t += dt;
 		
+		dtVsub(DirLink, anim->endPos, anim->startPos);
+		dtVsub(DirToEnd, anim->endPos, ag->npos);
+		const float dirDot = dtVdot2D(DirLink, DirToEnd);
 		const float dist = dtVdist2DSqr(ag->npos, anim->endPos);
 		const float distThres = dtSqr(5.0f);
 		const float heightDiff = dtAbs(ag->npos[1] - anim->endPos[1]);
 		const float heightThres = ag->params.height * 0.5f;
-		if (dist < distThres && heightDiff < heightThres)
+
+		// check height difference + distance or overshooting 
+		if ((dist < distThres || dirDot < 0.0f) && heightDiff < heightThres)
 		{
 			// Reset animation
 			anim->active = 0;
diff --git a/Engine/Source/Runtime/Online/Android/OnlineSubsystemGooglePlay/Private/OnlineSubsystemGooglePlay.cpp b/Engine/Source/Runtime/Online/Android/OnlineSubsystemGooglePlay/Private/OnlineSubsystemGooglePlay.cpp
--- a/Engine/Source/Runtime/Online/Android/OnlineSubsystemGooglePlay/Private/OnlineSubsystemGooglePlay.cpp
+++ b/Engine/Source/Runtime/Online/Android/OnlineSubsystemGooglePlay/Private/OnlineSubsystemGooglePlay.cpp
@@ -51,6 +51,10 @@ IOnlineFriendsPtr FOnlineSubsystemGooglePlay::GetFriendsInterface() const
 	return nullptr;
 }
 
+IOnlineGroupsPtr FOnlineSubsystemGooglePlay::GetGroupsInterface() const
+{
+	return nullptr;
+}
 
 IOnlineSharedCloudPtr FOnlineSubsystemGooglePlay::GetSharedCloudInterface() const
 {
@@ -179,6 +183,11 @@ bool FOnlineSubsystemGooglePlay::Init()
 
 bool FOnlineSubsystemGooglePlay::Tick(float DeltaTime)
 {
+	if (!FOnlineSubsystemImpl::Tick(DeltaTime))
+	{
+		return false;
+	}
+
 	if (OnlineAsyncTaskThreadRunnable)
 	{
 		OnlineAsyncTaskThreadRunnable->GameTick();
diff --git a/Engine/Source/Runtime/Online/Android/OnlineSubsystemGooglePlay/Public/OnlineSubsystemGooglePlay.h b/Engine/Source/Runtime/Online/Android/OnlineSubsystemGooglePlay/Public/OnlineSubsystemGooglePlay.h
--- a/Engine/Source/Runtime/Online/Android/OnlineSubsystemGooglePlay/Public/OnlineSubsystemGooglePlay.h
+++ b/Engine/Source/Runtime/Online/Android/OnlineSubsystemGooglePlay/Public/OnlineSubsystemGooglePlay.h
@@ -19,8 +19,7 @@ class FRunnableThread;
  * OnlineSubsystemGooglePlay - Implementation of the online subsystem for Google Play services
  */
 class ONLINESUBSYSTEMGOOGLEPLAY_API FOnlineSubsystemGooglePlay : 
-	public FOnlineSubsystemImpl,
-	public FTickerObjectBase
+	public FOnlineSubsystemImpl
 {
 public:
 	FOnlineSubsystemGooglePlay();
@@ -29,6 +28,7 @@ public:
 	// Begin IOnlineSubsystem Interface
 	virtual IOnlineSessionPtr GetSessionInterface() const override;
 	virtual IOnlineFriendsPtr GetFriendsInterface() const override;
+	virtual IOnlineGroupsPtr GetGroupsInterface() const override;
 	virtual IOnlineSharedCloudPtr GetSharedCloudInterface() const override;
 	virtual IOnlineUserCloudPtr GetUserCloudInterface() const override;
 	virtual IOnlineLeaderboardsPtr GetLeaderboardsInterface() const override;
diff --git a/Engine/Source/Runtime/Online/IOS/OnlineSubsystemIOS/Private/OnlineSubsystemIOS.cpp b/Engine/Source/Runtime/Online/IOS/OnlineSubsystemIOS/Private/OnlineSubsystemIOS.cpp
--- a/Engine/Source/Runtime/Online/IOS/OnlineSubsystemIOS/Private/OnlineSubsystemIOS.cpp
+++ b/Engine/Source/Runtime/Online/IOS/OnlineSubsystemIOS/Private/OnlineSubsystemIOS.cpp
@@ -20,6 +20,10 @@ IOnlineFriendsPtr FOnlineSubsystemIOS::GetFriendsInterface() const
 	return FriendsInterface;
 }
 
+IOnlineGroupsPtr FOnlineSubsystemIOS::GetGroupsInterface() const
+{
+	return nullptr;
+}
 
 IOnlineSharedCloudPtr FOnlineSubsystemIOS::GetSharedCloudInterface() const
 {
@@ -158,6 +162,11 @@ bool FOnlineSubsystemIOS::Init()
 
 bool FOnlineSubsystemIOS::Tick(float DeltaTime)
 {
+	if (!FOnlineSubsystemImpl::Tick(DeltaTime))
+	{
+		return false;
+	}
+
 	if (SessionInterface.IsValid())
 	{
 		SessionInterface->Tick(DeltaTime);
diff --git a/Engine/Source/Runtime/Online/IOS/OnlineSubsystemIOS/Public/OnlineSubsystemIOS.h b/Engine/Source/Runtime/Online/IOS/OnlineSubsystemIOS/Public/OnlineSubsystemIOS.h
--- a/Engine/Source/Runtime/Online/IOS/OnlineSubsystemIOS/Public/OnlineSubsystemIOS.h
+++ b/Engine/Source/Runtime/Online/IOS/OnlineSubsystemIOS/Public/OnlineSubsystemIOS.h
@@ -11,8 +11,7 @@
  *	OnlineSubsystemIOS - Implementation of the online subsystem for IOS services
  */
 class ONLINESUBSYSTEMIOS_API FOnlineSubsystemIOS : 
-	public FOnlineSubsystemImpl,
-	public FTickerObjectBase
+	public FOnlineSubsystemImpl
 {
 
 public:
@@ -22,6 +21,7 @@ public:
 	// Begin IOnlineSubsystem Interface
 	virtual IOnlineSessionPtr GetSessionInterface() const override;
 	virtual IOnlineFriendsPtr GetFriendsInterface() const override;
+	virtual IOnlineGroupsPtr GetGroupsInterface() const override;
 	virtual IOnlineSharedCloudPtr GetSharedCloudInterface() const override;
 	virtual IOnlineUserCloudPtr GetUserCloudInterface() const override;
 	virtual IOnlineLeaderboardsPtr GetLeaderboardsInterface() const override;
diff --git a/Engine/Source/Runtime/Online/ImageDownload/ImageDownload.Build.cs b/Engine/Source/Runtime/Online/ImageDownload/ImageDownload.Build.cs
new file mode 100644
--- /dev/null
+++ b/Engine/Source/Runtime/Online/ImageDownload/ImageDownload.Build.cs
@@ -0,0 +1,20 @@
+// Copyright 1998-2015 Epic Games, Inc. All Rights Reserved.
+
+using UnrealBuildTool;
+
+public class ImageDownload : ModuleRules
+{
+	public ImageDownload(TargetInfo Target)
+    {
+        Definitions.Add("IMAGEDOWNLOAD_PACKAGE=1");
+
+        PrivateDependencyModuleNames.AddRange(
+			new string[] { 
+				"Core",
+				"HTTP",
+				"SlateCore",
+				"ImageWrapper",
+			}
+			);
+    }
+}
diff --git a/Engine/Source/Runtime/Online/ImageDownload/Private/ImageDownloadPCH.h b/Engine/Source/Runtime/Online/ImageDownload/Private/ImageDownloadPCH.h
new file mode 100644
--- /dev/null
+++ b/Engine/Source/Runtime/Online/ImageDownload/Private/ImageDownloadPCH.h
@@ -0,0 +1,13 @@
+// Copyright 1998-2015 Epic Games, Inc. All Rights Reserved.
+
+#pragma once
+
+#include "Core.h"
+#include "ImageWrapper.h"
+#include "IHttpRequest.h"
+#include "IHttpResponse.h"
+#include "HttpModule.h"
+#include "SlateCore.h"
+
+DECLARE_LOG_CATEGORY_EXTERN(LogImageDownload, Log, All);
+
diff --git a/Engine/Source/Runtime/Online/ImageDownload/Private/WebImage.cpp b/Engine/Source/Runtime/Online/ImageDownload/Private/WebImage.cpp
new file mode 100644
--- /dev/null
+++ b/Engine/Source/Runtime/Online/ImageDownload/Private/WebImage.cpp
@@ -0,0 +1,146 @@
+// Copyright 1998-2014 Epic Games, Inc. All Rights Reserved.
+
+#include "ImageDownloadPCH.h"
+#include "WebImage.h"
+
+static const FString IMAGE_MIME_PNG(TEXT("image/png"));
+static const FString IMAGE_MIME_PNG_ALT(TEXT("image/x-png"));
+static const FString IMAGE_MIME_JPEG(TEXT("image/jpeg"));
+static const FString IMAGE_MIME_BMP(TEXT("image/bmp"));
+
+FWebImage::FWebImage()
+: StandInBrush(FCoreStyle::Get().GetDefaultBrush())
+, bDownloadSuccess(false)
+{
+}
+
+TAttribute< const FSlateBrush* > FWebImage::Attr() const
+{
+	return TAttribute< const FSlateBrush* >(AsShared(), &FWebImage::GetBrush);
+}
+
+bool FWebImage::BeginDownload(const FString& UrlIn, const FOnImageDownloaded& DownloadCb)
+{
+	CancelDownload();
+
+	// store the url
+	Url = UrlIn;
+
+	// make a new request
+	TSharedRef<IHttpRequest> HttpRequest = FHttpModule::Get().CreateRequest();
+	HttpRequest->SetVerb(TEXT("GET"));
+	HttpRequest->SetURL(Url);
+	HttpRequest->SetHeader(TEXT("Accept"), FString::Printf(TEXT("%s, %s, %s; q=0.8, %s; q=0.5"), *IMAGE_MIME_PNG, *IMAGE_MIME_PNG_ALT, *IMAGE_MIME_JPEG, *IMAGE_MIME_BMP));
+	HttpRequest->OnProcessRequestComplete().BindSP(this, &FWebImage::HttpRequestComplete, DownloadCb);
+
+	// queue the request
+	PendingRequest = HttpRequest;
+	if (!HttpRequest->ProcessRequest())
+	{
+		// clear our local handle
+		PendingRequest.Reset();
+		return false;
+	}
+	return true;
+}
+
+void FWebImage::HttpRequestComplete(FHttpRequestPtr HttpRequest, FHttpResponsePtr HttpResponse, bool bSucceeded, FOnImageDownloaded DownloadCb)
+{
+	// clear our handle to the request
+	PendingRequest.Reset();
+
+	// get the request URL
+	check(HttpRequest.IsValid()); // this should be valid, we did just send a request...
+	bool bSuccess = ProcessHttpResponse(HttpRequest->GetURL(), bSucceeded ? HttpResponse : FHttpResponsePtr());
+
+	// save this info
+	bDownloadSuccess = bSuccess;
+	DownloadTimeUtc = FDateTime::UtcNow();
+
+	// fire the response delegate
+	if (DownloadCb.IsBound())
+	{
+		DownloadCb.Execute(bSuccess);
+	}
+}
+
+bool FWebImage::ProcessHttpResponse(const FString& RequestUrl, FHttpResponsePtr HttpResponse)
+{
+	// check for successful response
+	if (!HttpResponse.IsValid())
+	{
+		UE_LOG(LogImageDownload, Error, TEXT("Image Download: Connection Failed. url=%s"), *RequestUrl);
+		return false;
+	}
+
+	// check status code
+	int32 StatusCode = HttpResponse->GetResponseCode();
+	if (StatusCode / 100 != 2)
+	{
+		UE_LOG(LogImageDownload, Error, TEXT("Image Download: HTTP response %d. url=%s"), StatusCode, *RequestUrl);
+		return false;
+	}
+
+	// get the content type and figure out the image format from it
+	FString ContentType = HttpResponse->GetContentType();
+	EImageFormat::Type ImageFormat;
+	if (ContentType.Equals(IMAGE_MIME_PNG, ESearchCase::IgnoreCase) || ContentType.Equals(IMAGE_MIME_PNG_ALT, ESearchCase::IgnoreCase))
+	{
+		ImageFormat = EImageFormat::PNG;
+	}
+	else if (ContentType.Equals(IMAGE_MIME_JPEG, ESearchCase::IgnoreCase))
+	{
+		ImageFormat = EImageFormat::JPEG;
+	}
+	else if (ContentType.Equals(IMAGE_MIME_BMP, ESearchCase::IgnoreCase))
+	{
+		ImageFormat = EImageFormat::BMP;
+	}
+	else
+	{
+		UE_LOG(LogImageDownload, Error, TEXT("Image Download: Invalid Content-Type '%s'. url=%s"), *ContentType, *RequestUrl);
+		return false;
+	}
+
+	// build an image wrapper for this type
+	static const FName MODULE_IMAGE_WRAPPER("ImageWrapper");
+	IImageWrapperModule& ImageWrapperModule = FModuleManager::LoadModuleChecked<IImageWrapperModule>(MODULE_IMAGE_WRAPPER);
+	IImageWrapperPtr ImageWrapper = ImageWrapperModule.CreateImageWrapper(ImageFormat);
+	if (!ImageWrapper.IsValid())
+	{
+		UE_LOG(LogImageDownload, Error, TEXT("Image Download: Unable to make image wrapper for %s"), *ContentType);
+		return false;
+	}
+
+	// parse the content
+	const TArray<uint8>& Content = HttpResponse->GetContent();
+	if (!ImageWrapper->SetCompressed(Content.GetData(), Content.Num()))
+	{
+		UE_LOG(LogImageDownload, Error, TEXT("Image Download: Unable to parse %s from %s"), *ContentType, *RequestUrl);
+		return false;
+	}
+
+	// get the raw image data
+	const TArray<uint8>* RawImageData = nullptr;
+	if (!ImageWrapper->GetRaw(ERGBFormat::BGRA, 8, RawImageData) || RawImageData == nullptr)
+	{
+		UE_LOG(LogImageDownload, Error, TEXT("Image Download: Unable to convert %s to BGRA 8"), *ContentType);
+		return false;
+	}
+
+	// make a dynamic image
+	FName ResourceName(*RequestUrl);
+	DownloadedBrush = FSlateDynamicImageBrush::CreateWithImageData(ResourceName, FVector2D(ImageWrapper->GetWidth(), ImageWrapper->GetHeight()), *RawImageData);
+	return DownloadedBrush.IsValid();
+}
+
+void FWebImage::CancelDownload()
+{
+	if (PendingRequest.IsValid())
+	{
+		FHttpRequestPtr Request = PendingRequest;
+		PendingRequest.Reset();
+		Request->CancelRequest();
+	}
+	bDownloadSuccess = false;
+}
diff --git a/Engine/Source/Runtime/Online/ImageDownload/Public/WebImage.h b/Engine/Source/Runtime/Online/ImageDownload/Public/WebImage.h
new file mode 100644
--- /dev/null
+++ b/Engine/Source/Runtime/Online/ImageDownload/Public/WebImage.h
@@ -0,0 +1,96 @@
+// Copyright 1998-2015 Epic Games, Inc. All Rights Reserved.
+
+#pragma once
+
+class IHttpRequest;
+typedef TSharedPtr<class IHttpRequest> FHttpRequestPtr;
+typedef TSharedPtr<class IHttpResponse, ESPMode::ThreadSafe> FHttpResponsePtr;
+
+/** 
+ * This class manages downloading an image and swapping it out a standin once it's done.
+ *
+ * Example (not cached):
+ *   FWebImage Image;
+ *   Image.SetStandInBrush(FStyle::Get()->GetBrush("Foo"));
+ *   Image.BeginDownload(Url, OptionalCallback);
+ *
+ *   SNew(SImage)
+ *   .Image(Image.Attr())
+ *
+ * You may want to get your FWebImage from a FWebImageCache so we're not re-downloading the same URL all the time.
+ *
+ * Example (cached):
+ *   FWebImageCache ImageCache; // usually global
+ *   ImageCache.SetStandInBrush(FStyle::Get()->GetBrush("Foo"));
+ *
+ *   SNew(SImage)
+ *   .Image(ImageCache.Download(Url).Attr())
+ */
+class IMAGEDOWNLOAD_API FWebImage 
+	: public TSharedFromThis<FWebImage>
+{
+public:
+	FWebImage();
+
+	/**
+	 * Fired when the image finishes downloading or is canceled.
+	 * @param Success True if the downloaded image will now be returned by GetBrush()
+	 */
+	DECLARE_DELEGATE_OneParam(FOnImageDownloaded, bool);
+
+	/** Set the brush that is currently being returned (this will be overridden when any async download completes) */
+	FORCEINLINE FWebImage& SetStandInBrush(TAttribute<const FSlateBrush*> StandInBrushIn) { StandInBrush = StandInBrushIn; DownloadedBrush.Reset(); return *this; }
+
+	/** Begin downloading an image. This will automatically set the current brush to the downloaded image when it completes (if successful) */
+	bool BeginDownload(const FString& Url, const FOnImageDownloaded& DownloadCallback = FOnImageDownloaded());
+
+	/** Cancel any download in progress */
+	void CancelDownload();
+
+public:
+
+	/** Use .Attr() to pass this brush into a slate attribute */
+	TAttribute< const FSlateBrush* > Attr() const;
+
+	/** Get the current brush displayed (will automatically change when download completes) */
+	FORCEINLINE const FSlateBrush* GetBrush() const { return DownloadedBrush.IsValid() ? DownloadedBrush.Get() : StandInBrush.Get(); }
+
+	/** Only returns the downloaded brush. May be null if the download hasn't finished or was unsuccessful */
+	FORCEINLINE const FSlateBrush* GetDownloadedBrush() const { return DownloadedBrush.IsValid() ? DownloadedBrush.Get() : nullptr; }
+
+	/** Is there a pending HTTP request */
+	FORCEINLINE bool IsDownloadPending() const { return PendingRequest.IsValid(); }
+
+	/** Has the download finished AND was it successful */
+	FORCEINLINE bool DidDownloadSucceed() const { return bDownloadSuccess; }
+
+	/** Has the download finished AND did it fail */
+	FORCEINLINE bool DidDownloadFail() const { return !IsDownloadPending() && !DidDownloadSucceed(); }
+
+	/** What URL was requested */
+	FORCEINLINE const FString& GetUrl() const { return Url; }
+
+private:
+	/** request complete callback */
+	void HttpRequestComplete(FHttpRequestPtr HttpRequest, FHttpResponsePtr HttpResponse, bool bSucceeded, FOnImageDownloaded DownloadCb);
+	bool ProcessHttpResponse(const FString& RequestUrl, FHttpResponsePtr HttpResponse);
+
+private:
+	/** The Url being downloaded */
+	FString Url;
+
+	/** The image resource to show */
+	TAttribute< const FSlateBrush* > StandInBrush;
+
+	/** The most recently downloaded and generated brush */
+	TSharedPtr<FSlateDynamicImageBrush> DownloadedBrush;
+
+	/** Any pending request */
+	TSharedPtr<IHttpRequest> PendingRequest;
+
+	/** Have we successfully downloaded the URL we asked for */
+	bool bDownloadSuccess;
+
+	/** When did the download complete */
+	FDateTime DownloadTimeUtc;
+};
\ No newline at end of file
diff --git a/Engine/Source/Runtime/Online/OnlineSubsystem/Private/OnlineSubsystemImpl.cpp b/Engine/Source/Runtime/Online/OnlineSubsystem/Private/OnlineSubsystemImpl.cpp
--- a/Engine/Source/Runtime/Online/OnlineSubsystem/Private/OnlineSubsystemImpl.cpp
+++ b/Engine/Source/Runtime/Online/OnlineSubsystem/Private/OnlineSubsystemImpl.cpp
@@ -30,6 +30,32 @@ FOnlineSubsystemImpl::~FOnlineSubsystemImpl()
 	}
 }
 
+void FOnlineSubsystemImpl::ExecuteDelegateNextTick(const FNextTickDelegate& Callback)
+{
+	NextTickQueue.Enqueue(Callback);
+}
+
+bool FOnlineSubsystemImpl::Tick(float DeltaTime)
+{
+	if (!NextTickQueue.IsEmpty())
+	{
+		// unload the next-tick queue into our buffer. Any further executes (from within callbacks) will happen NEXT frame (as intended)
+		FNextTickDelegate Temp;
+		while (NextTickQueue.Dequeue(Temp))
+		{
+			CurrentTickBuffer.Add(Temp);
+		}
+
+		// execute any functions in the current tick array
+		for (const auto& Callback : CurrentTickBuffer)
+		{
+			Callback.ExecuteIfBound();
+		}
+		CurrentTickBuffer.SetNum(0, false); // keep the memory around
+	}
+	return true;
+}
+
 void FOnlineSubsystemImpl::InitNamedInterfaces()
 {
 	NamedInterfaces = NewObject<UNamedInterfaces>();
diff --git a/Engine/Source/Runtime/Online/OnlineSubsystem/Public/OnlineSubsystem.h b/Engine/Source/Runtime/Online/OnlineSubsystem/Public/OnlineSubsystem.h
--- a/Engine/Source/Runtime/Online/OnlineSubsystem/Public/OnlineSubsystem.h
+++ b/Engine/Source/Runtime/Online/OnlineSubsystem/Public/OnlineSubsystem.h
@@ -33,6 +33,7 @@ DECLARE_CYCLE_STAT_EXTERN(TEXT("VoiceInt"), STAT_Voice_Interface, STATGROUP_Onli
 /** Forward declarations of all interface classes */
 typedef TSharedPtr<class IOnlineSession, ESPMode::ThreadSafe> IOnlineSessionPtr;
 typedef TSharedPtr<class IOnlineFriends, ESPMode::ThreadSafe> IOnlineFriendsPtr;
+typedef TSharedPtr<class IOnlineGroups, ESPMode::ThreadSafe> IOnlineGroupsPtr;
 typedef TSharedPtr<class IOnlineSharedCloud, ESPMode::ThreadSafe> IOnlineSharedCloudPtr;
 typedef TSharedPtr<class IOnlineUserCloud, ESPMode::ThreadSafe> IOnlineUserCloudPtr;
 typedef TSharedPtr<class IOnlineEntitlements, ESPMode::ThreadSafe> IOnlineEntitlementsPtr;
@@ -143,6 +144,12 @@ public:
 	 */
 	virtual IOnlineFriendsPtr GetFriendsInterface() const = 0;
 
+	/**
+	 * Get the interface for accessing the groups services
+	 * @return Interface pointer for appropriate groups service
+	 */
+	virtual IOnlineGroupsPtr GetGroupsInterface() const = 0;
+
 	/** 
 	 * Get the interface for sharing user files in the cloud
 	 * @return Interface pointer for the appropriate cloud service
diff --git a/Engine/Source/Runtime/Online/OnlineSubsystem/Public/OnlineSubsystemImpl.h b/Engine/Source/Runtime/Online/OnlineSubsystem/Public/OnlineSubsystemImpl.h
--- a/Engine/Source/Runtime/Online/OnlineSubsystem/Public/OnlineSubsystemImpl.h
+++ b/Engine/Source/Runtime/Online/OnlineSubsystem/Public/OnlineSubsystemImpl.h
@@ -5,10 +5,14 @@
 #include "OnlineSubsystem.h"
 #include "OnlineSubsystemPackage.h"
 
+DECLARE_DELEGATE(FNextTickDelegate);
+
 /**
  *	FOnlineSubsystemImpl - common functionality to share across online platforms, not intended for direct use
  */
-class ONLINESUBSYSTEM_API FOnlineSubsystemImpl : public IOnlineSubsystem
+class ONLINESUBSYSTEM_API FOnlineSubsystemImpl 
+	: public IOnlineSubsystem
+	, public FTickerObjectBase
 {
 protected:
 
@@ -28,6 +32,12 @@ protected:
 	/** Load in any named interfaces specified by the ini configuration */
 	void InitNamedInterfaces();
 
+	/** Queue to hold callbacks scheduled for next tick using ExecuteNextTick */
+	TQueue<FNextTickDelegate, EQueueMode::Mpsc> NextTickQueue;
+
+	/** Buffer to hold callbacks for the current tick (so it's safe to call ExecuteNextTick within a tick callback) */
+	TArray<FNextTickDelegate> CurrentTickBuffer;
+
 public:
 	
 	virtual ~FOnlineSubsystemImpl();
@@ -44,6 +54,10 @@ public:
 
 	virtual void SetUsingMultiplayerFeatures(const FUniqueNetId& UniqueId, bool bUsingMP) override {};
 
+	// FTickerObjectBase
+
+	virtual bool Tick(float DeltaTime) override;
+
 	// FOnlineSubsystemImpl
 
 	/**
@@ -51,6 +65,20 @@ public:
 	 */
 	FName GetInstanceName() const { return InstanceName; }
 
+	/**
+	 * Queue a delegate to be executed on the next tick
+	 */
+	void ExecuteDelegateNextTick(const FNextTickDelegate& Callback);
+
+	/**
+	 * Templated helper for calling ExecuteDelegateNextTick with a lambda function
+	 */
+	template<typename LAMBDA_TYPE>
+	FORCEINLINE void ExecuteNextTick(LAMBDA_TYPE&& Callback)
+	{
+		ExecuteDelegateNextTick(FNextTickDelegate::CreateLambda(Callback));
+	}
+
 	/** Name given to default OSS instances (disambiguates for PIE) */
 	static const FName DefaultInstanceName;
 };
diff --git a/Engine/Source/Runtime/Online/OnlineSubsystemAmazon/Private/OnlineSubsystemAmazon.cpp b/Engine/Source/Runtime/Online/OnlineSubsystemAmazon/Private/OnlineSubsystemAmazon.cpp
--- a/Engine/Source/Runtime/Online/OnlineSubsystemAmazon/Private/OnlineSubsystemAmazon.cpp
+++ b/Engine/Source/Runtime/Online/OnlineSubsystemAmazon/Private/OnlineSubsystemAmazon.cpp
@@ -71,6 +71,11 @@ IOnlineIdentityPtr FOnlineSubsystemAmazon::GetIdentityInterface() const
 
 bool FOnlineSubsystemAmazon::Tick(float DeltaTime)
 {
+	if (!FOnlineSubsystemImpl::Tick(DeltaTime))
+	{
+		return false;
+	}
+
 	if (IdentityInterface.IsValid())
 	{
 		TickToggle ^= 1;
diff --git a/Engine/Source/Runtime/Online/OnlineSubsystemAmazon/Public/OnlineSubsystemAmazon.h b/Engine/Source/Runtime/Online/OnlineSubsystemAmazon/Public/OnlineSubsystemAmazon.h
--- a/Engine/Source/Runtime/Online/OnlineSubsystemAmazon/Public/OnlineSubsystemAmazon.h
+++ b/Engine/Source/Runtime/Online/OnlineSubsystemAmazon/Public/OnlineSubsystemAmazon.h
@@ -14,8 +14,7 @@ typedef TSharedPtr<class FOnlineIdentityAmazon, ESPMode::ThreadSafe> FOnlineIden
  * Amazon subsystem
  */
 class ONLINESUBSYSTEMAMAZON_API FOnlineSubsystemAmazon :
-	public FOnlineSubsystemImpl,
-	public FTickerObjectBase
+	public FOnlineSubsystemImpl
 
 {
 	class FOnlineFactoryAmazon* AmazonFactory;
@@ -42,6 +41,10 @@ public:
 	{
 		return NULL;
 	}
+	virtual IOnlineGroupsPtr GetGroupsInterface() const override
+	{
+		return nullptr;
+	}
 	virtual IOnlineSharedCloudPtr GetSharedCloudInterface() const override
 	{
 		return NULL;
diff --git a/Engine/Source/Runtime/Online/OnlineSubsystemFacebook/Private/IOS/OnlineSubsystemFacebook.cpp b/Engine/Source/Runtime/Online/OnlineSubsystemFacebook/Private/IOS/OnlineSubsystemFacebook.cpp
--- a/Engine/Source/Runtime/Online/OnlineSubsystemFacebook/Private/IOS/OnlineSubsystemFacebook.cpp
+++ b/Engine/Source/Runtime/Online/OnlineSubsystemFacebook/Private/IOS/OnlineSubsystemFacebook.cpp
@@ -38,6 +38,11 @@ IOnlineFriendsPtr FOnlineSubsystemFacebook::GetFriendsInterface() const
 	return FacebookFriends;
 }
 
+IOnlineGroupsPtr FOnlineSubsystemFacebook::GetGroupsInterface() const
+{
+	return nullptr;
+}
+
 IOnlineSharedCloudPtr FOnlineSubsystemFacebook::GetSharedCloudInterface() const
 {
 	return NULL;
@@ -155,6 +160,11 @@ bool FOnlineSubsystemFacebook::Exec(UWorld* InWorld, const TCHAR* Cmd, FOutputDe
 
 bool FOnlineSubsystemFacebook::Tick(float DeltaTime)
 {
+	if (!FOnlineSubsystemImpl::Tick(DeltaTime))
+	{
+		return false;
+	}
+
 	return true;
 }
 
diff --git a/Engine/Source/Runtime/Online/OnlineSubsystemFacebook/Private/Windows/OnlineSubsystemFacebook.cpp b/Engine/Source/Runtime/Online/OnlineSubsystemFacebook/Private/Windows/OnlineSubsystemFacebook.cpp
--- a/Engine/Source/Runtime/Online/OnlineSubsystemFacebook/Private/Windows/OnlineSubsystemFacebook.cpp
+++ b/Engine/Source/Runtime/Online/OnlineSubsystemFacebook/Private/Windows/OnlineSubsystemFacebook.cpp
@@ -18,6 +18,11 @@ IOnlineFriendsPtr FOnlineSubsystemFacebook::GetFriendsInterface() const
 	return FacebookFriends;
 }
 
+IOnlineGroupsPtr FOnlineSubsystemFacebook::GetGroupsInterface() const
+{
+	return nullptr;
+}
+
 
 IOnlineSharedCloudPtr FOnlineSubsystemFacebook::GetSharedCloudInterface() const
 {
@@ -121,6 +126,11 @@ IOnlineTurnBasedPtr FOnlineSubsystemFacebook::GetTurnBasedInterface() const
 
 bool FOnlineSubsystemFacebook::Tick(float DeltaTime)
 {
+	if (!FOnlineSubsystemImpl::Tick(DeltaTime))
+	{
+		return false;
+	}
+
 	if (FacebookIdentity.IsValid())
 	{		
 		FacebookIdentity->Tick(DeltaTime);
diff --git a/Engine/Source/Runtime/Online/OnlineSubsystemFacebook/Public/OnlineSubsystemFacebook.h b/Engine/Source/Runtime/Online/OnlineSubsystemFacebook/Public/OnlineSubsystemFacebook.h
--- a/Engine/Source/Runtime/Online/OnlineSubsystemFacebook/Public/OnlineSubsystemFacebook.h
+++ b/Engine/Source/Runtime/Online/OnlineSubsystemFacebook/Public/OnlineSubsystemFacebook.h
@@ -17,13 +17,13 @@ typedef TSharedPtr<class FOnlineUserFacebook, ESPMode::ThreadSafe> FOnlineUserFa
  */
 class ONLINESUBSYSTEMFACEBOOK_API FOnlineSubsystemFacebook 
 	: public FOnlineSubsystemImpl
-	, public FTickerObjectBase
 {
 public:
 
 	// IOnlineSubsystem Interface
 	virtual IOnlineSessionPtr GetSessionInterface() const override;
 	virtual IOnlineFriendsPtr GetFriendsInterface() const override;
+	virtual IOnlineGroupsPtr GetGroupsInterface() const override;
 	virtual IOnlineSharedCloudPtr GetSharedCloudInterface() const override;
 	virtual IOnlineUserCloudPtr GetUserCloudInterface() const override;
 	virtual IOnlineLeaderboardsPtr GetLeaderboardsInterface() const override;
diff --git a/Engine/Source/Runtime/Online/OnlineSubsystemNull/Private/OnlineSubsystemNull.cpp b/Engine/Source/Runtime/Online/OnlineSubsystemNull/Private/OnlineSubsystemNull.cpp
--- a/Engine/Source/Runtime/Online/OnlineSubsystemNull/Private/OnlineSubsystemNull.cpp
+++ b/Engine/Source/Runtime/Online/OnlineSubsystemNull/Private/OnlineSubsystemNull.cpp
@@ -21,6 +21,11 @@ IOnlineFriendsPtr FOnlineSubsystemNull::GetFriendsInterface() const
 	return NULL;
 }
 
+IOnlineGroupsPtr FOnlineSubsystemNull::GetGroupsInterface() const
+{
+	return nullptr;
+}
+
 IOnlineSharedCloudPtr FOnlineSubsystemNull::GetSharedCloudInterface() const
 {
 	return NULL;
@@ -115,6 +120,11 @@ IOnlineTurnBasedPtr FOnlineSubsystemNull::GetTurnBasedInterface() const
 
 bool FOnlineSubsystemNull::Tick(float DeltaTime)
 {
+	if (!FOnlineSubsystemImpl::Tick(DeltaTime))
+	{
+		return false;
+	}
+
 	if (OnlineAsyncTaskThreadRunnable)
 	{
 		OnlineAsyncTaskThreadRunnable->GameTick();
diff --git a/Engine/Source/Runtime/Online/OnlineSubsystemSteam/Private/OnlineSubsystemSteam.cpp b/Engine/Source/Runtime/Online/OnlineSubsystemSteam/Private/OnlineSubsystemSteam.cpp
--- a/Engine/Source/Runtime/Online/OnlineSubsystemSteam/Private/OnlineSubsystemSteam.cpp
+++ b/Engine/Source/Runtime/Online/OnlineSubsystemSteam/Private/OnlineSubsystemSteam.cpp
@@ -164,6 +164,11 @@ IOnlineFriendsPtr FOnlineSubsystemSteam::GetFriendsInterface() const
 	return FriendInterface;
 }
 
+IOnlineGroupsPtr FOnlineSubsystemSteam::GetGroupsInterface() const
+{
+	return nullptr;
+}
+
 IOnlineSharedCloudPtr FOnlineSubsystemSteam::GetSharedCloudInterface() const
 {
 	return SharedCloudInterface;
@@ -268,6 +273,11 @@ void FOnlineSubsystemSteam::QueueAsyncOutgoingItem(FOnlineAsyncItem* AsyncItem)
 
 bool FOnlineSubsystemSteam::Tick(float DeltaTime)
 {
+	if (!FOnlineSubsystemImpl::Tick(DeltaTime))
+	{
+		return false;
+	}
+
 	if (OnlineAsyncTaskThreadRunnable)
 	{
 		OnlineAsyncTaskThreadRunnable->GameTick();
diff --git a/Engine/Source/Runtime/Online/OnlineSubsystemUtils/Classes/OnlineBeacon.h b/Engine/Source/Runtime/Online/OnlineSubsystemUtils/Classes/OnlineBeacon.h
--- a/Engine/Source/Runtime/Online/OnlineSubsystemUtils/Classes/OnlineBeacon.h
+++ b/Engine/Source/Runtime/Online/OnlineSubsystemUtils/Classes/OnlineBeacon.h
@@ -6,11 +6,6 @@
 #pragma once
 #include "OnlineBeacon.generated.h"
 
-/** Time client beacon will wait to establish a connection with the beacon host */
-#define BEACON_CONNECTION_INITIAL_TIMEOUT 5.0f
-/** Time client beacon will wait after establishing a connection for packets before giving up */
-#define BEACON_CONNECTION_TIMEOUT 45.0f
-
 ONLINESUBSYSTEMUTILS_API DECLARE_LOG_CATEGORY_EXTERN(LogBeacon, Display, All);
 
 /** States that a beacon can be in */
@@ -68,7 +63,7 @@ class ONLINESUBSYSTEMUTILS_API AOnlineBeacon : public AActor, public FNetworkNot
 		if (bPause)
 		{
 			UE_LOG(LogBeacon, Verbose, TEXT("Reservation Beacon Requests Paused."));
-			NetDriver->SetWorld(NULL);
+			NetDriver->SetWorld(nullptr);
 			NetDriver->Notify = this;
 			BeaconState = EBeaconState::DenyRequests;
 		}
@@ -102,12 +97,18 @@ class ONLINESUBSYSTEMUTILS_API AOnlineBeacon : public AActor, public FNetworkNot
 	virtual UNetConnection* GetNetConnection() override;	
 
 protected:
+
+	/** Time beacon will wait to establish a connection with the beacon host */
+	UPROPERTY(Config)
+	float BeaconConnectionInitialTimeout;
+	/** Time beacon will wait for packets after establishing a connection before giving up */
+	UPROPERTY(Config)
+	float BeaconConnectionTimeout;
+
 	/** Net driver routing network traffic */
 	UPROPERTY()
 	UNetDriver* NetDriver;
-	/** Name of net driver to use for communication */
-	UPROPERTY(Config)
-	FName BeaconNetDriverName;
+
 	/** Network connection */
 	UPROPERTY()
 	UNetConnection* BeaconConnection;
@@ -117,7 +118,7 @@ protected:
 	FDelegateHandle HandleNetworkFailureDelegateHandle;
 
 	/** Common initialization for all beacon types */
-	bool InitBase();
+	virtual bool InitBase();
 
 	/** Notification that failure needs to be handled */
 	virtual void OnFailure();
diff --git a/Engine/Source/Runtime/Online/OnlineSubsystemUtils/Classes/PartyBeaconClient.h b/Engine/Source/Runtime/Online/OnlineSubsystemUtils/Classes/PartyBeaconClient.h
--- a/Engine/Source/Runtime/Online/OnlineSubsystemUtils/Classes/PartyBeaconClient.h
+++ b/Engine/Source/Runtime/Online/OnlineSubsystemUtils/Classes/PartyBeaconClient.h
@@ -2,11 +2,13 @@
 
 #pragma once
 
-#include "Runtime/Online/OnlineSubsystemUtils/Classes/OnlineBeaconClient.h"
-#include "PartyBeaconHost.h"
+#include "OnlineBeaconClient.h"
+#include "PartyBeaconState.h"
 
 #include "PartyBeaconClient.generated.h"
 
+#define PARTY_BEACON_TYPE TEXT("PartyBeacon")
+
 /**
  * Delegate triggered when a response from the party beacon host has been received
  *
@@ -14,11 +16,6 @@
  */
 DECLARE_DELEGATE_OneParam(FOnReservationRequestComplete, EPartyReservationResult::Type);
 
-/**
- * Delegate triggered on failures to connect to the party host beacon
- */
-DECLARE_DELEGATE(FOnHostConnectionFailure);
-
 /**
  * A beacon client used for making reservations with an existing game session
  */
@@ -28,12 +25,11 @@ class ONLINESUBSYSTEMUTILS_API APartyBeaconClient : public AOnlineBeaconClient
 	GENERATED_UCLASS_BODY()
 
 	// Begin AOnlineBeacon Interface
-	virtual FString GetBeaconType() override { return TEXT("PartyBeacon"); }
+	virtual FString GetBeaconType() override { return PARTY_BEACON_TYPE; }
 	// End AOnlineBeacon Interface
 
 	// Begin AOnlineBeaconClient Interface
 	virtual void OnConnected() override;
-	virtual void OnFailure() override;
 	// End AOnlineBeaconClient Interface
 
 	/**
@@ -55,7 +51,7 @@ class ONLINESUBSYSTEMUTILS_API APartyBeaconClient : public AOnlineBeaconClient
 	virtual void CancelReservation();
 
 	/**
-	 * Response from the host session after making a reservation requestion
+	 * Response from the host session after making a reservation request
 	 *
 	 * @param ReservationResponse response from server
 	 */
@@ -65,15 +61,10 @@ class ONLINESUBSYSTEMUTILS_API APartyBeaconClient : public AOnlineBeaconClient
 	/**
 	 * Delegate triggered when a response from the party beacon host has been received
 	 *
-	 * @param ReservationResponse response from the server
+	 * @return delegate to handle response from the server
 	 */
 	FOnReservationRequestComplete& OnReservationRequestComplete() { return ReservationRequestComplete; }
 
-	/**
-	* Delegate triggered on failures to connect to the party host beacon
-	*/
-	FOnHostConnectionFailure& OnHostConnectionFailure() { return HostConnectionFailure; }
-
 	/**
 	* @return the pending reservation associated with this beacon client
 	*/
@@ -83,8 +74,6 @@ protected:
 
 	/** Delegate for reservation request responses */
 	FOnReservationRequestComplete ReservationRequestComplete;
-	/** Delegate for party host beacon connection failures */
-	FOnHostConnectionFailure HostConnectionFailure;
 
 	/** Session Id of the destination host */
 	FString DestSessionId;
