Made CanActivateAbility SourceTags and TargetTags optional and passed by reference.

diff --git a/Engine/Source/Runtime/GameplayAbilities/Private/Abilities/GameplayAbility.cpp b/Engine/Source/Runtime/GameplayAbilities/Private/Abilities/GameplayAbility.cpp
--- a/Engine/Source/Runtime/GameplayAbilities/Private/Abilities/GameplayAbility.cpp
+++ b/Engine/Source/Runtime/GameplayAbilities/Private/Abilities/GameplayAbility.cpp
@@ -146,7 +146,7 @@ bool UGameplayAbility::IsSupportedForNetworking() const
 	return Supported;
 }
 
-bool UGameplayAbility::DoesAbilitySatisfyTagRequirements(const UAbilitySystemComponent& AbilitySystemComponent, FGameplayTagContainer SourceTags, FGameplayTagContainer TargetTags, OUT FGameplayTagContainer* OptionalRelevantTags) const
+bool UGameplayAbility::DoesAbilitySatisfyTagRequirements(const UAbilitySystemComponent& AbilitySystemComponent, const FGameplayTagContainer* SourceTags, const FGameplayTagContainer* TargetTags, OUT FGameplayTagContainer* OptionalRelevantTags) const
 {
 	bool bBlocked = false, bMissing = false;
 
@@ -177,29 +177,35 @@ bool UGameplayAbility::DoesAbilitySatisfyTagRequirements(const UAbilitySystemCom
 		}
 	}
 
-	if (SourceBlockedTags.Num() || SourceRequiredTags.Num())
+	if (SourceTags != nullptr)
 	{
-		if (SourceTags.MatchesAny(SourceBlockedTags, false))
+		if (SourceBlockedTags.Num() || SourceRequiredTags.Num())
 		{
-			bBlocked = true;
-		}
+			if (SourceTags->MatchesAny(SourceBlockedTags, false))
+			{
+				bBlocked = true;
+			}
 
-		if (!SourceTags.MatchesAll(SourceRequiredTags, true))
-		{
-			bMissing = true;
+			if (!SourceTags->MatchesAll(SourceRequiredTags, true))
+			{
+				bMissing = true;
+			}
 		}
 	}
 
-	if (TargetBlockedTags.Num() || TargetRequiredTags.Num())
+	if (TargetTags != nullptr)
 	{
-		if (TargetTags.MatchesAny(TargetBlockedTags, false))
+		if (TargetBlockedTags.Num() || TargetRequiredTags.Num())
 		{
-			bBlocked = true;
-		}
+			if (TargetTags->MatchesAny(TargetBlockedTags, false))
+			{
+				bBlocked = true;
+			}
 
-		if (!TargetTags.MatchesAll(TargetRequiredTags, true))
-		{
-			bMissing = true;
+			if (!TargetTags->MatchesAll(TargetRequiredTags, true))
+			{
+				bMissing = true;
+			}
 		}
 	}
 
@@ -223,7 +229,7 @@ bool UGameplayAbility::DoesAbilitySatisfyTagRequirements(const UAbilitySystemCom
 	return true;
 }
 
-bool UGameplayAbility::CanActivateAbility(const FGameplayAbilitySpecHandle Handle, const FGameplayAbilityActorInfo* ActorInfo, FGameplayTagContainer SourceTags, FGameplayTagContainer TargetTags, OUT FGameplayTagContainer* OptionalRelevantTags) const
+bool UGameplayAbility::CanActivateAbility(const FGameplayAbilitySpecHandle Handle, const FGameplayAbilityActorInfo* ActorInfo, const FGameplayTagContainer* SourceTags, const FGameplayTagContainer* TargetTags, OUT FGameplayTagContainer* OptionalRelevantTags) const
 {
 	// Don't set the actor info, CanActivate is called on the CDO
 
diff --git a/Engine/Source/Runtime/GameplayAbilities/Private/Abilities/GameplayAbility_CharacterJump.cpp b/Engine/Source/Runtime/GameplayAbilities/Private/Abilities/GameplayAbility_CharacterJump.cpp
--- a/Engine/Source/Runtime/GameplayAbilities/Private/Abilities/GameplayAbility_CharacterJump.cpp
+++ b/Engine/Source/Runtime/GameplayAbilities/Private/Abilities/GameplayAbility_CharacterJump.cpp
@@ -39,7 +39,7 @@ void UGameplayAbility_CharacterJump::InputReleased(const FGameplayAbilitySpecHan
 	}
 }
 
-bool UGameplayAbility_CharacterJump::CanActivateAbility(const FGameplayAbilitySpecHandle Handle, const FGameplayAbilityActorInfo* ActorInfo, FGameplayTagContainer SourceTags, FGameplayTagContainer TargetTags, OUT FGameplayTagContainer* OptionalRelevantTags) const
+bool UGameplayAbility_CharacterJump::CanActivateAbility(const FGameplayAbilitySpecHandle Handle, const FGameplayAbilityActorInfo* ActorInfo, const FGameplayTagContainer* SourceTags, const FGameplayTagContainer* TargetTags, OUT FGameplayTagContainer* OptionalRelevantTags) const
 {
 	if (!Super::CanActivateAbility(Handle, ActorInfo, SourceTags, TargetTags, OptionalRelevantTags))
 	{
diff --git a/Engine/Source/Runtime/GameplayAbilities/Private/AbilitySystemComponent.cpp b/Engine/Source/Runtime/GameplayAbilities/Private/AbilitySystemComponent.cpp
--- a/Engine/Source/Runtime/GameplayAbilities/Private/AbilitySystemComponent.cpp
+++ b/Engine/Source/Runtime/GameplayAbilities/Private/AbilitySystemComponent.cpp
@@ -1361,7 +1361,7 @@ void UAbilitySystemComponent::DisplayDebug(class UCanvas* Canvas, const class FD
 				StatusText = TEXT(" (TagBlocked)");
 				AbilityTextColor = FColor::Red;
 			}
-			else if (AbilitySpec.Ability->CanActivateAbility(AbilitySpec.Handle, AbilityActorInfo.Get(), FGameplayTagContainer(), FGameplayTagContainer()) == false)
+			else if (AbilitySpec.Ability->CanActivateAbility(AbilitySpec.Handle, AbilityActorInfo.Get()) == false)
 			{
 				StatusText = TEXT(" (CantActivate)");
 				AbilityTextColor = FColor::Red;
diff --git a/Engine/Source/Runtime/GameplayAbilities/Private/AbilitySystemComponent_Abilities.cpp b/Engine/Source/Runtime/GameplayAbilities/Private/AbilitySystemComponent_Abilities.cpp
--- a/Engine/Source/Runtime/GameplayAbilities/Private/AbilitySystemComponent_Abilities.cpp
+++ b/Engine/Source/Runtime/GameplayAbilities/Private/AbilitySystemComponent_Abilities.cpp
@@ -739,7 +739,7 @@ void UAbilitySystemComponent::GetActivatableGameplayAbilitySpecsByAllMatchingTag
 		{
 			// Do NOT consider abilities that are blocked by tags currently.  That way, we can use the blocking to find
 			// an appropriate ability based on tags when we have more than one ability that match the GameplayTagContainer.
-			if (Spec.Ability->DoesAbilitySatisfyTagRequirements(*this, FGameplayTagContainer(), FGameplayTagContainer()))
+			if (Spec.Ability->DoesAbilitySatisfyTagRequirements(*this))
 			{
 				MatchingGameplayAbilities.Add(const_cast<FGameplayAbilitySpec*>(&Spec));
 			}
@@ -836,12 +836,12 @@ bool UAbilitySystemComponent::TryActivateAbility(FGameplayAbilitySpecHandle Hand
 	// TODO: We should run this on on the non instanced, but run it on the instance once for now, need to fixup data
 	// (Always do a non instanced CanActivate check)
 
-	FGameplayTagContainer SourceTags;
-	FGameplayTagContainer TargetTags;
+	const FGameplayTagContainer* SourceTags = nullptr;
+	const FGameplayTagContainer* TargetTags = nullptr;
 	if (TriggerEventData != nullptr)
 	{
-		SourceTags = TriggerEventData->InstigatorTags;
-		TargetTags = TriggerEventData->TargetTags;
+		SourceTags = &TriggerEventData->InstigatorTags;
+		TargetTags = &TriggerEventData->TargetTags;
 	}
 
 	if (InstancedAbility)
diff --git a/Engine/Source/Runtime/GameplayAbilities/Public/Abilities/GameplayAbility.h b/Engine/Source/Runtime/GameplayAbilities/Public/Abilities/GameplayAbility.h
--- a/Engine/Source/Runtime/GameplayAbilities/Public/Abilities/GameplayAbility.h
+++ b/Engine/Source/Runtime/GameplayAbilities/Public/Abilities/GameplayAbility.h
@@ -123,7 +123,7 @@ public:
 		
 
 	/** Returns true if this ability can be activated right now. Has no side effects */
-	virtual bool CanActivateAbility(const FGameplayAbilitySpecHandle Handle, const FGameplayAbilityActorInfo* ActorInfo, FGameplayTagContainer SourceTags, FGameplayTagContainer TargetTags, OUT FGameplayTagContainer* OptionalRelevantTags = nullptr) const;
+	virtual bool CanActivateAbility(const FGameplayAbilitySpecHandle Handle, const FGameplayAbilityActorInfo* ActorInfo, const FGameplayTagContainer* SourceTags = nullptr, const FGameplayTagContainer* TargetTags = nullptr, OUT FGameplayTagContainer* OptionalRelevantTags = nullptr) const;
 
 	/** Returns true if this ability can be triggered right now. Has no side effects */
 	virtual bool ShouldAbilityRespondToEvent(const FGameplayAbilityActorInfo* ActorInfo, const FGameplayEventData* Payload) const;
@@ -137,7 +137,7 @@ public:
 	virtual const FGameplayTagContainer* GetCooldownTags() const;
 	
 	/** Returns true if none of the ability's tags are blocked and if it doesn't have a "Blocking" tag and has all "Required" tags. */
-	bool DoesAbilitySatisfyTagRequirements(const UAbilitySystemComponent& AbilitySystemComponent, FGameplayTagContainer SourceTags, FGameplayTagContainer TargetTags, OUT FGameplayTagContainer* OptionalRelevantTags = nullptr) const;
+	bool DoesAbilitySatisfyTagRequirements(const UAbilitySystemComponent& AbilitySystemComponent, const FGameplayTagContainer* SourceTags = nullptr, const FGameplayTagContainer* TargetTags = nullptr, OUT FGameplayTagContainer* OptionalRelevantTags = nullptr) const;
 
 	EGameplayAbilityInstancingPolicy::Type GetInstancingPolicy() const
 	{
diff --git a/Engine/Source/Runtime/GameplayAbilities/Public/Abilities/GameplayAbility_CharacterJump.h b/Engine/Source/Runtime/GameplayAbilities/Public/Abilities/GameplayAbility_CharacterJump.h
--- a/Engine/Source/Runtime/GameplayAbilities/Public/Abilities/GameplayAbility_CharacterJump.h
+++ b/Engine/Source/Runtime/GameplayAbilities/Public/Abilities/GameplayAbility_CharacterJump.h
@@ -15,8 +15,8 @@ class GAMEPLAYABILITIES_API UGameplayAbility_CharacterJump : public UGameplayAbi
 
 public:
 
-	virtual bool CanActivateAbility(const FGameplayAbilitySpecHandle Handle, const FGameplayAbilityActorInfo* ActorInfo, FGameplayTagContainer SourceTags, FGameplayTagContainer TargetTags, OUT FGameplayTagContainer* OptionalRelevantTags = nullptr) const override;
-	
+	virtual bool CanActivateAbility(const FGameplayAbilitySpecHandle Handle, const FGameplayAbilityActorInfo* ActorInfo, const FGameplayTagContainer* SourceTags = nullptr, const FGameplayTagContainer* TargetTags = nullptr, OUT FGameplayTagContainer* OptionalRelevantTags = nullptr) const override;
+
 	virtual void ActivateAbility(const FGameplayAbilitySpecHandle Handle, const FGameplayAbilityActorInfo* OwnerInfo, const FGameplayAbilityActivationInfo ActivationInfo, const FGameplayEventData* TriggerEventData) override;
 
 	virtual void InputReleased(const FGameplayAbilitySpecHandle Handle, const FGameplayAbilityActorInfo* ActorInfo, const FGameplayAbilityActivationInfo ActivationInfo) override;
