Ability system merge from dev branch

diff --git a/Engine/Source/Runtime/Engine/Classes/Engine/NetSerialization.h b/Engine/Source/Runtime/Engine/Classes/Engine/NetSerialization.h
--- a/Engine/Source/Runtime/Engine/Classes/Engine/NetSerialization.h
+++ b/Engine/Source/Runtime/Engine/Classes/Engine/NetSerialization.h
@@ -397,9 +397,17 @@ bool FFastArraySerializer::FastArrayDeltaSerialize( TArray<Type> &Items, FNetDel
 			ArraySerializer.ItemMap.Empty();
 			for (int32 i=0; i < Items.Num(); ++i)
 			{
-				if ( Items[i].ReplicationID == INDEX_NONE )
+				if ( Items[i].ReplicationID == INDEX_NONE  )
 				{
-					UE_LOG( LogNetSerialization, Warning, TEXT( "FastArrayDeltaSerialize: Item with uninitialized ReplicationID. Struct: %s, ItemIndex: %i" ), *InnerStruct->GetOwnerStruct()->GetName(), i );
+					if (Parms.Writer)
+					{
+						UE_LOG( LogNetSerialization, Warning, TEXT( "FastArrayDeltaSerialize: Item with uninitialized ReplicationID. Struct: %s, ItemIndex: %i" ), *InnerStruct->GetOwnerStruct()->GetName(), i );
+					}
+					else
+					{
+						// This is benign for clients, they may add things to their local array without assigning a ReplicationID
+						continue;
+					}
 				}
 				ArraySerializer.ItemMap.Add( Items[i].ReplicationID, i );
 			}
diff --git a/Engine/Source/Runtime/GameplayAbilities/Private/Abilities/GameplayAbility.cpp b/Engine/Source/Runtime/GameplayAbilities/Private/Abilities/GameplayAbility.cpp
--- a/Engine/Source/Runtime/GameplayAbilities/Private/Abilities/GameplayAbility.cpp
+++ b/Engine/Source/Runtime/GameplayAbilities/Private/Abilities/GameplayAbility.cpp
@@ -450,23 +450,36 @@ void UGameplayAbility::SetShouldBlockOtherAbilities(bool bShouldBlockAbilities)
 	}
 }
 
-void UGameplayAbility::CancelAbility(const FGameplayAbilitySpecHandle Handle, const FGameplayAbilityActorInfo* ActorInfo, const FGameplayAbilityActivationInfo ActivationInfo)
+void UGameplayAbility::CancelAbility(const FGameplayAbilitySpecHandle Handle, const FGameplayAbilityActorInfo* ActorInfo, const FGameplayAbilityActivationInfo ActivationInfo, bool bReplicateCancelAbility)
 {
 	if (CanBeCanceled())
 	{
+		// Replicate the the server/client if needed
+		if (bReplicateCancelAbility)
+		{
+			ActorInfo->AbilitySystemComponent->ReplicateEndOrCancelAbility(Handle, ActivationInfo, this, true);
+		}
+
 		// Gives the Ability BP a chance to perform custom logic/cleanup when any active ability states are active
-		UAbilitySystemComponent* Comp = ActorInfo->AbilitySystemComponent.Get();
-		if (Comp && Comp->OnAbilityStateInterrupted.IsBound())
+		if (OnGameplayAbilityCancelled.IsBound())
 		{
-			Comp->OnAbilityStateInterrupted.Broadcast();
+			OnGameplayAbilityCancelled.Broadcast();
 		}
 
-		EndAbility(Handle, ActorInfo, ActivationInfo, true);
+		// End the ability but don't replicate it, we replicate the CancelAbility call directly
+		EndAbility(Handle, ActorInfo, ActivationInfo, false);
 	}
 }
 
 void UGameplayAbility::EndAbility(const FGameplayAbilitySpecHandle Handle, const FGameplayAbilityActorInfo* ActorInfo, const FGameplayAbilityActivationInfo ActivationInfo, bool bReplicateEndAbility)
-{ 
+{
+	// Protect against EndAbility being called multiple times
+	// Ending an AbilityState may cause this to be invoked again
+	if (bIsActive == false && GetInstancingPolicy() != EGameplayAbilityInstancingPolicy::NonInstanced)
+	{
+		return;
+	}
+
 	// check to see if this is an NonInstanced or if the ability is active.
 	FGameplayAbilitySpec* Spec = ActorInfo->AbilitySystemComponent->FindAbilitySpecFromHandle(Handle);
 	if ((Spec != nullptr) ? Spec->IsActive() : IsActive())
@@ -486,8 +499,13 @@ void UGameplayAbility::EndAbility(const FGameplayAbilitySpecHandle Handle, const
 		OnGameplayAbilityEnded.ExecuteIfBound(this);
 		OnGameplayAbilityEnded.Unbind();
 
+		if (GetInstancingPolicy() != EGameplayAbilityInstancingPolicy::NonInstanced)
+		{
+			bIsActive = false;
+		}
+
 		// Tell all our tasks that we are finished and they should cleanup
-		for (int32 TaskIdx = ActiveTasks.Num() - 1; TaskIdx >= 0; --TaskIdx)
+		for (int32 TaskIdx = ActiveTasks.Num() - 1; TaskIdx >= 0 && ActiveTasks.Num() > 0; --TaskIdx)
 		{
 			TWeakObjectPtr<UAbilityTask> Task = ActiveTasks[TaskIdx];
 			if (Task.IsValid())
@@ -497,21 +515,23 @@ void UGameplayAbility::EndAbility(const FGameplayAbilitySpecHandle Handle, const
 		}
 		ActiveTasks.Reset();	// Empty the array but dont resize memory, since this object is probably going to be destroyed very soon anyways.
 
-		if (GetInstancingPolicy() != EGameplayAbilityInstancingPolicy::NonInstanced)
-		{
-			bIsActive = false;
-		}
-
 		if (ActorInfo && ActorInfo->AbilitySystemComponent.IsValid())
 		{
 			if (bReplicateEndAbility)
 			{
-				ActorInfo->AbilitySystemComponent->ReplicateEndAbility(Handle, ActivationInfo, this);
+				ActorInfo->AbilitySystemComponent->ReplicateEndOrCancelAbility(Handle, ActivationInfo, this, false);
 			}
 
 			// Remove tags we added to owner
 			ActorInfo->AbilitySystemComponent->RemoveLooseGameplayTags(ActivationOwnedTags);
 
+			// Remove tracked GameplayCues that we added
+			for (FGameplayTag& GameplayCueTag : TrackedGameplayCues)
+			{
+				ActorInfo->AbilitySystemComponent->RemoveGameplayCue(GameplayCueTag);
+			}
+			TrackedGameplayCues.Empty();
+
 			if (CanBeCanceled())
 			{
 				// If we're still cancelable, cancel it now
@@ -1024,10 +1044,9 @@ void UGameplayAbility::EndAbilityState(FName OptionalStateNameToEnd)
 {
 	check(CurrentActorInfo);
 
-	UAbilitySystemComponent* ASC = CurrentActorInfo->AbilitySystemComponent.Get();
-	if (ASC && ASC->OnAbilityStateEnded.IsBound())
+	if (OnGameplayAbilityStateEnded.IsBound())
 	{
-		ASC->OnAbilityStateEnded.Broadcast(OptionalStateNameToEnd);
+		OnGameplayAbilityStateEnded.Broadcast(OptionalStateNameToEnd);
 	}
 }
 
@@ -1050,16 +1069,23 @@ void UGameplayAbility::K2_ExecuteGameplayCue(FGameplayTag GameplayCueTag, FGamep
 	CurrentActorInfo->AbilitySystemComponent->ExecuteGameplayCue(GameplayCueTag, Context);
 }
 
-void UGameplayAbility::K2_AddGameplayCue(FGameplayTag GameplayCueTag, FGameplayEffectContextHandle Context)
+void UGameplayAbility::K2_AddGameplayCue(FGameplayTag GameplayCueTag, FGameplayEffectContextHandle Context, bool bRemoveOnAbilityEnd)
 {
 	check(CurrentActorInfo);
 	CurrentActorInfo->AbilitySystemComponent->AddGameplayCue(GameplayCueTag, Context);
+
+	if (bRemoveOnAbilityEnd)
+	{
+		TrackedGameplayCues.Add(GameplayCueTag);
+	}
 }
 
 void UGameplayAbility::K2_RemoveGameplayCue(FGameplayTag GameplayCueTag)
 {
 	check(CurrentActorInfo);
 	CurrentActorInfo->AbilitySystemComponent->RemoveGameplayCue(GameplayCueTag);
+
+	TrackedGameplayCues.Remove(GameplayCueTag);
 }
 
 FGameplayEffectContextHandle UGameplayAbility::GetContextFromOwner(FGameplayAbilityTargetDataHandle OptionalTargetData) const
diff --git a/Engine/Source/Runtime/GameplayAbilities/Private/Abilities/GameplayAbilityTargetActor.cpp b/Engine/Source/Runtime/GameplayAbilities/Private/Abilities/GameplayAbilityTargetActor.cpp
--- a/Engine/Source/Runtime/GameplayAbilities/Private/Abilities/GameplayAbilityTargetActor.cpp
+++ b/Engine/Source/Runtime/GameplayAbilities/Private/Abilities/GameplayAbilityTargetActor.cpp
@@ -20,6 +20,28 @@ AGameplayAbilityTargetActor::AGameplayAbilityTargetActor(const FObjectInitialize
 	bDestroyOnConfirmation = true;
 }
 
+void AGameplayAbilityTargetActor::Destroyed()
+{
+	// We must remove ourselves from GenericLocalConfirmCallbacks/GenericLocalCancelCallbacks, since while these are bound they will inhibit any *other* abilities
+	// that are bound to the same key.
+
+	if (OwningAbility)
+	{
+		const FGameplayAbilityActorInfo* Info = OwningAbility->GetCurrentActorInfo();
+		if (Info && Info->IsLocallyControlled())
+		{
+			UAbilitySystemComponent* ASC = Info->AbilitySystemComponent.Get();
+			if (ASC)
+			{
+				ASC->GenericLocalConfirmCallbacks.RemoveDynamic(this, &AGameplayAbilityTargetActor::ConfirmTargeting);
+				ASC->GenericLocalCancelCallbacks.RemoveDynamic(this, &AGameplayAbilityTargetActor::CancelTargeting);
+			}
+		}
+	}
+
+	Super::Destroyed();
+}
+
 void AGameplayAbilityTargetActor::GetLifetimeReplicatedProps(TArray< FLifetimeProperty > & OutLifetimeProps) const
 {
 	Super::GetLifetimeReplicatedProps(OutLifetimeProps);
@@ -51,7 +73,7 @@ void AGameplayAbilityTargetActor::ConfirmTargetingAndContinue()
 void AGameplayAbilityTargetActor::ConfirmTargeting()
 {
 	UAbilitySystemComponent* ASC = OwningAbility->GetCurrentActorInfo()->AbilitySystemComponent.Get();
-	ASC->AbilityReplicatedEventDelegate(EAbilityReplicatedClientEvent::GenericConfirm, OwningAbility->GetCurrentAbilitySpecHandle(), OwningAbility->GetCurrentActivationInfo().GetActivationPredictionKey() ).Remove(GenericConfirmHandle);
+	ASC->AbilityReplicatedEventDelegate(EAbilityGenericReplicatedEvent::GenericConfirm, OwningAbility->GetCurrentAbilitySpecHandle(), OwningAbility->GetCurrentActivationInfo().GetActivationPredictionKey() ).Remove(GenericConfirmHandle);
 
 	if (IsConfirmTargetingAllowed())
 	{
@@ -67,7 +89,7 @@ void AGameplayAbilityTargetActor::ConfirmTargeting()
 void AGameplayAbilityTargetActor::CancelTargeting()
 {
 	UAbilitySystemComponent* ASC = OwningAbility->GetCurrentActorInfo()->AbilitySystemComponent.Get();
-	ASC->AbilityReplicatedEventDelegate(EAbilityReplicatedClientEvent::GenericCancel, OwningAbility->GetCurrentAbilitySpecHandle(), OwningAbility->GetCurrentActivationInfo().GetActivationPredictionKey() ).Remove(GenericCancelHandle);
+	ASC->AbilityReplicatedEventDelegate(EAbilityGenericReplicatedEvent::GenericCancel, OwningAbility->GetCurrentAbilitySpecHandle(), OwningAbility->GetCurrentActivationInfo().GetActivationPredictionKey() ).Remove(GenericCancelHandle);
 
 	CanceledDelegate.Broadcast(FGameplayAbilityTargetDataHandle());
 	Destroy();
@@ -128,18 +150,18 @@ void AGameplayAbilityTargetActor::BindToConfirmCancelInputs()
 			FGameplayAbilitySpecHandle Handle = OwningAbility->GetCurrentAbilitySpecHandle();
 			FPredictionKey PredKey = OwningAbility->GetCurrentActivationInfo().GetActivationPredictionKey();
 
-			GenericConfirmHandle = ASC->AbilityReplicatedEventDelegate(EAbilityReplicatedClientEvent::GenericConfirm, Handle, PredKey ).AddUObject(this, &AGameplayAbilityTargetActor::ConfirmTargeting);
-			GenericCancelHandle = ASC->AbilityReplicatedEventDelegate(EAbilityReplicatedClientEvent::GenericCancel, Handle, PredKey ).AddUObject(this, &AGameplayAbilityTargetActor::CancelTargeting);
+			GenericConfirmHandle = ASC->AbilityReplicatedEventDelegate(EAbilityGenericReplicatedEvent::GenericConfirm, Handle, PredKey ).AddUObject(this, &AGameplayAbilityTargetActor::ConfirmTargeting);
+			GenericCancelHandle = ASC->AbilityReplicatedEventDelegate(EAbilityGenericReplicatedEvent::GenericCancel, Handle, PredKey ).AddUObject(this, &AGameplayAbilityTargetActor::CancelTargeting);
 			
-			if (ASC->CallReplicatedEventDelegateIfSet(EAbilityReplicatedClientEvent::GenericConfirm, Handle, PredKey))
+			if (ASC->CallReplicatedEventDelegateIfSet(EAbilityGenericReplicatedEvent::GenericConfirm, Handle, PredKey))
 			{
 				return;
 			}
 			
-			if (ASC->CallReplicatedEventDelegateIfSet(EAbilityReplicatedClientEvent::GenericCancel, Handle, PredKey))
+			if (ASC->CallReplicatedEventDelegateIfSet(EAbilityGenericReplicatedEvent::GenericCancel, Handle, PredKey))
 			{
 				return;
 			}
 		}
 	}
-}
\ No newline at end of file
+}
diff --git a/Engine/Source/Runtime/GameplayAbilities/Private/Abilities/GameplayAbility_CharacterJump.cpp b/Engine/Source/Runtime/GameplayAbilities/Private/Abilities/GameplayAbility_CharacterJump.cpp
--- a/Engine/Source/Runtime/GameplayAbilities/Private/Abilities/GameplayAbility_CharacterJump.cpp
+++ b/Engine/Source/Runtime/GameplayAbilities/Private/Abilities/GameplayAbility_CharacterJump.cpp
@@ -35,7 +35,7 @@ void UGameplayAbility_CharacterJump::InputReleased(const FGameplayAbilitySpecHan
 {
 	if (ActorInfo != NULL && ActorInfo->AvatarActor != NULL)
 	{
-		CancelAbility(Handle, ActorInfo, ActivationInfo);
+		CancelAbility(Handle, ActorInfo, ActivationInfo, true);
 	}
 }
 
@@ -56,9 +56,9 @@ bool UGameplayAbility_CharacterJump::CanActivateAbility(const FGameplayAbilitySp
  *	Montage that *it* played was still playing, and if so, to cancel it. If this is something we need to support, we may need some
  *	light weight data structure to represent 'non intanced abilities in action' with a way to cancel/end them.
  */
-void UGameplayAbility_CharacterJump::CancelAbility(const FGameplayAbilitySpecHandle Handle, const FGameplayAbilityActorInfo* ActorInfo, const FGameplayAbilityActivationInfo ActivationInfo)
+void UGameplayAbility_CharacterJump::CancelAbility(const FGameplayAbilitySpecHandle Handle, const FGameplayAbilityActorInfo* ActorInfo, const FGameplayAbilityActivationInfo ActivationInfo, bool bReplicateCancelAbility)
 {
-	Super::CancelAbility(Handle, ActorInfo, ActivationInfo);
+	Super::CancelAbility(Handle, ActorInfo, ActivationInfo, bReplicateCancelAbility);
 	
 	ACharacter * Character = CastChecked<ACharacter>(ActorInfo->AvatarActor.Get());
 	Character->StopJumping();
diff --git a/Engine/Source/Runtime/GameplayAbilities/Private/Abilities/Tasks/AbilityTask.cpp b/Engine/Source/Runtime/GameplayAbilities/Private/Abilities/Tasks/AbilityTask.cpp
--- a/Engine/Source/Runtime/GameplayAbilities/Private/Abilities/Tasks/AbilityTask.cpp
+++ b/Engine/Source/Runtime/GameplayAbilities/Private/Abilities/Tasks/AbilityTask.cpp
@@ -184,7 +184,7 @@ bool UAbilityTask::IsForRemoteClient() const
 	return false;
 }
 
-bool UAbilityTask::CallOrAddReplicatedDelegate(EAbilityReplicatedClientEvent::Type Event, FSimpleMulticastDelegate::FDelegate Delegate)
+bool UAbilityTask::CallOrAddReplicatedDelegate(EAbilityGenericReplicatedEvent::Type Event, FSimpleMulticastDelegate::FDelegate Delegate)
 {
 	return AbilitySystemComponent->CallOrAddReplicatedDelegate(Event, GetAbilitySpecHandle(), GetActivationPredictionKey(), Delegate);
 }
diff --git a/Engine/Source/Runtime/GameplayAbilities/Private/Abilities/Tasks/AbilityTask_MoveToLocation.cpp b/Engine/Source/Runtime/GameplayAbilities/Private/Abilities/Tasks/AbilityTask_MoveToLocation.cpp
--- a/Engine/Source/Runtime/GameplayAbilities/Private/Abilities/Tasks/AbilityTask_MoveToLocation.cpp
+++ b/Engine/Source/Runtime/GameplayAbilities/Private/Abilities/Tasks/AbilityTask_MoveToLocation.cpp
@@ -107,4 +107,23 @@ void UAbilityTask_MoveToLocation::GetLifetimeReplicatedProps(TArray< FLifetimePr
 	DOREPLIFETIME(UAbilityTask_MoveToLocation, TargetLocation);
 	DOREPLIFETIME(UAbilityTask_MoveToLocation, DurationOfMovement);
 	DOREPLIFETIME(UAbilityTask_MoveToLocation, LerpCurve);
-}
\ No newline at end of file
+}
+
+void UAbilityTask_MoveToLocation::OnDestroy(bool AbilityIsEnding)
+{
+	AActor* MyActor = GetAvatarActor();
+	if (MyActor)
+	{
+		ACharacter* MyCharacter = Cast<ACharacter>(MyActor);
+		if (MyCharacter)
+		{
+			UCharacterMovementComponent* CharMoveComp = Cast<UCharacterMovementComponent>(MyCharacter->GetMovementComponent());
+			if (CharMoveComp && CharMoveComp->MovementMode == MOVE_Custom)
+			{
+				CharMoveComp->SetMovementMode(MOVE_Falling);
+			}
+		}
+	}
+
+	Super::OnDestroy(AbilityIsEnding);
+}
diff --git a/Engine/Source/Runtime/GameplayAbilities/Private/Abilities/Tasks/AbilityTask_NetworkSyncPoint.cpp b/Engine/Source/Runtime/GameplayAbilities/Private/Abilities/Tasks/AbilityTask_NetworkSyncPoint.cpp
new file mode 100644
--- /dev/null
+++ b/Engine/Source/Runtime/GameplayAbilities/Private/Abilities/Tasks/AbilityTask_NetworkSyncPoint.cpp
@@ -0,0 +1,85 @@
+// Copyright 1998-2015 Epic Games, Inc. All Rights Reserved.
+
+#include "AbilitySystemPrivatePCH.h"
+#include "Abilities/Tasks/AbilityTask_NetworkSyncPoint.h"
+#include "GameplayPrediction.h"
+
+UAbilityTask_NetworkSyncPoint::UAbilityTask_NetworkSyncPoint(const FObjectInitializer& ObjectInitializer)
+	: Super(ObjectInitializer)
+{
+	ReplicatedEventToListenFor = EAbilityGenericReplicatedEvent::MAX;
+}
+
+void UAbilityTask_NetworkSyncPoint::OnSignalCallback()
+{
+	if (AbilitySystemComponent.IsValid())
+	{
+		AbilitySystemComponent->ConsumeGenericReplicatedEvent(ReplicatedEventToListenFor, GetAbilitySpecHandle(), GetActivationPredictionKey());
+	}
+	SyncFinished();
+}
+
+UAbilityTask_NetworkSyncPoint* UAbilityTask_NetworkSyncPoint::WaitNetSync(class UObject* WorldContextObject, EAbilityTaskNetSyncType InSyncType)
+{
+	auto MyObj = NewTask<UAbilityTask_NetworkSyncPoint>(WorldContextObject);
+	MyObj->SyncType = InSyncType;
+	return MyObj;
+}
+
+void UAbilityTask_NetworkSyncPoint::Activate()
+{
+	if (AbilitySystemComponent.IsValid())
+	{
+		if (IsPredictingClient())
+		{
+			if (SyncType != EAbilityTaskNetSyncType::OnlyServerWait )
+			{
+				// As long as we are waiting (!= OnlyServerWait), listen for the GenericSignalFromServer event
+				ReplicatedEventToListenFor = EAbilityGenericReplicatedEvent::GenericSignalFromServer;
+			}
+			if (SyncType != EAbilityTaskNetSyncType::OnlyClientWait)
+			{
+				// As long as the server is waiting (!= OnlyClientWait), send the Server and RPC for this signal
+				AbilitySystemComponent->ServerSetReplicatedEvent(EAbilityGenericReplicatedEvent::GenericSignalFromClient, GetAbilitySpecHandle(), GetActivationPredictionKey(), AbilitySystemComponent->ScopedPredictionKey);
+			}
+			
+		}
+		else if (IsForRemoteClient())
+		{
+			if (SyncType != EAbilityTaskNetSyncType::OnlyClientWait )
+			{
+				// As long as we are waiting (!= OnlyClientWait), listen for the GenericSignalFromClient event
+				ReplicatedEventToListenFor = EAbilityGenericReplicatedEvent::GenericSignalFromClient;
+			}
+			if (SyncType != EAbilityTaskNetSyncType::OnlyServerWait)
+			{
+				// As long as the client is waiting (!= OnlyServerWait), send the Server and RPC for this signal
+				AbilitySystemComponent->ClientSetReplicatedEvent(EAbilityGenericReplicatedEvent::GenericSignalFromServer, GetAbilitySpecHandle(), GetActivationPredictionKey());
+			}
+		}
+
+		if (ReplicatedEventToListenFor != EAbilityGenericReplicatedEvent::MAX)
+		{
+			CallOrAddReplicatedDelegate(ReplicatedEventToListenFor, FSimpleMulticastDelegate::FDelegate::CreateUObject(this, &UAbilityTask_NetworkSyncPoint::OnSignalCallback));
+		}
+		else
+		{
+			// We aren't waiting for a replicated event, so the sync is complete.
+			SyncFinished();
+		}
+	}
+}
+
+void UAbilityTask_NetworkSyncPoint::OnDestroy(bool AbilityEnded)
+{
+	Super::OnDestroy(AbilityEnded);
+}
+
+void UAbilityTask_NetworkSyncPoint::SyncFinished()
+{
+	if (!IsPendingKill())
+	{
+		OnSync.Broadcast();
+		EndTask();
+	}
+}
\ No newline at end of file
diff --git a/Engine/Source/Runtime/GameplayAbilities/Private/Abilities/Tasks/AbilityTask_PlayMontageAndWait.cpp b/Engine/Source/Runtime/GameplayAbilities/Private/Abilities/Tasks/AbilityTask_PlayMontageAndWait.cpp
--- a/Engine/Source/Runtime/GameplayAbilities/Private/Abilities/Tasks/AbilityTask_PlayMontageAndWait.cpp
+++ b/Engine/Source/Runtime/GameplayAbilities/Private/Abilities/Tasks/AbilityTask_PlayMontageAndWait.cpp
@@ -33,6 +33,26 @@ void UAbilityTask_PlayMontageAndWait::OnMontageEnded(UAnimMontage* Montage, bool
 	EndTask();
 }
 
+void UAbilityTask_PlayMontageAndWait::OnMontageInterrupted()
+{
+	// Check if the montage is still playing
+	// The ability would have been interrupted, in which case we should automatically stop the montage
+	if (AbilitySystemComponent.IsValid() && Ability.IsValid())
+	{
+		if (AbilitySystemComponent->IsAnimatingAbility(Ability.Get())
+			&& AbilitySystemComponent->GetCurrentMontage() == MontageToPlay)
+		{
+			// Unbind the delegate in this case so OnMontageEnded does not get called as well
+			BlendingOutDelegate.Unbind();
+
+			AbilitySystemComponent->CurrentMontageStop();
+
+			// Let the BP handle the interrupt as well
+			OnInterrupted.Broadcast();
+		}
+	}
+}
+
 UAbilityTask_PlayMontageAndWait* UAbilityTask_PlayMontageAndWait::CreatePlayMontageAndWaitProxy(UObject* WorldContextObject,
 	FName TaskInstanceName, UAnimMontage *MontageToPlay, float Rate, FName StartSection)
 {
@@ -53,7 +73,14 @@ void UAbilityTask_PlayMontageAndWait::Activate()
 		{
 			if (AbilitySystemComponent->PlayMontage(Ability.Get(), Ability->GetCurrentActivationInfo(), MontageToPlay, Rate, StartSection) > 0.f)
 			{
-				FOnMontageBlendingOutStarted BlendingOutDelegate;
+				// Playing a montage could potentially fire off a callback into game code which could kill this ability! Early out if we are  pending kill.
+				if (IsPendingKill())
+				{
+					return;
+				}
+
+				InterruptedHandle = Ability->OnGameplayAbilityCancelled.AddUObject(this, &UAbilityTask_PlayMontageAndWait::OnMontageInterrupted);
+
 				BlendingOutDelegate.BindUObject(this, &UAbilityTask_PlayMontageAndWait::OnMontageEnded);
 				ActorInfo->AnimInstance->Montage_SetBlendingOutDelegate(BlendingOutDelegate);
 			}
@@ -77,6 +104,12 @@ void UAbilityTask_PlayMontageAndWait::OnDestroy(bool AbilityEnded)
 	// Note: Clearing montage end delegate isn't necessary since its not a multicast and will be cleared when the next montage plays.
 	// (If we are destroyed, it will detect this and not do anything)
 
+	// This delegate, however, should be cleared as it is a multicast
+	if (Ability.IsValid())
+	{
+		Ability->OnGameplayAbilityCancelled.Remove(InterruptedHandle);
+	}
+
 	Super::OnDestroy(AbilityEnded);
 }
 
diff --git a/Engine/Source/Runtime/GameplayAbilities/Private/Abilities/Tasks/AbilityTask_StartAbilityState.cpp b/Engine/Source/Runtime/GameplayAbilities/Private/Abilities/Tasks/AbilityTask_StartAbilityState.cpp
--- a/Engine/Source/Runtime/GameplayAbilities/Private/Abilities/Tasks/AbilityTask_StartAbilityState.cpp
+++ b/Engine/Source/Runtime/GameplayAbilities/Private/Abilities/Tasks/AbilityTask_StartAbilityState.cpp
@@ -8,6 +8,8 @@ UAbilityTask_StartAbilityState::UAbilityTask_StartAbilityState(const FObjectInit
 : Super(ObjectInitializer)
 {
 	bEndCurrentState = true;
+	bWasEnded = false;
+	bWasInterrupted = false;
 }
 
 UAbilityTask_StartAbilityState* UAbilityTask_StartAbilityState::StartAbilityState(UObject* WorldContextObject, FName StateName, bool bEndCurrentState)
@@ -19,61 +21,59 @@ UAbilityTask_StartAbilityState* UAbilityTask_StartAbilityState::StartAbilityStat
 
 void UAbilityTask_StartAbilityState::Activate()
 {
-	UAbilitySystemComponent* ASC = AbilitySystemComponent.Get();
-
-	if (ASC)
+	if (Ability.IsValid())
 	{
-		if (bEndCurrentState && ASC->OnAbilityStateEnded.IsBound())
+		if (bEndCurrentState && Ability->OnGameplayAbilityStateEnded.IsBound())
 		{
-			ASC->OnAbilityStateEnded.Broadcast(NAME_None);
+			Ability->OnGameplayAbilityStateEnded.Broadcast(NAME_None);
 		}
 
-		EndStateHandle = ASC->OnAbilityStateEnded.AddUObject(this, &UAbilityTask_StartAbilityState::OnEndState);
-		InterruptStateHandle = ASC->OnAbilityStateInterrupted.AddUObject(this, &UAbilityTask_StartAbilityState::OnInterruptState);
+		EndStateHandle = Ability->OnGameplayAbilityStateEnded.AddUObject(this, &UAbilityTask_StartAbilityState::OnEndState);
+		InterruptStateHandle = Ability->OnGameplayAbilityCancelled.AddUObject(this, &UAbilityTask_StartAbilityState::OnInterruptState);
 	}
 }
 
 void UAbilityTask_StartAbilityState::OnDestroy(bool AbilityEnded)
-{	
-	// If the ability was ended, make sure to first notify the BP via 'OnStateEnded'
-	if (AbilityEnded && OnStateEnded.IsBound())
+{
+	Super::OnDestroy(AbilityEnded);
+
+	if (bWasInterrupted && OnStateInterrupted.IsBound())
+	{
+		OnStateInterrupted.Broadcast();
+	}
+	else if ((bWasEnded || AbilityEnded) && OnStateEnded.IsBound())
 	{
 		OnStateEnded.Broadcast();
 	}
 
-	UAbilitySystemComponent* ASC = AbilitySystemComponent.Get();
-
-	if (ASC)
+	if (Ability.IsValid())
 	{
-		ASC->OnAbilityStateInterrupted.Remove(InterruptStateHandle);
-		ASC->OnAbilityStateEnded.Remove(EndStateHandle);
+		Ability->OnGameplayAbilityCancelled.Remove(InterruptStateHandle);
+		Ability->OnGameplayAbilityStateEnded.Remove(EndStateHandle);
 	}
-
-	Super::OnDestroy(AbilityEnded);
 }
 
 void UAbilityTask_StartAbilityState::OnEndState(FName StateNameToEnd)
 {
-	// All states end if 'Name_NONE' is passed to this delegate
+	// All states end if 'NAME_None' is passed to this delegate
 	if (StateNameToEnd.IsNone() || StateNameToEnd == InstanceName)
 	{
-		EndTask();
+		bWasEnded = true;
 
-		if (OnStateEnded.IsBound())
-		{
-			OnStateEnded.Broadcast();
-		}
+		EndTask();
 	}
 }
 
 void UAbilityTask_StartAbilityState::OnInterruptState()
 {
-	EndTask();
+	bWasInterrupted = true;
+}
+
+void UAbilityTask_StartAbilityState::ExternalCancel()
+{
+	bWasInterrupted = true;
 
-	if (OnStateInterrupted.IsBound())
-	{
-		OnStateInterrupted.Broadcast();
-	}
+	Super::ExternalCancel();
 }
 
 FString UAbilityTask_StartAbilityState::GetDebugString() const
diff --git a/Engine/Source/Runtime/GameplayAbilities/Private/Abilities/Tasks/AbilityTask_WaitConfirmCancel.cpp b/Engine/Source/Runtime/GameplayAbilities/Private/Abilities/Tasks/AbilityTask_WaitConfirmCancel.cpp
--- a/Engine/Source/Runtime/GameplayAbilities/Private/Abilities/Tasks/AbilityTask_WaitConfirmCancel.cpp
+++ b/Engine/Source/Runtime/GameplayAbilities/Private/Abilities/Tasks/AbilityTask_WaitConfirmCancel.cpp
@@ -19,7 +19,7 @@ void UAbilityTask_WaitConfirmCancel::OnConfirmCallback()
 	
 	if (AbilitySystemComponent.IsValid())
 	{
-		AbilitySystemComponent->ConsumeClientReplicatedEvent(EAbilityReplicatedClientEvent::GenericConfirm, GetAbilitySpecHandle(), GetActivationPredictionKey());
+		AbilitySystemComponent->ConsumeGenericReplicatedEvent(EAbilityGenericReplicatedEvent::GenericConfirm, GetAbilitySpecHandle(), GetActivationPredictionKey());
 		OnConfirm.Broadcast();
 		EndTask();
 	}
@@ -29,7 +29,7 @@ void UAbilityTask_WaitConfirmCancel::OnCancelCallback()
 {
 	if (AbilitySystemComponent.IsValid())
 	{
-		AbilitySystemComponent->ConsumeClientReplicatedEvent(EAbilityReplicatedClientEvent::GenericCancel, GetAbilitySpecHandle(), GetActivationPredictionKey());
+		AbilitySystemComponent->ConsumeGenericReplicatedEvent(EAbilityGenericReplicatedEvent::GenericCancel, GetAbilitySpecHandle(), GetActivationPredictionKey());
 		OnCancel.Broadcast();
 		EndTask();
 	}
@@ -39,7 +39,7 @@ void UAbilityTask_WaitConfirmCancel::OnLocalConfirmCallback()
 {
 	if (AbilitySystemComponent.IsValid() && IsPredictingClient())
 	{
-		AbilitySystemComponent->ServerSetReplicatedClientEvent(EAbilityReplicatedClientEvent::GenericConfirm, GetAbilitySpecHandle(), GetActivationPredictionKey() ,AbilitySystemComponent->ScopedPredictionKey);
+		AbilitySystemComponent->ServerSetReplicatedEvent(EAbilityGenericReplicatedEvent::GenericConfirm, GetAbilitySpecHandle(), GetActivationPredictionKey() ,AbilitySystemComponent->ScopedPredictionKey);
 	}
 	OnConfirmCallback();
 }
@@ -48,7 +48,7 @@ void UAbilityTask_WaitConfirmCancel::OnLocalCancelCallback()
 {
 	if (AbilitySystemComponent.IsValid() && IsPredictingClient())
 	{
-		AbilitySystemComponent->ServerSetReplicatedClientEvent(EAbilityReplicatedClientEvent::GenericCancel, GetAbilitySpecHandle(), GetActivationPredictionKey() ,AbilitySystemComponent->ScopedPredictionKey);
+		AbilitySystemComponent->ServerSetReplicatedEvent(EAbilityGenericReplicatedEvent::GenericCancel, GetAbilitySpecHandle(), GetActivationPredictionKey() ,AbilitySystemComponent->ScopedPredictionKey);
 	}
 	OnCancelCallback();
 }
@@ -75,12 +75,12 @@ void UAbilityTask_WaitConfirmCancel::Activate()
 		}
 		else
 		{
-			if (CallOrAddReplicatedDelegate(EAbilityReplicatedClientEvent::GenericConfirm,  FSimpleMulticastDelegate::FDelegate::CreateUObject(this, &UAbilityTask_WaitConfirmCancel::OnConfirmCallback)) )
+			if (CallOrAddReplicatedDelegate(EAbilityGenericReplicatedEvent::GenericConfirm,  FSimpleMulticastDelegate::FDelegate::CreateUObject(this, &UAbilityTask_WaitConfirmCancel::OnConfirmCallback)) )
 			{
 				// GenericConfirm was already received from the client and we just called OnConfirmCallback. The task is done.
 				return;
 			}
-			if (CallOrAddReplicatedDelegate(EAbilityReplicatedClientEvent::GenericCancel,  FSimpleMulticastDelegate::FDelegate::CreateUObject(this, &UAbilityTask_WaitConfirmCancel::OnCancelCallback)) )
+			if (CallOrAddReplicatedDelegate(EAbilityGenericReplicatedEvent::GenericCancel,  FSimpleMulticastDelegate::FDelegate::CreateUObject(this, &UAbilityTask_WaitConfirmCancel::OnCancelCallback)) )
 			{
 				// GenericCancel was already received from the client and we just called OnCancelCallback. The task is done.
 				return;
@@ -98,4 +98,4 @@ void UAbilityTask_WaitConfirmCancel::OnDestroy(bool AbilityEnding)
 	}
 
 	Super::OnDestroy(AbilityEnding);
-}
\ No newline at end of file
+}
diff --git a/Engine/Source/Runtime/GameplayAbilities/Private/Abilities/Tasks/AbilityTask_WaitGameplayEffectApplied.cpp b/Engine/Source/Runtime/GameplayAbilities/Private/Abilities/Tasks/AbilityTask_WaitGameplayEffectApplied.cpp
--- a/Engine/Source/Runtime/GameplayAbilities/Private/Abilities/Tasks/AbilityTask_WaitGameplayEffectApplied.cpp
+++ b/Engine/Source/Runtime/GameplayAbilities/Private/Abilities/Tasks/AbilityTask_WaitGameplayEffectApplied.cpp
@@ -2,9 +2,7 @@
 
 #include "AbilitySystemPrivatePCH.h"
 #include "Abilities/Tasks/AbilityTask_WaitGameplayEffectApplied.h"
-
 #include "AbilitySystemComponent.h"
-
 #include "Abilities/GameplayAbility.h"
 #include "GameplayEffectExtension.h"
 
@@ -15,7 +13,7 @@ UAbilityTask_WaitGameplayEffectApplied::UAbilityTask_WaitGameplayEffectApplied(c
 
 void UAbilityTask_WaitGameplayEffectApplied::Activate()
 {
-	if (AbilitySystemComponent.IsValid())
+	if (GetASC())
 	{
 		RegisterDelegate();
 	}
@@ -50,10 +48,28 @@ void UAbilityTask_WaitGameplayEffectApplied::OnApplyGameplayEffectCallback(UAbil
 
 void UAbilityTask_WaitGameplayEffectApplied::OnDestroy(bool AbilityEnded)
 {
-	if (AbilitySystemComponent.IsValid())
+	if (GetASC())
 	{
 		RemoveDelegate();
 	}
 
 	Super::OnDestroy(AbilityEnded);
+}
+
+void UAbilityTask_WaitGameplayEffectApplied::SetExternalActor(AActor* InActor)
+{
+	if (InActor)
+	{
+		UseExternalOwner = true;
+		ExternalOwner  = UAbilitySystemGlobals::GetAbilitySystemComponentFromActor(InActor);
+	}
+}
+
+UAbilitySystemComponent* UAbilityTask_WaitGameplayEffectApplied::GetASC()
+{
+	if (UseExternalOwner)
+	{
+		return ExternalOwner.Get();
+	}
+	return AbilitySystemComponent.Get();
 }
\ No newline at end of file
diff --git a/Engine/Source/Runtime/GameplayAbilities/Private/Abilities/Tasks/AbilityTask_WaitGameplayEffectApplied_Self.cpp b/Engine/Source/Runtime/GameplayAbilities/Private/Abilities/Tasks/AbilityTask_WaitGameplayEffectApplied_Self.cpp
--- a/Engine/Source/Runtime/GameplayAbilities/Private/Abilities/Tasks/AbilityTask_WaitGameplayEffectApplied_Self.cpp
+++ b/Engine/Source/Runtime/GameplayAbilities/Private/Abilities/Tasks/AbilityTask_WaitGameplayEffectApplied_Self.cpp
@@ -13,13 +13,14 @@ UAbilityTask_WaitGameplayEffectApplied_Self::UAbilityTask_WaitGameplayEffectAppl
 {
 }
 
-UAbilityTask_WaitGameplayEffectApplied_Self* UAbilityTask_WaitGameplayEffectApplied_Self::WaitGameplayEffectAppliedToSelf(UObject* WorldContextObject, const FGameplayTargetDataFilterHandle InFilter, FGameplayTagRequirements InSourceTagRequirements, FGameplayTagRequirements InTargetTagRequirements, bool InTriggerOnce)
+UAbilityTask_WaitGameplayEffectApplied_Self* UAbilityTask_WaitGameplayEffectApplied_Self::WaitGameplayEffectAppliedToSelf(UObject* WorldContextObject, AActor* TestActor, const FGameplayTargetDataFilterHandle InFilter, FGameplayTagRequirements InSourceTagRequirements, FGameplayTagRequirements InTargetTagRequirements, bool InTriggerOnce, AActor* OptionalExternalOwner)
 {
 	auto MyObj = NewTask<UAbilityTask_WaitGameplayEffectApplied_Self>(WorldContextObject);
 	MyObj->Filter = InFilter;
 	MyObj->SourceTagRequirements = InSourceTagRequirements;
 	MyObj->TargetTagRequirements = InTargetTagRequirements;
 	MyObj->TriggerOnce = InTriggerOnce;
+	MyObj->SetExternalActor(OptionalExternalOwner);
 	return MyObj;
 }
 
@@ -30,10 +31,10 @@ void UAbilityTask_WaitGameplayEffectApplied_Self::BroadcastDelegate(AActor* Avat
 
 void UAbilityTask_WaitGameplayEffectApplied_Self::RegisterDelegate()
 {
-	OnApplyGameplayEffectCallbackDelegateHandle = AbilitySystemComponent->OnGameplayEffectAppliedDelegateToSelf.AddUObject(this, &UAbilityTask_WaitGameplayEffectApplied::OnApplyGameplayEffectCallback);
+	OnApplyGameplayEffectCallbackDelegateHandle = GetASC()->OnGameplayEffectAppliedDelegateToSelf.AddUObject(this, &UAbilityTask_WaitGameplayEffectApplied::OnApplyGameplayEffectCallback);
 }
 
 void UAbilityTask_WaitGameplayEffectApplied_Self::RemoveDelegate()
 {
-	AbilitySystemComponent->OnGameplayEffectAppliedDelegateToSelf.Remove(OnApplyGameplayEffectCallbackDelegateHandle);
+	GetASC()->OnGameplayEffectAppliedDelegateToSelf.Remove(OnApplyGameplayEffectCallbackDelegateHandle);
 }
\ No newline at end of file
diff --git a/Engine/Source/Runtime/GameplayAbilities/Private/Abilities/Tasks/AbilityTask_WaitGameplayEffectApplied_Target.cpp b/Engine/Source/Runtime/GameplayAbilities/Private/Abilities/Tasks/AbilityTask_WaitGameplayEffectApplied_Target.cpp
--- a/Engine/Source/Runtime/GameplayAbilities/Private/Abilities/Tasks/AbilityTask_WaitGameplayEffectApplied_Target.cpp
+++ b/Engine/Source/Runtime/GameplayAbilities/Private/Abilities/Tasks/AbilityTask_WaitGameplayEffectApplied_Target.cpp
@@ -13,13 +13,14 @@ UAbilityTask_WaitGameplayEffectApplied_Target::UAbilityTask_WaitGameplayEffectAp
 {
 }
 
-UAbilityTask_WaitGameplayEffectApplied_Target* UAbilityTask_WaitGameplayEffectApplied_Target::WaitGameplayEffectAppliedToTarget(UObject* WorldContextObject, const FGameplayTargetDataFilterHandle InFilter, FGameplayTagRequirements InSourceTagRequirements, FGameplayTagRequirements InTargetTagRequirements, bool InTriggerOnce)
+UAbilityTask_WaitGameplayEffectApplied_Target* UAbilityTask_WaitGameplayEffectApplied_Target::WaitGameplayEffectAppliedToTarget(UObject* WorldContextObject, const FGameplayTargetDataFilterHandle InFilter, FGameplayTagRequirements InSourceTagRequirements, FGameplayTagRequirements InTargetTagRequirements, bool InTriggerOnce, AActor* OptionalExternalOwner)
 {
 	auto MyObj = NewTask<UAbilityTask_WaitGameplayEffectApplied_Target>(WorldContextObject);
 	MyObj->Filter = InFilter;
 	MyObj->SourceTagRequirements = InSourceTagRequirements;
 	MyObj->TargetTagRequirements = InTargetTagRequirements;
 	MyObj->TriggerOnce = InTriggerOnce;
+	MyObj->SetExternalActor(OptionalExternalOwner);
 	return MyObj;
 }
 
@@ -30,10 +31,10 @@ void UAbilityTask_WaitGameplayEffectApplied_Target::BroadcastDelegate(AActor* Av
 
 void UAbilityTask_WaitGameplayEffectApplied_Target::RegisterDelegate()
 {
-	OnApplyGameplayEffectCallbackDelegateHandle = AbilitySystemComponent->OnGameplayEffectAppliedDelegateToTarget.AddUObject(this, &UAbilityTask_WaitGameplayEffectApplied::OnApplyGameplayEffectCallback);
+	OnApplyGameplayEffectCallbackDelegateHandle = GetASC()->OnGameplayEffectAppliedDelegateToTarget.AddUObject(this, &UAbilityTask_WaitGameplayEffectApplied::OnApplyGameplayEffectCallback);
 }
 
 void UAbilityTask_WaitGameplayEffectApplied_Target::RemoveDelegate()
 {
-	AbilitySystemComponent->OnGameplayEffectAppliedDelegateToTarget.Remove(OnApplyGameplayEffectCallbackDelegateHandle);
+	GetASC()->OnGameplayEffectAppliedDelegateToTarget.Remove(OnApplyGameplayEffectCallbackDelegateHandle);
 }
\ No newline at end of file
diff --git a/Engine/Source/Runtime/GameplayAbilities/Private/Abilities/Tasks/AbilityTask_WaitGameplayTag.cpp b/Engine/Source/Runtime/GameplayAbilities/Private/Abilities/Tasks/AbilityTask_WaitGameplayTag.cpp
new file mode 100644
--- /dev/null
+++ b/Engine/Source/Runtime/GameplayAbilities/Private/Abilities/Tasks/AbilityTask_WaitGameplayTag.cpp
@@ -0,0 +1,100 @@
+// Copyright 1998-2015 Epic Games, Inc. All Rights Reserved.
+
+#include "AbilitySystemPrivatePCH.h"
+#include "Abilities/Tasks/AbilityTask_WaitGameplayTag.h"
+#include "AbilitySystemComponent.h"
+#include "Abilities/GameplayAbility.h"
+
+// ----------------------------------------------------------------
+
+UAbilityTask_WaitGameplayTagAdded::UAbilityTask_WaitGameplayTagAdded(const FObjectInitializer& ObjectInitializer)
+: Super(ObjectInitializer)
+{
+
+}
+
+UAbilityTask_WaitGameplayTagAdded* UAbilityTask_WaitGameplayTagAdded::WaitGameplayTagAdd(UObject* WorldContextObject, FGameplayTag Tag, AActor* OptionalExternalTarget, bool OnlyTriggerOnce)
+{
+	auto MyObj = NewTask<UAbilityTask_WaitGameplayTagAdded>(WorldContextObject);
+	MyObj->Tag = Tag;
+	MyObj->SetExternalTarget(OptionalExternalTarget);
+	MyObj->OnlyTriggerOnce = OnlyTriggerOnce;
+
+	return MyObj;
+}
+
+void UAbilityTask_WaitGameplayTagAdded::Activate()
+{
+	UAbilitySystemComponent* ASC = GetTargetASC();
+	if (ASC && ASC->HasMatchingGameplayTag(Tag))
+	{			
+		Added.Broadcast();
+		if(OnlyTriggerOnce)
+		{
+			EndTask();
+			return;
+		}
+	}
+
+	Super::Activate();
+}
+
+void UAbilityTask_WaitGameplayTagAdded::GameplayTagCallback(const FGameplayTag Tag, int32 NewCount)
+{
+	if (NewCount==1)
+	{
+		Added.Broadcast();
+		if(OnlyTriggerOnce)
+		{
+			EndTask();
+		}
+	}
+}
+
+// ----------------------------------------------------------------
+
+UAbilityTask_WaitGameplayTagRemoved::UAbilityTask_WaitGameplayTagRemoved(const FObjectInitializer& ObjectInitializer)
+: Super(ObjectInitializer)
+{
+
+}
+
+UAbilityTask_WaitGameplayTagRemoved* UAbilityTask_WaitGameplayTagRemoved::WaitGameplayTagRemove(UObject* WorldContextObject, FGameplayTag Tag, AActor* OptionalExternalTarget, bool OnlyTriggerOnce)
+{
+	auto MyObj = NewTask<UAbilityTask_WaitGameplayTagRemoved>(WorldContextObject);
+	MyObj->Tag = Tag;
+	MyObj->SetExternalTarget(OptionalExternalTarget);
+	MyObj->OnlyTriggerOnce = OnlyTriggerOnce;
+
+	return MyObj;
+}
+
+void UAbilityTask_WaitGameplayTagRemoved::Activate()
+{
+	UAbilitySystemComponent* ASC = GetTargetASC();
+	if (ASC && ASC->HasMatchingGameplayTag(Tag) == false)
+	{			
+		Removed.Broadcast();
+		if(OnlyTriggerOnce)
+		{
+			EndTask();
+			return;
+		}
+	}
+
+	Super::Activate();
+}
+
+void UAbilityTask_WaitGameplayTagRemoved::GameplayTagCallback(const FGameplayTag Tag, int32 NewCount)
+{
+	if (NewCount==0)
+	{
+		Removed.Broadcast();
+		if(OnlyTriggerOnce)
+		{
+			EndTask();
+		}
+	}
+}
+
+// ----------------------------------------------------------------
\ No newline at end of file
diff --git a/Engine/Source/Runtime/GameplayAbilities/Private/Abilities/Tasks/AbilityTask_WaitGameplayTagBase.cpp b/Engine/Source/Runtime/GameplayAbilities/Private/Abilities/Tasks/AbilityTask_WaitGameplayTagBase.cpp
new file mode 100644
--- /dev/null
+++ b/Engine/Source/Runtime/GameplayAbilities/Private/Abilities/Tasks/AbilityTask_WaitGameplayTagBase.cpp
@@ -0,0 +1,56 @@
+// Copyright 1998-2015 Epic Games, Inc. All Rights Reserved.
+
+#include "AbilitySystemPrivatePCH.h"
+#include "Abilities/Tasks/AbilityTask_WaitGameplayTagBase.h"
+#include "AbilitySystemComponent.h"
+#include "Abilities/GameplayAbility.h"
+
+UAbilityTask_WaitGameplayTag::UAbilityTask_WaitGameplayTag(const FObjectInitializer& ObjectInitializer)
+	: Super(ObjectInitializer)
+{
+	RegisteredCallback = false;
+	UseExternalTarget = false;
+	OnlyTriggerOnce = false;
+}
+
+void UAbilityTask_WaitGameplayTag::Activate()
+{
+	if (AbilitySystemComponent.IsValid())
+	{
+		DelegateHandle = AbilitySystemComponent->RegisterGameplayTagEvent(Tag).AddUObject(this, &UAbilityTask_WaitGameplayTag::GameplayTagCallback);
+		RegisteredCallback = true;
+	}
+}
+
+void UAbilityTask_WaitGameplayTag::GameplayTagCallback(const FGameplayTag Tag, int32 NewCount)
+{
+}
+
+void UAbilityTask_WaitGameplayTag::OnDestroy(bool AbilityIsEnding)
+{
+	if (RegisteredCallback && AbilitySystemComponent.IsValid())
+	{
+		AbilitySystemComponent->RegisterGameplayTagEvent(Tag).Remove(DelegateHandle);
+	}
+
+	Super::OnDestroy(AbilityIsEnding);
+}
+
+UAbilitySystemComponent* UAbilityTask_WaitGameplayTag::GetTargetASC()
+{
+	if (UseExternalTarget)
+	{
+		return OptionalExternalTarget.Get();
+	}
+
+	return AbilitySystemComponent.Get();
+}
+
+void UAbilityTask_WaitGameplayTag::SetExternalTarget(AActor* Actor)
+{
+	if (Actor)
+	{
+		UseExternalTarget = true;
+		OptionalExternalTarget = UAbilitySystemGlobals::GetAbilitySystemComponentFromActor(Actor);
+	}
+}
diff --git a/Engine/Source/Runtime/GameplayAbilities/Private/Abilities/Tasks/AbilityTask_WaitInputPress.cpp b/Engine/Source/Runtime/GameplayAbilities/Private/Abilities/Tasks/AbilityTask_WaitInputPress.cpp
--- a/Engine/Source/Runtime/GameplayAbilities/Private/Abilities/Tasks/AbilityTask_WaitInputPress.cpp
+++ b/Engine/Source/Runtime/GameplayAbilities/Private/Abilities/Tasks/AbilityTask_WaitInputPress.cpp
@@ -19,16 +19,16 @@ void UAbilityTask_WaitInputPress::OnPressCallback()
 		return;
 	}
 
-	AbilitySystemComponent->AbilityReplicatedEventDelegate(EAbilityReplicatedClientEvent::InputPressed, GetAbilitySpecHandle(), GetActivationPredictionKey()).Remove(DelegateHandle);
+	AbilitySystemComponent->AbilityReplicatedEventDelegate(EAbilityGenericReplicatedEvent::InputPressed, GetAbilitySpecHandle(), GetActivationPredictionKey()).Remove(DelegateHandle);
 
 	if (IsPredictingClient())
 	{
 		// Tell the server about this
-		AbilitySystemComponent->ServerSetReplicatedClientEvent(EAbilityReplicatedClientEvent::InputPressed, GetAbilitySpecHandle(), GetActivationPredictionKey(), AbilitySystemComponent->ScopedPredictionKey);
+		AbilitySystemComponent->ServerSetReplicatedEvent(EAbilityGenericReplicatedEvent::InputPressed, GetAbilitySpecHandle(), GetActivationPredictionKey(), AbilitySystemComponent->ScopedPredictionKey);
 	}
 	else
 	{
-		AbilitySystemComponent->ConsumeClientReplicatedEvent(EAbilityReplicatedClientEvent::InputPressed, GetAbilitySpecHandle(), GetActivationPredictionKey());
+		AbilitySystemComponent->ConsumeGenericReplicatedEvent(EAbilityGenericReplicatedEvent::InputPressed, GetAbilitySpecHandle(), GetActivationPredictionKey());
 	}
 
 	// We are done. Kill us so we don't keep getting broadcast messages
@@ -46,10 +46,10 @@ void UAbilityTask_WaitInputPress::Activate()
 	StartTime = GetWorld()->GetTimeSeconds();
 	if (Ability.IsValid())
 	{
-		DelegateHandle = AbilitySystemComponent->AbilityReplicatedEventDelegate(EAbilityReplicatedClientEvent::InputPressed, GetAbilitySpecHandle(), GetActivationPredictionKey()).AddUObject(this, &UAbilityTask_WaitInputPress::OnPressCallback);
+		DelegateHandle = AbilitySystemComponent->AbilityReplicatedEventDelegate(EAbilityGenericReplicatedEvent::InputPressed, GetAbilitySpecHandle(), GetActivationPredictionKey()).AddUObject(this, &UAbilityTask_WaitInputPress::OnPressCallback);
 		if (IsForRemoteClient())
 		{
-			AbilitySystemComponent->CallReplicatedEventDelegateIfSet(EAbilityReplicatedClientEvent::InputPressed, GetAbilitySpecHandle(), GetActivationPredictionKey());
+			AbilitySystemComponent->CallReplicatedEventDelegateIfSet(EAbilityGenericReplicatedEvent::InputPressed, GetAbilitySpecHandle(), GetActivationPredictionKey());
 		}
 	}
 }
diff --git a/Engine/Source/Runtime/GameplayAbilities/Private/Abilities/Tasks/AbilityTask_WaitInputRelease.cpp b/Engine/Source/Runtime/GameplayAbilities/Private/Abilities/Tasks/AbilityTask_WaitInputRelease.cpp
--- a/Engine/Source/Runtime/GameplayAbilities/Private/Abilities/Tasks/AbilityTask_WaitInputRelease.cpp
+++ b/Engine/Source/Runtime/GameplayAbilities/Private/Abilities/Tasks/AbilityTask_WaitInputRelease.cpp
@@ -19,16 +19,16 @@ void UAbilityTask_WaitInputRelease::OnReleaseCallback()
 		return;
 	}
 
-	AbilitySystemComponent->AbilityReplicatedEventDelegate(EAbilityReplicatedClientEvent::InputReleased, GetAbilitySpecHandle(), GetActivationPredictionKey()).Remove(DelegateHandle);
+	AbilitySystemComponent->AbilityReplicatedEventDelegate(EAbilityGenericReplicatedEvent::InputReleased, GetAbilitySpecHandle(), GetActivationPredictionKey()).Remove(DelegateHandle);
 
 	if (IsPredictingClient())
 	{
 		// Tell the server about this
-		AbilitySystemComponent->ServerSetReplicatedClientEvent(EAbilityReplicatedClientEvent::InputReleased, GetAbilitySpecHandle(), GetActivationPredictionKey(), AbilitySystemComponent->ScopedPredictionKey);
+		AbilitySystemComponent->ServerSetReplicatedEvent(EAbilityGenericReplicatedEvent::InputReleased, GetAbilitySpecHandle(), GetActivationPredictionKey(), AbilitySystemComponent->ScopedPredictionKey);
 	}
 	else
 	{
-		AbilitySystemComponent->ConsumeClientReplicatedEvent(EAbilityReplicatedClientEvent::InputReleased, GetAbilitySpecHandle(), GetActivationPredictionKey());
+		AbilitySystemComponent->ConsumeGenericReplicatedEvent(EAbilityGenericReplicatedEvent::InputReleased, GetAbilitySpecHandle(), GetActivationPredictionKey());
 	}
 
 	// We are done. Kill us so we don't keep getting broadcast messages
@@ -46,10 +46,10 @@ void UAbilityTask_WaitInputRelease::Activate()
 	StartTime = GetWorld()->GetTimeSeconds();
 	if (Ability.IsValid())
 	{
-		DelegateHandle = AbilitySystemComponent->AbilityReplicatedEventDelegate(EAbilityReplicatedClientEvent::InputReleased, GetAbilitySpecHandle(), GetActivationPredictionKey()).AddUObject(this, &UAbilityTask_WaitInputRelease::OnReleaseCallback);
+		DelegateHandle = AbilitySystemComponent->AbilityReplicatedEventDelegate(EAbilityGenericReplicatedEvent::InputReleased, GetAbilitySpecHandle(), GetActivationPredictionKey()).AddUObject(this, &UAbilityTask_WaitInputRelease::OnReleaseCallback);
 		if (IsForRemoteClient())
 		{
-			AbilitySystemComponent->CallReplicatedEventDelegateIfSet(EAbilityReplicatedClientEvent::InputReleased, GetAbilitySpecHandle(), GetActivationPredictionKey());
+			AbilitySystemComponent->CallReplicatedEventDelegateIfSet(EAbilityGenericReplicatedEvent::InputReleased, GetAbilitySpecHandle(), GetActivationPredictionKey());
 		}
 	}
 }
diff --git a/Engine/Source/Runtime/GameplayAbilities/Private/Abilities/Tasks/AbilityTask_WaitTargetData.cpp b/Engine/Source/Runtime/GameplayAbilities/Private/Abilities/Tasks/AbilityTask_WaitTargetData.cpp
--- a/Engine/Source/Runtime/GameplayAbilities/Private/Abilities/Tasks/AbilityTask_WaitTargetData.cpp
+++ b/Engine/Source/Runtime/GameplayAbilities/Private/Abilities/Tasks/AbilityTask_WaitTargetData.cpp
@@ -147,7 +147,6 @@ void UAbilityTask_WaitTargetData::FinalizeTargetActor(AGameplayAbilityTargetActo
 	check(Ability.IsValid());
 
 	// User ability activation is inhibited while this is active
-	AbilitySystemComponent->SetUserAbilityActivationInhibited(true);
 	AbilitySystemComponent->SpawnedTargetActors.Push(SpawnedActor);
 
 	SpawnedActor->StartTargeting(Ability.Get());
@@ -260,7 +259,7 @@ void UAbilityTask_WaitTargetData::OnTargetDataReadyCallback(FGameplayAbilityTarg
 		else if (ConfirmationType == EGameplayTargetingConfirmation::UserConfirmed)
 		{
 			// We aren't going to send the target data, but we will send a generic confirmed message.
-			AbilitySystemComponent->ServerSetReplicatedClientEvent(EAbilityReplicatedClientEvent::GenericConfirm, GetAbilitySpecHandle(), GetActivationPredictionKey(), AbilitySystemComponent->ScopedPredictionKey);
+			AbilitySystemComponent->ServerSetReplicatedEvent(EAbilityGenericReplicatedEvent::GenericConfirm, GetAbilitySpecHandle(), GetActivationPredictionKey(), AbilitySystemComponent->ScopedPredictionKey);
 		}
 	}
 
@@ -279,7 +278,15 @@ void UAbilityTask_WaitTargetData::OnTargetDataCancelledCallback(FGameplayAbility
 
 	if (IsPredictingClient())
 	{
-		AbilitySystemComponent->ServerSetReplicatedTargetDataCancelled(GetAbilitySpecHandle(), GetActivationPredictionKey(), AbilitySystemComponent->ScopedPredictionKey );
+		if (!TargetActor->ShouldProduceTargetDataOnServer)
+		{
+			AbilitySystemComponent->ServerSetReplicatedTargetDataCancelled(GetAbilitySpecHandle(), GetActivationPredictionKey(), AbilitySystemComponent->ScopedPredictionKey );
+		}
+		else
+		{
+			// We aren't going to send the target data, but we will send a generic confirmed message.
+			AbilitySystemComponent->ServerSetReplicatedEvent(EAbilityGenericReplicatedEvent::GenericCancel, GetAbilitySpecHandle(), GetActivationPredictionKey(), AbilitySystemComponent->ScopedPredictionKey);
+		}
 	}
 	Cancelled.Broadcast(Data);
 	EndTask();
@@ -310,8 +317,6 @@ void UAbilityTask_WaitTargetData::ExternalCancel()
 
 void UAbilityTask_WaitTargetData::OnDestroy(bool AbilityEnded)
 {
-	AbilitySystemComponent->SetUserAbilityActivationInhibited(false);
-
 	if (TargetActor.IsValid())
 	{
 		TargetActor->Destroy();
@@ -337,4 +342,3 @@ bool UAbilityTask_WaitTargetData::ShouldReplicateDataToServer() const
 
 
 // --------------------------------------------------------------------------------------
-
diff --git a/Engine/Source/Runtime/GameplayAbilities/Private/AbilitySystemBlueprintLibrary.cpp b/Engine/Source/Runtime/GameplayAbilities/Private/AbilitySystemBlueprintLibrary.cpp
--- a/Engine/Source/Runtime/GameplayAbilities/Private/AbilitySystemBlueprintLibrary.cpp
+++ b/Engine/Source/Runtime/GameplayAbilities/Private/AbilitySystemBlueprintLibrary.cpp
@@ -566,4 +566,43 @@ FGameplayEffectSpecHandle UAbilitySystemBlueprintLibrary::AddLinkedGameplayEffec
 	}
 
 	return SpecHandle;
+}
+
+
+FGameplayEffectSpecHandle UAbilitySystemBlueprintLibrary::SetStackCount(FGameplayEffectSpecHandle SpecHandle, int32 StackCount)
+{
+	FGameplayEffectSpec* Spec = SpecHandle.Data.Get();
+	if (Spec)
+	{
+		Spec->StackCount = StackCount;
+	}
+	else
+	{
+		ABILITY_LOG(Warning, TEXT("UAbilitySystemBlueprintLibrary::AddLinkedGameplayEffectSpec called with invalid SpecHandle"));
+	}
+	return SpecHandle;
+}
+	
+FGameplayEffectSpecHandle UAbilitySystemBlueprintLibrary::SetStackCountToMax(FGameplayEffectSpecHandle SpecHandle)
+{
+	FGameplayEffectSpec* Spec = SpecHandle.Data.Get();
+	if (Spec && Spec->Def)
+	{
+		Spec->StackCount = Spec->Def->StackLimitCount;
+	}
+	else
+	{
+		ABILITY_LOG(Warning, TEXT("UAbilitySystemBlueprintLibrary::AddLinkedGameplayEffectSpec called with invalid SpecHandle"));
+	}
+	return SpecHandle;
+}
+
+int32 UAbilitySystemBlueprintLibrary::GetActiveGameplayEffectStackCount(FActiveGameplayEffectHandle ActiveHandle)
+{
+	UAbilitySystemComponent* ASC = ActiveHandle.GetOwningAbilitySystemComponent();
+	if (ASC)
+	{
+		return ASC->GetCurrentStackCount(ActiveHandle);
+	}
+	return 0;
 }
\ No newline at end of file
diff --git a/Engine/Source/Runtime/GameplayAbilities/Private/AbilitySystemComponent.cpp b/Engine/Source/Runtime/GameplayAbilities/Private/AbilitySystemComponent.cpp
--- a/Engine/Source/Runtime/GameplayAbilities/Private/AbilitySystemComponent.cpp
+++ b/Engine/Source/Runtime/GameplayAbilities/Private/AbilitySystemComponent.cpp
@@ -46,6 +46,9 @@ UAbilitySystemComponent::UAbilitySystemComponent(const FObjectInitializer& Objec
 	bReplicates = true;
 
 	UserAbilityActivationInhibited = false;
+
+	GenericConfirmInputID = INDEX_NONE;
+	GenericCancelInputID = INDEX_NONE;
 }
 
 UAbilitySystemComponent::~UAbilitySystemComponent()
@@ -503,6 +506,8 @@ FActiveGameplayEffectHandle UAbilitySystemComponent::ApplyGameplayEffectSpecToSe
 	// of the add operation
 	FScopedActiveGameplayEffectLock ScopeLock(ActiveGameplayEffects);
 
+	const bool bIsNetAuthority = IsOwnerActorAuthoritative();
+
 	// Check Network Authority
 	if (!HasNetworkAuthorityToApplyGameplayEffect(PredictionKey))
 	{
@@ -660,7 +665,7 @@ FActiveGameplayEffectHandle UAbilitySystemComponent::ApplyGameplayEffectSpecToSe
 	//		Remove any active gameplay effects that match the RemoveGameplayEffectsWithTags in the definition for this spec
 	//		Only call this if we are the Authoritative owner and we have some RemoveGameplayEffectsWithTags.CombinedTag to remove
 	// ------------------------------------------------------
-	if (IsOwnerActorAuthoritative() && Spec.Def->RemoveGameplayEffectsWithTags.CombinedTags.Num() > 0)
+	if (bIsNetAuthority && Spec.Def->RemoveGameplayEffectsWithTags.CombinedTags.Num() > 0)
 	{
 		// Clear tags is always removing all stacks.
 		FActiveGameplayEffectQuery ClearQuery(&Spec.Def->RemoveGameplayEffectsWithTags.CombinedTags);
@@ -670,8 +675,34 @@ FActiveGameplayEffectHandle UAbilitySystemComponent::ApplyGameplayEffectSpecToSe
 		}
 		ActiveGameplayEffects.RemoveActiveEffects(ClearQuery, -1);
 	}
+	
+
+	// ------------------------------------------------------
+	//	Apply Granted Abilities
+	//	
+	//	Note: Doing this before apply TargetEffectSpecs, but we could just as easily apply after.
+	//	Hedging bet for now that we are more likely to want to have 'passive ability that reacts to linked GE'
+	//	over 'need ability to activate after linked GE is applied'.
+	//	
+	//	Note2: This is allowing instant GEs to permanently grant abilities. This could be disallowed if needed.
+	// ------------------------------------------------------
+	if (bIsNetAuthority)
+	{
+		for (FGameplayAbilitySpecDef& AbilitySpecDef : OurCopyOfSpec->GrantedAbilitySpecs)
+		{
+			// Only do this is we havent assigned the ability yet! This prevents cases where stacking GEs
+			// would regrant the ability every time the stack was applied
+			if (AbilitySpecDef.AssignedHandle.IsValid() == false)
+			{
+				GiveAbility( FGameplayAbilitySpec(AbilitySpecDef, MyHandle) );
+			}
+		}	
+	}
 
+	// ------------------------------------------------------
+	// Apply Linked effects
 	// todo: this is ignoring the returned handles, should we put them into a TArray and return all of the handles?
+	// ------------------------------------------------------
 	for (const FGameplayEffectSpecHandle TargetSpec: Spec.TargetEffectSpecs)
 	{
 		if (TargetSpec.IsValid())
@@ -737,6 +768,16 @@ float UAbilitySystemComponent::GetGameplayEffectMagnitude(FActiveGameplayEffectH
 	return ActiveGameplayEffects.GetGameplayEffectMagnitude(Handle, Attribute);
 }
 
+int32 UAbilitySystemComponent::GetCurrentStackCount(FActiveGameplayEffectHandle Handle) const
+{
+	const FActiveGameplayEffect* ActiveGE = ActiveGameplayEffects.GetActiveGameplayEffect(Handle);
+	if (ActiveGE)
+	{
+		return ActiveGE->Spec.StackCount;
+	}
+	return 0;
+}
+
 void UAbilitySystemComponent::InvokeGameplayCueEvent(const FGameplayEffectSpecForRPC &Spec, EGameplayCueEvent::Type EventType)
 {
 	AActor* ActorAvatar = AbilityActorInfo->AvatarActor.Get();
@@ -1139,6 +1180,11 @@ void UAbilitySystemComponent::OnGameplayEffectAppliedToSelf(UAbilitySystemCompon
 	OnGameplayEffectAppliedDelegateToSelf.Broadcast(Source, SpecApplied, ActiveHandle);
 }
 
+TArray<TWeakObjectPtr<UAbilityTask> >&	UAbilitySystemComponent::GetAbilityActiveTasks(UGameplayAbility* Ability)
+{
+	return Ability->ActiveTasks;
+}
+
 // ------------------------------------------------------------------------
 
 FString ASC_CleanupName(FString Str)
@@ -1171,7 +1217,7 @@ void UAbilitySystemComponent::DisplayDebug(class UCanvas* Canvas, const class FD
 
 	if (BlockedAbilityTags.GetExplicitGameplayTags().Num() > 0)
 	{
-		YPos += Canvas->DrawText(GEngine->GetTinyFont(), FString::Printf(TEXT("BlockedAbilityTags: %s"), *BlockedAbilityTags.GetExplicitGameplayTags().ToStringSimple()), 4.f, YPos);
+		YL = Canvas->DrawText(GEngine->GetTinyFont(), FString::Printf(TEXT("BlockedAbilityTags: %s"), *BlockedAbilityTags.GetExplicitGameplayTags().ToStringSimple()), 4.f, YPos);
 		YPos += YL;
 	}
 
diff --git a/Engine/Source/Runtime/GameplayAbilities/Private/AbilitySystemComponent_Abilities.cpp b/Engine/Source/Runtime/GameplayAbilities/Private/AbilitySystemComponent_Abilities.cpp
--- a/Engine/Source/Runtime/GameplayAbilities/Private/AbilitySystemComponent_Abilities.cpp
+++ b/Engine/Source/Runtime/GameplayAbilities/Private/AbilitySystemComponent_Abilities.cpp
@@ -155,6 +155,7 @@ void UAbilitySystemComponent::InitAbilityActorInfo(AActor* InOwnerActor, AActor*
 
 	if (AvatarChanged)
 	{
+		ABILITYLIST_SCOPE_LOCK();
 		for (FGameplayAbilitySpec& Spec : ActivatableAbilities.Items)
 		{
 			if (Spec.Ability)
@@ -230,6 +231,13 @@ FGameplayAbilitySpecHandle UAbilitySystemComponent::GiveAbility(FGameplayAbility
 {	
 	check(Spec.Ability);
 	check(IsOwnerActorAuthoritative());	// Should be called on authority
+
+	// If locked, add to pending list. The Spec.Handle is not regenerated when we receive, so returning this is ok.
+	if (AbilityScopeLockCount > 0)
+	{
+		AbilityPendingAdds.Add(Spec);
+		return Spec.Handle;
+	}
 	
 	FGameplayAbilitySpec& OwnedSpec = ActivatableAbilities.Items[ActivatableAbilities.Items.Add(Spec)];
 	
@@ -283,9 +291,27 @@ FGameplayAbilitySpecHandle UAbilitySystemComponent::GiveAbilityAndActivateOnce(F
 	return AddedAbilityHandle;
 }
 
+void UAbilitySystemComponent::SetRemoveAbilityOnEnd(FGameplayAbilitySpecHandle AbilitySpecHandle)
+{
+	FGameplayAbilitySpec* FoundSpec = FindAbilitySpecFromHandle(AbilitySpecHandle);
+	if (FoundSpec)
+	{
+		if (FoundSpec->IsActive())
+		{
+			FoundSpec->RemoveAfterActivation = true;
+			FoundSpec->InputID = INDEX_NONE;
+		}
+		else
+		{
+			ClearAbility(AbilitySpecHandle);
+		}
+	}
+}
+
 void UAbilitySystemComponent::ClearAllAbilities()
 {
 	check(IsOwnerActorAuthoritative());	// Should be called on authority
+	check(AbilityScopeLockCount == 0);	// We should never be calling this from a scoped lock situation.
 
 	// Note we aren't marking any old abilities pending kill. This shouldn't matter since they will be garbage collected.
 	for (FGameplayAbilitySpec& Spec : ActivatableAbilities.Items)
@@ -308,11 +334,22 @@ void UAbilitySystemComponent::ClearAbility(const FGameplayAbilitySpecHandle& Han
 		check(ActivatableAbilities.Items[Idx].Handle.IsValid());
 		if (ActivatableAbilities.Items[Idx].Handle == Handle)
 		{
-			OnRemoveAbility(ActivatableAbilities.Items[Idx]);
-			ActivatableAbilities.Items.RemoveAtSwap(Idx);
-			ActivatableAbilities.MarkArrayDirty();
+			if (AbilityScopeLockCount > 0)
+			{
+				if (ActivatableAbilities.Items[Idx].PendingRemove == false)
+				{
+					ActivatableAbilities.Items[Idx].PendingRemove = true;
+					AbilityPendingRemoves.Add(Handle);
+				}
+			}
+			else
+			{
 
-			CheckForClearedAbilities();
+				OnRemoveAbility(ActivatableAbilities.Items[Idx]);
+				ActivatableAbilities.Items.RemoveAtSwap(Idx);
+				ActivatableAbilities.MarkArrayDirty();
+				CheckForClearedAbilities();
+			}
 			return;
 		}
 	}
@@ -459,6 +496,28 @@ void UAbilitySystemComponent::CheckForClearedAbilities()
 	}
 }
 
+void UAbilitySystemComponent::IncrementAbilityListLock()
+{
+	AbilityScopeLockCount++;
+}
+void UAbilitySystemComponent::DecrementAbilityListLock()
+{
+	if (--AbilityScopeLockCount == 0)
+	{
+		TArray<FGameplayAbilitySpec, TInlineAllocator<2> > LocalPendingAdds(MoveTemp(AbilityPendingAdds));
+		for (FGameplayAbilitySpec& Spec : LocalPendingAdds)
+		{
+			GiveAbility(Spec);
+		}
+
+		TArray<FGameplayAbilitySpecHandle, TInlineAllocator<2> > LocalPendingRemoves(MoveTemp(AbilityPendingRemoves));
+		for (FGameplayAbilitySpecHandle& Handle : LocalPendingRemoves)
+		{
+			ClearAbility(Handle);
+		}
+	}
+}
+
 const TArray<FGameplayAbilitySpec>& UAbilitySystemComponent::GetActivatableAbilities() const
 {
 	return ActivatableAbilities.Items;
@@ -481,7 +540,15 @@ FGameplayAbilitySpec* UAbilitySystemComponent::FindAbilitySpecFromHandle(FGamepl
 
 void UAbilitySystemComponent::MarkAbilitySpecDirty(FGameplayAbilitySpec& Spec)
 {
-	ActivatableAbilities.MarkItemDirty(Spec);
+	if (IsOwnerActorAuthoritative())
+	{
+		ActivatableAbilities.MarkItemDirty(Spec);
+	}
+	else
+	{
+		// Clients predicting should call MarkArrayDirty to force the internal replication map to be rebuilt.
+		ActivatableAbilities.MarkArrayDirty();
+	}
 }
 
 FGameplayAbilitySpec* UAbilitySystemComponent::FindAbilitySpecFromInputID(int32 InputID)
@@ -526,8 +593,16 @@ UGameplayAbility* UAbilitySystemComponent::CreateNewInstanceOfAbility(FGameplayA
 
 void UAbilitySystemComponent::NotifyAbilityEnded(FGameplayAbilitySpecHandle Handle, UGameplayAbility* Ability)
 {
+	check(Ability);
 	FGameplayAbilitySpec* Spec = FindAbilitySpecFromHandle(Handle);
+	if (Spec == nullptr)
+	{
+		// The ability spec may have been removed while we were ending. We can assume everything was cleaned up if the spec isnt here.
+		return;
+	}
 	check(Spec);
+	check(Ability);
+
 	check(Ability);
 	ENetRole OwnerRole = GetOwnerRole();
 
@@ -584,6 +659,7 @@ void UAbilitySystemComponent::NotifyAbilityEnded(FGameplayAbilitySpecHandle Hand
 
 void UAbilitySystemComponent::CancelAbility(UGameplayAbility* Ability)
 {
+	ABILITYLIST_SCOPE_LOCK();
 	for (FGameplayAbilitySpec& Spec : ActivatableAbilities.Items)
 	{
 		if (Spec.Ability == Ability)
@@ -595,6 +671,7 @@ void UAbilitySystemComponent::CancelAbility(UGameplayAbility* Ability)
 
 void UAbilitySystemComponent::CancelAbilities(const FGameplayTagContainer* WithTags, const FGameplayTagContainer* WithoutTags, UGameplayAbility* Ignore)
 {
+	ABILITYLIST_SCOPE_LOCK();
 	for (FGameplayAbilitySpec& Spec : ActivatableAbilities.Items)
 	{
 		if (!Spec.IsActive() || Spec.Ability == nullptr)
@@ -624,20 +701,21 @@ void UAbilitySystemComponent::CancelAbilitySpec(FGameplayAbilitySpec& Spec, UGam
 		{
 			if (InstanceAbility && Ignore != InstanceAbility)
 			{
-				InstanceAbility->CancelAbility(Spec.Handle, ActorInfo, InstanceAbility->GetCurrentActivationInfo());
+				InstanceAbility->CancelAbility(Spec.Handle, ActorInfo, InstanceAbility->GetCurrentActivationInfo(), true);
 			}
 		}
 	}
 	else
 	{
 		// Try to cancel the non instanced, this may not necessarily work
-		Spec.Ability->CancelAbility(Spec.Handle, ActorInfo, FGameplayAbilityActivationInfo());
+		Spec.Ability->CancelAbility(Spec.Handle, ActorInfo, FGameplayAbilityActivationInfo(), true);
 	}
 	MarkAbilitySpecDirty(Spec);
 }
 
 void UAbilitySystemComponent::CancelAllAbilities(UGameplayAbility* Ignore)
 {
+	ABILITYLIST_SCOPE_LOCK();
 	for (FGameplayAbilitySpec& Spec : ActivatableAbilities.Items)
 	{
 		if (Spec.Ability && Spec.Ability->IsActive())
@@ -938,6 +1016,13 @@ bool UAbilitySystemComponent::TryActivateAbility(FGameplayAbilitySpecHandle Hand
 	}
 	else if (Ability->GetNetExecutionPolicy() == EGameplayAbilityNetExecutionPolicy::ServerInitiated)
 	{
+		if (NetMode != ROLE_Authority)
+		{
+			// Client hasn't technically activated the ability yet, so he shouldn't have advanced the active count.
+			// (this is ugly but the only case where this happens, and we want the ActiveCount to increment prior to activating the ability in the other case).
+			Spec->ActiveCount--;
+		}
+
 		// Tell the server to start the ability, this will return to the client 
 		if (TriggerEventData)
 		{
@@ -1048,7 +1133,7 @@ void UAbilitySystemComponent::InternalServerTryActiveAbility(FGameplayAbilitySpe
 	}
 
 	// Consume any pending target info, to clear out cancels from old executions
-	ConsumeAllClientReplicatedData(Handle, PredictionKey);
+	ConsumeAllReplicatedData(Handle, PredictionKey);
 
 	FScopedPredictionWindow ScopedPredictionWindow(this, PredictionKey);
 
@@ -1126,7 +1211,7 @@ void UAbilitySystemComponent::InternalServerTryActiveAbility(FGameplayAbilitySpe
 	MarkAbilitySpecDirty(*Spec);
 }
 
-void UAbilitySystemComponent::ReplicateEndAbility(FGameplayAbilitySpecHandle Handle, FGameplayAbilityActivationInfo ActivationInfo, UGameplayAbility* Ability)
+void UAbilitySystemComponent::ReplicateEndOrCancelAbility(FGameplayAbilitySpecHandle Handle, FGameplayAbilityActivationInfo ActivationInfo, UGameplayAbility* Ability, bool bWasCanceled)
 {
 	if (Ability->NetExecutionPolicy == EGameplayAbilityNetExecutionPolicy::LocalPredicted || Ability->NetExecutionPolicy == EGameplayAbilityNetExecutionPolicy::ServerInitiated)
 	{
@@ -1135,20 +1220,34 @@ void UAbilitySystemComponent::ReplicateEndAbility(FGameplayAbilitySpecHandle Han
 		{
 			if (!AbilityActorInfo->IsLocallyControlled())
 			{
-				// Only tell the client about the end ability if we're not the local controller
-				ClientEndAbility(Handle, ActivationInfo);
+				// Only tell the client about the end/cancel ability if we're not the local controller
+				if (bWasCanceled)
+				{
+					ClientCancelAbility(Handle, ActivationInfo);
+				}
+				else
+				{
+					ClientEndAbility(Handle, ActivationInfo);
+				}
 			}
 		}
 		else
 		{
 			// This passes up the current prediction key if we have one
-			ServerEndAbility(Handle, ActivationInfo, ScopedPredictionKey);
+			if (bWasCanceled)
+			{
+				ServerCancelAbility(Handle, ActivationInfo);
+			}
+			else
+			{
+				ServerEndAbility(Handle, ActivationInfo, ScopedPredictionKey);
+			}
 		}
 	}
 }
 
-//This is only called when ending an ability in response to a remote instruction.
-void UAbilitySystemComponent::RemoteEndAbility(FGameplayAbilitySpecHandle AbilityToEnd, FGameplayAbilityActivationInfo ActivationInfo)
+//This is only called when ending or canceling an ability in response to a remote instruction.
+void UAbilitySystemComponent::RemoteEndOrCancelAbility(FGameplayAbilitySpecHandle AbilityToEnd, FGameplayAbilityActivationInfo ActivationInfo, bool bWasCanceled)
 {
 	FGameplayAbilitySpec* AbilitySpec = FindAbilitySpecFromHandle(AbilityToEnd);
 	if (AbilitySpec && AbilitySpec->Ability && AbilitySpec->IsActive())
@@ -1160,8 +1259,17 @@ void UAbilitySystemComponent::RemoteEndAbility(FGameplayAbilitySpecHandle Abilit
 			// Check if the ability is the same prediction key (can both by 0) and has been confirmed. If so cancel it.
 			if (Instance->CurrentActivationInfo.GetActivationPredictionKey() == ActivationInfo.GetActivationPredictionKey() && Instance->CurrentActivationInfo.bCanBeEndedByOtherInstance)
 			{
-				// End the ability but don't replicate it back to whoever called us
-				Instance->EndAbility(Instance->CurrentSpecHandle, Instance->CurrentActorInfo, Instance->CurrentActivationInfo, false);
+				// End/Cancel the ability but don't replicate it back to whoever called us
+				if (bWasCanceled)
+				{
+					// Since this was a remote cancel, we should force it through. We do not support 'server says ability was cancelled but client disagrees that it can be'.
+					Instance->SetCanBeCanceled(true);
+					Instance->CancelAbility(Instance->CurrentSpecHandle, Instance->CurrentActorInfo, Instance->CurrentActivationInfo, false);
+				}
+				else
+				{
+					Instance->EndAbility(Instance->CurrentSpecHandle, Instance->CurrentActorInfo, Instance->CurrentActivationInfo, false);
+				}
 			}
 		}
 	}
@@ -1170,8 +1278,8 @@ void UAbilitySystemComponent::RemoteEndAbility(FGameplayAbilitySpecHandle Abilit
 void UAbilitySystemComponent::ServerEndAbility_Implementation(FGameplayAbilitySpecHandle AbilityToEnd, FGameplayAbilityActivationInfo ActivationInfo, FPredictionKey PredictionKey)
 {
 	FScopedPredictionWindow ScopedPrediction(this, PredictionKey);
-
-	RemoteEndAbility(AbilityToEnd, ActivationInfo);
+	
+	RemoteEndOrCancelAbility(AbilityToEnd, ActivationInfo, false);
 }
 
 bool UAbilitySystemComponent::ServerEndAbility_Validate(FGameplayAbilitySpecHandle AbilityToEnd, FGameplayAbilityActivationInfo ActivationInfo, FPredictionKey PredictionKey)
@@ -1181,7 +1289,22 @@ bool UAbilitySystemComponent::ServerEndAbility_Validate(FGameplayAbilitySpecHand
 
 void UAbilitySystemComponent::ClientEndAbility_Implementation(FGameplayAbilitySpecHandle AbilityToEnd, FGameplayAbilityActivationInfo ActivationInfo)
 {
-	RemoteEndAbility(AbilityToEnd, ActivationInfo);
+	RemoteEndOrCancelAbility(AbilityToEnd, ActivationInfo, false);
+}
+
+void UAbilitySystemComponent::ServerCancelAbility_Implementation(FGameplayAbilitySpecHandle AbilityToCancel, FGameplayAbilityActivationInfo ActivationInfo)
+{
+	RemoteEndOrCancelAbility(AbilityToCancel, ActivationInfo, true);
+}
+
+bool UAbilitySystemComponent::ServerCancelAbility_Validate(FGameplayAbilitySpecHandle AbilityToCancel, FGameplayAbilityActivationInfo ActivationInfo)
+{
+	return true;
+}
+
+void UAbilitySystemComponent::ClientCancelAbility_Implementation(FGameplayAbilitySpecHandle AbilityToCancel, FGameplayAbilityActivationInfo ActivationInfo)
+{
+	RemoteEndOrCancelAbility(AbilityToCancel, ActivationInfo, true);
 }
 
 static_assert(sizeof(int16) == sizeof(FPredictionKey::KeyType), "Sizeof PredictionKey::KeyType does not match RPC parameters in AbilitySystemComponent ClientActivateAbilityFailed_Implementation");
@@ -1569,22 +1692,48 @@ void UAbilitySystemComponent::BindAbilityActivationToInputComponent(UInputCompon
 	}
 
 	// Bind Confirm/Cancel. Note: these have to come last!
+	if (BindInfo.ConfirmTargetCommand.IsEmpty() == false)
 	{
 		FInputActionBinding AB(FName(*BindInfo.ConfirmTargetCommand), IE_Pressed);
 		AB.ActionDelegate.GetDelegateForManualSet().BindUObject(this, &UAbilitySystemComponent::LocalInputConfirm);
 		InputComponent->AddActionBinding(AB);
 	}
 	
+	if (BindInfo.CancelTargetCommand.IsEmpty() == false)
 	{
 		FInputActionBinding AB(FName(*BindInfo.CancelTargetCommand), IE_Pressed);
 		AB.ActionDelegate.GetDelegateForManualSet().BindUObject(this, &UAbilitySystemComponent::LocalInputCancel);
 		InputComponent->AddActionBinding(AB);
 	}
 
+	if (BindInfo.CancelTargetInputID >= 0)
+	{
+		GenericCancelInputID = BindInfo.CancelTargetInputID;
+	}
+	if (BindInfo.ConfirmTargetInputID >= 0)
+	{
+		GenericConfirmInputID = BindInfo.ConfirmTargetInputID;
+	}
 }
 
 void UAbilitySystemComponent::AbilityLocalInputPressed(int32 InputID)
 {
+	// Consume the input if this InputID is overloaded with GenericConfirm/Cancel and the GenericConfim/Cancel callback is bound
+	if (InputID == GenericConfirmInputID && GenericLocalConfirmCallbacks.IsBound())
+	{
+		LocalInputConfirm();
+		return;
+	}
+
+	if (InputID == GenericCancelInputID && GenericLocalCancelCallbacks.IsBound())
+	{
+		LocalInputCancel();
+		return;
+	}
+
+	// ---------------------------------------------------------
+
+	ABILITYLIST_SCOPE_LOCK();
 	for (FGameplayAbilitySpec& Spec : ActivatableAbilities.Items)
 	{
 		if (Spec.InputID == InputID)
@@ -1602,15 +1751,12 @@ void UAbilitySystemComponent::AbilityLocalInputPressed(int32 InputID)
 					AbilitySpecInputPressed(Spec);
 
 					// Invoke the InputPressed event. This is not replicated here. If someone is listening, they may replicate the InputPressed event to the server.
-					InvokeReplicatedClientEvent(EAbilityReplicatedClientEvent::InputPressed, Spec.Handle, Spec.ActivationInfo.GetActivationPredictionKey());					
+					InvokeReplicatedEvent(EAbilityGenericReplicatedEvent::InputPressed, Spec.Handle, Spec.ActivationInfo.GetActivationPredictionKey());					
 				}
 				else
 				{
 					// Ability is not active, so try to activate it
-					if (!TryActivateAbility(Spec.Handle))
-					{
-						Spec.InputPressed = false;
-					}
+					TryActivateAbility(Spec.Handle);
 				}
 			}
 		}
@@ -1619,6 +1765,7 @@ void UAbilitySystemComponent::AbilityLocalInputPressed(int32 InputID)
 
 void UAbilitySystemComponent::AbilityLocalInputReleased(int32 InputID)
 {
+	ABILITYLIST_SCOPE_LOCK();
 	for (FGameplayAbilitySpec& Spec : ActivatableAbilities.Items)
 	{
 		if (Spec.InputID == InputID)
@@ -1633,7 +1780,7 @@ void UAbilitySystemComponent::AbilityLocalInputReleased(int32 InputID)
 
 				AbilitySpecInputReleased(Spec);
 				
-				InvokeReplicatedClientEvent(EAbilityReplicatedClientEvent::InputReleased, Spec.Handle, Spec.ActivationInfo.GetActivationPredictionKey());
+				InvokeReplicatedEvent(EAbilityGenericReplicatedEvent::InputReleased, Spec.Handle, Spec.ActivationInfo.GetActivationPredictionKey());
 			}
 		}
 	}
@@ -1716,6 +1863,7 @@ void UAbilitySystemComponent::TryActivateInputHeldAbilities()
 
 	if (PC && PC->PlayerInput)
 	{
+		ABILITYLIST_SCOPE_LOCK();
 		for (FGameplayAbilitySpec& Spec : ActivatableAbilities.Items)
 		{
 			if ((Spec.InputID == INDEX_NONE) || !Spec.InputPressed || Spec.IsActive() || !Spec.Ability)
@@ -1734,11 +1882,13 @@ void UAbilitySystemComponent::TryActivateInputHeldAbilities()
 void UAbilitySystemComponent::LocalInputConfirm()
 {
 	GenericLocalConfirmCallbacks.Broadcast();
+	GenericLocalConfirmCallbacks.Clear();
 }
 
 void UAbilitySystemComponent::LocalInputCancel()
 {	
 	GenericLocalCancelCallbacks.Broadcast();
+	GenericLocalCancelCallbacks.Clear();
 }
 
 void UAbilitySystemComponent::TargetConfirm()
@@ -2251,7 +2401,7 @@ UGameplayAbility* UAbilitySystemComponent::GetAnimatingAbility()
 
 void UAbilitySystemComponent::ConfirmAbilityTargetData(FGameplayAbilitySpecHandle AbilityHandle, FPredictionKey AbilityOriginalPredictionKey, const FGameplayAbilityTargetDataHandle& TargetData, const FGameplayTag& ApplicationTag)
 {
-	FAbilityClientDataCache* CachedData = AbilityClientDataMap.Find(FGameplayAbilitySpecHandleAndPredictionKey(AbilityHandle, AbilityOriginalPredictionKey));
+	FAbilityReplicatedDataCache* CachedData = AbilityTargetDataMap.Find(FGameplayAbilitySpecHandleAndPredictionKey(AbilityHandle, AbilityOriginalPredictionKey));
 	if (CachedData)
 	{
 		CachedData->TargetSetDelegate.Broadcast(TargetData, ApplicationTag);
@@ -2260,7 +2410,7 @@ void UAbilitySystemComponent::ConfirmAbilityTargetData(FGameplayAbilitySpecHandl
 
 void UAbilitySystemComponent::CancelAbilityTargetData(FGameplayAbilitySpecHandle AbilityHandle, FPredictionKey AbilityOriginalPredictionKey)
 {
-	FAbilityClientDataCache* CachedData = AbilityClientDataMap.Find(FGameplayAbilitySpecHandleAndPredictionKey(AbilityHandle, AbilityOriginalPredictionKey));
+	FAbilityReplicatedDataCache* CachedData = AbilityTargetDataMap.Find(FGameplayAbilitySpecHandleAndPredictionKey(AbilityHandle, AbilityOriginalPredictionKey));
 	if (CachedData)
 	{
 		CachedData->Reset();
@@ -2268,9 +2418,9 @@ void UAbilitySystemComponent::CancelAbilityTargetData(FGameplayAbilitySpecHandle
 	}
 }
 
-void UAbilitySystemComponent::ConsumeAllClientReplicatedData(FGameplayAbilitySpecHandle AbilityHandle, FPredictionKey AbilityOriginalPredictionKey)
+void UAbilitySystemComponent::ConsumeAllReplicatedData(FGameplayAbilitySpecHandle AbilityHandle, FPredictionKey AbilityOriginalPredictionKey)
 {
-	FAbilityClientDataCache* CachedData = AbilityClientDataMap.Find(FGameplayAbilitySpecHandleAndPredictionKey(AbilityHandle, AbilityOriginalPredictionKey));
+	FAbilityReplicatedDataCache* CachedData = AbilityTargetDataMap.Find(FGameplayAbilitySpecHandleAndPredictionKey(AbilityHandle, AbilityOriginalPredictionKey));
 	if (CachedData)
 	{
 		CachedData->Reset();
@@ -2279,7 +2429,7 @@ void UAbilitySystemComponent::ConsumeAllClientReplicatedData(FGameplayAbilitySpe
 
 void UAbilitySystemComponent::ConsumeClientReplicatedTargetData(FGameplayAbilitySpecHandle AbilityHandle, FPredictionKey AbilityOriginalPredictionKey)
 {
-	FAbilityClientDataCache* CachedData = AbilityClientDataMap.Find(FGameplayAbilitySpecHandleAndPredictionKey(AbilityHandle, AbilityOriginalPredictionKey));
+	FAbilityReplicatedDataCache* CachedData = AbilityTargetDataMap.Find(FGameplayAbilitySpecHandleAndPredictionKey(AbilityHandle, AbilityOriginalPredictionKey));
 	if (CachedData)
 	{
 		CachedData->TargetData.Clear();
@@ -2288,9 +2438,9 @@ void UAbilitySystemComponent::ConsumeClientReplicatedTargetData(FGameplayAbility
 	}
 }
 
-void UAbilitySystemComponent::ConsumeClientReplicatedEvent(EAbilityReplicatedClientEvent::Type EventType, FGameplayAbilitySpecHandle AbilityHandle, FPredictionKey AbilityOriginalPredictionKey)
+void UAbilitySystemComponent::ConsumeGenericReplicatedEvent(EAbilityGenericReplicatedEvent::Type EventType, FGameplayAbilitySpecHandle AbilityHandle, FPredictionKey AbilityOriginalPredictionKey)
 {
-	FAbilityClientDataCache* CachedData = AbilityClientDataMap.Find(FGameplayAbilitySpecHandleAndPredictionKey(AbilityHandle, AbilityOriginalPredictionKey));
+	FAbilityReplicatedDataCache* CachedData = AbilityTargetDataMap.Find(FGameplayAbilitySpecHandleAndPredictionKey(AbilityHandle, AbilityOriginalPredictionKey));
 	if (CachedData)
 	{
 		CachedData->GenericEvents[EventType].bTriggered = false;
@@ -2299,23 +2449,32 @@ void UAbilitySystemComponent::ConsumeClientReplicatedEvent(EAbilityReplicatedCli
 
 // --------------------------------------------------------------------------
 
-void UAbilitySystemComponent::ServerSetReplicatedClientEvent_Implementation(EAbilityReplicatedClientEvent::Type EventType, FGameplayAbilitySpecHandle AbilityHandle, FPredictionKey AbilityOriginalPredictionKey,  FPredictionKey CurrentPredictionKey)
+void UAbilitySystemComponent::ServerSetReplicatedEvent_Implementation(EAbilityGenericReplicatedEvent::Type EventType, FGameplayAbilitySpecHandle AbilityHandle, FPredictionKey AbilityOriginalPredictionKey,  FPredictionKey CurrentPredictionKey)
 {
 	FScopedPredictionWindow ScopedPrediction(this, CurrentPredictionKey);
 
-	InvokeReplicatedClientEvent(EventType, AbilityHandle, AbilityOriginalPredictionKey);
+	InvokeReplicatedEvent(EventType, AbilityHandle, AbilityOriginalPredictionKey);
 }
 
-void UAbilitySystemComponent::InvokeReplicatedClientEvent(EAbilityReplicatedClientEvent::Type EventType, FGameplayAbilitySpecHandle AbilityHandle, FPredictionKey AbilityOriginalPredictionKey)
+bool UAbilitySystemComponent::InvokeReplicatedEvent(EAbilityGenericReplicatedEvent::Type EventType, FGameplayAbilitySpecHandle AbilityHandle, FPredictionKey AbilityOriginalPredictionKey)
 {
-	FAbilityClientDataCache& ReplicatedData = AbilityClientDataMap.FindOrAdd(FGameplayAbilitySpecHandleAndPredictionKey(AbilityHandle, AbilityOriginalPredictionKey));
+	FAbilityReplicatedDataCache& ReplicatedData = AbilityTargetDataMap.FindOrAdd(FGameplayAbilitySpecHandleAndPredictionKey(AbilityHandle, AbilityOriginalPredictionKey));
 	ReplicatedData.GenericEvents[(uint8)EventType].bTriggered = true;
-	ReplicatedData.GenericEvents[EventType].Delegate.Broadcast();
+
+	if (ReplicatedData.GenericEvents[EventType].Delegate.IsBound())
+	{
+		ReplicatedData.GenericEvents[EventType].Delegate.Broadcast();
+		return true;
+	}
+	else
+	{
+		return false;
+	}
 }
 
-bool UAbilitySystemComponent::ServerSetReplicatedClientEvent_Validate(EAbilityReplicatedClientEvent::Type EventType, FGameplayAbilitySpecHandle AbilityHandle, FPredictionKey AbilityOriginalPredictionKey,  FPredictionKey CurrentPredictionKey)
+bool UAbilitySystemComponent::ServerSetReplicatedEvent_Validate(EAbilityGenericReplicatedEvent::Type EventType, FGameplayAbilitySpecHandle AbilityHandle, FPredictionKey AbilityOriginalPredictionKey,  FPredictionKey CurrentPredictionKey)
 {
-	if (EventType >= EAbilityReplicatedClientEvent::MAX)
+	if (EventType >= EAbilityGenericReplicatedEvent::MAX)
 	{
 		return false;
 	}
@@ -2324,12 +2483,19 @@ bool UAbilitySystemComponent::ServerSetReplicatedClientEvent_Validate(EAbilityRe
 
 // -------
 
+void UAbilitySystemComponent::ClientSetReplicatedEvent_Implementation(EAbilityGenericReplicatedEvent::Type EventType, FGameplayAbilitySpecHandle AbilityHandle, FPredictionKey AbilityOriginalPredictionKey)
+{
+	InvokeReplicatedEvent(EventType, AbilityHandle, AbilityOriginalPredictionKey);
+}
+
+// -------
+
 void UAbilitySystemComponent::ServerSetReplicatedTargetData_Implementation(FGameplayAbilitySpecHandle AbilityHandle, FPredictionKey AbilityOriginalPredictionKey, FGameplayAbilityTargetDataHandle ReplicatedTargetData, FGameplayTag ApplicationTag, FPredictionKey CurrentPredictionKey)
 {
 	FScopedPredictionWindow ScopedPrediction(this, CurrentPredictionKey);
 
 	// Always adds to cache to store the new data
-	FAbilityClientDataCache& ReplicatedData = AbilityClientDataMap.FindOrAdd(FGameplayAbilitySpecHandleAndPredictionKey(AbilityHandle, AbilityOriginalPredictionKey));
+	FAbilityReplicatedDataCache& ReplicatedData = AbilityTargetDataMap.FindOrAdd(FGameplayAbilitySpecHandleAndPredictionKey(AbilityHandle, AbilityOriginalPredictionKey));
 
 	ReplicatedData.TargetData = ReplicatedTargetData;
 	ReplicatedData.ApplicationTag = ApplicationTag;
@@ -2350,7 +2516,7 @@ void UAbilitySystemComponent::ServerSetReplicatedTargetDataCancelled_Implementat
 	FScopedPredictionWindow ScopedPrediction(this, CurrentPredictionKey);
 
 	// Always adds to cache to store the new data
-	FAbilityClientDataCache& ReplicatedData = AbilityClientDataMap.FindOrAdd(FGameplayAbilitySpecHandleAndPredictionKey(AbilityHandle, AbilityOriginalPredictionKey));
+	FAbilityReplicatedDataCache& ReplicatedData = AbilityTargetDataMap.FindOrAdd(FGameplayAbilitySpecHandleAndPredictionKey(AbilityHandle, AbilityOriginalPredictionKey));
 
 	ReplicatedData.Reset();
 	ReplicatedData.bTargetCancelled = true;
@@ -2364,7 +2530,7 @@ bool UAbilitySystemComponent::ServerSetReplicatedTargetDataCancelled_Validate(FG
 
 void UAbilitySystemComponent::CallAllReplicatedDelegatesIfSet(FGameplayAbilitySpecHandle AbilityHandle, FPredictionKey AbilityOriginalPredictionKey)
 {
-	FAbilityClientDataCache* CachedData = AbilityClientDataMap.Find(FGameplayAbilitySpecHandleAndPredictionKey(AbilityHandle, AbilityOriginalPredictionKey));
+	FAbilityReplicatedDataCache* CachedData = AbilityTargetDataMap.Find(FGameplayAbilitySpecHandleAndPredictionKey(AbilityHandle, AbilityOriginalPredictionKey));
 	if (CachedData)
 	{
 		if (CachedData->bTargetConfirmed)
@@ -2376,7 +2542,7 @@ void UAbilitySystemComponent::CallAllReplicatedDelegatesIfSet(FGameplayAbilitySp
 			CachedData->TargetCancelledDelegate.Broadcast();
 		}
 
-		for (int32 idx=0; idx < EAbilityReplicatedClientEvent::MAX; ++idx)
+		for (int32 idx=0; idx < EAbilityGenericReplicatedEvent::MAX; ++idx)
 		{
 			if (CachedData->GenericEvents[idx].bTriggered)
 			{
@@ -2389,7 +2555,7 @@ void UAbilitySystemComponent::CallAllReplicatedDelegatesIfSet(FGameplayAbilitySp
 bool UAbilitySystemComponent::CallReplicatedTargetDataDelegatesIfSet(FGameplayAbilitySpecHandle AbilityHandle, FPredictionKey AbilityOriginalPredictionKey)
 {
 	bool CalledDelegate = false;
-	FAbilityClientDataCache* CachedData = AbilityClientDataMap.Find(FGameplayAbilitySpecHandleAndPredictionKey(AbilityHandle, AbilityOriginalPredictionKey));
+	FAbilityReplicatedDataCache* CachedData = AbilityTargetDataMap.Find(FGameplayAbilitySpecHandleAndPredictionKey(AbilityHandle, AbilityOriginalPredictionKey));
 	if (CachedData)
 	{
 		if (CachedData->bTargetConfirmed)
@@ -2407,9 +2573,9 @@ bool UAbilitySystemComponent::CallReplicatedTargetDataDelegatesIfSet(FGameplayAb
 	return CalledDelegate;
 }
 
-bool UAbilitySystemComponent::CallReplicatedEventDelegateIfSet(EAbilityReplicatedClientEvent::Type EventType, FGameplayAbilitySpecHandle AbilityHandle, FPredictionKey AbilityOriginalPredictionKey)
+bool UAbilitySystemComponent::CallReplicatedEventDelegateIfSet(EAbilityGenericReplicatedEvent::Type EventType, FGameplayAbilitySpecHandle AbilityHandle, FPredictionKey AbilityOriginalPredictionKey)
 {
-	FAbilityClientDataCache* CachedData = AbilityClientDataMap.Find(FGameplayAbilitySpecHandleAndPredictionKey(AbilityHandle, AbilityOriginalPredictionKey));
+	FAbilityReplicatedDataCache* CachedData = AbilityTargetDataMap.Find(FGameplayAbilitySpecHandleAndPredictionKey(AbilityHandle, AbilityOriginalPredictionKey));
 	if (CachedData && CachedData->GenericEvents[EventType].bTriggered)
 	{
 		// Already triggered, fire off delegate
@@ -2419,9 +2585,9 @@ bool UAbilitySystemComponent::CallReplicatedEventDelegateIfSet(EAbilityReplicate
 	return false;
 }
 
-bool UAbilitySystemComponent::CallOrAddReplicatedDelegate(EAbilityReplicatedClientEvent::Type EventType, FGameplayAbilitySpecHandle AbilityHandle, FPredictionKey AbilityOriginalPredictionKey, FSimpleMulticastDelegate::FDelegate Delegate)
+bool UAbilitySystemComponent::CallOrAddReplicatedDelegate(EAbilityGenericReplicatedEvent::Type EventType, FGameplayAbilitySpecHandle AbilityHandle, FPredictionKey AbilityOriginalPredictionKey, FSimpleMulticastDelegate::FDelegate Delegate)
 {
-	FAbilityClientDataCache& CachedData = AbilityClientDataMap.FindOrAdd(FGameplayAbilitySpecHandleAndPredictionKey(AbilityHandle, AbilityOriginalPredictionKey));
+	FAbilityReplicatedDataCache& CachedData = AbilityTargetDataMap.FindOrAdd(FGameplayAbilitySpecHandleAndPredictionKey(AbilityHandle, AbilityOriginalPredictionKey));
 	if (CachedData.GenericEvents[EventType].bTriggered)
 	{
 		// Already triggered, fire off delegate
@@ -2436,17 +2602,17 @@ bool UAbilitySystemComponent::CallOrAddReplicatedDelegate(EAbilityReplicatedClie
 
 FAbilityTargetDataSetDelegate& UAbilitySystemComponent::AbilityTargetDataSetDelegate(FGameplayAbilitySpecHandle AbilityHandle, FPredictionKey AbilityOriginalPredictionKey)
 {
-	return AbilityClientDataMap.FindOrAdd(FGameplayAbilitySpecHandleAndPredictionKey(AbilityHandle, AbilityOriginalPredictionKey)).TargetSetDelegate;
+	return AbilityTargetDataMap.FindOrAdd(FGameplayAbilitySpecHandleAndPredictionKey(AbilityHandle, AbilityOriginalPredictionKey)).TargetSetDelegate;
 }
 
 FSimpleMulticastDelegate& UAbilitySystemComponent::AbilityTargetDataCancelledDelegate(FGameplayAbilitySpecHandle AbilityHandle, FPredictionKey AbilityOriginalPredictionKey)
 {
-	return AbilityClientDataMap.FindOrAdd(FGameplayAbilitySpecHandleAndPredictionKey(AbilityHandle, AbilityOriginalPredictionKey)).TargetCancelledDelegate;
+	return AbilityTargetDataMap.FindOrAdd(FGameplayAbilitySpecHandleAndPredictionKey(AbilityHandle, AbilityOriginalPredictionKey)).TargetCancelledDelegate;
 }
 
-FSimpleMulticastDelegate& UAbilitySystemComponent::AbilityReplicatedEventDelegate(EAbilityReplicatedClientEvent::Type EventType, FGameplayAbilitySpecHandle AbilityHandle, FPredictionKey AbilityOriginalPredictionKey)
+FSimpleMulticastDelegate& UAbilitySystemComponent::AbilityReplicatedEventDelegate(EAbilityGenericReplicatedEvent::Type EventType, FGameplayAbilitySpecHandle AbilityHandle, FPredictionKey AbilityOriginalPredictionKey)
 {
-	return AbilityClientDataMap.FindOrAdd(FGameplayAbilitySpecHandleAndPredictionKey(AbilityHandle, AbilityOriginalPredictionKey)).GenericEvents[EventType].Delegate;
+	return AbilityTargetDataMap.FindOrAdd(FGameplayAbilitySpecHandleAndPredictionKey(AbilityHandle, AbilityOriginalPredictionKey)).GenericEvents[EventType].Delegate;
 }
 
 
diff --git a/Engine/Source/Runtime/GameplayAbilities/Private/GameplayAbilityTypes.cpp b/Engine/Source/Runtime/GameplayAbilities/Private/GameplayAbilityTypes.cpp
--- a/Engine/Source/Runtime/GameplayAbilities/Private/GameplayAbilityTypes.cpp
+++ b/Engine/Source/Runtime/GameplayAbilities/Private/GameplayAbilityTypes.cpp
@@ -158,4 +158,26 @@ void FGameplayAbilitySpec::PostReplicatedAdd(const struct FGameplayAbilitySpecCo
 void FGameplayAbilitySpecContainer::RegisterWithOwner(UAbilitySystemComponent* Owner)
 {
 	this->Owner = Owner;
+}
+
+// ----------------------------------------------------
+
+FGameplayAbilitySpec:: FGameplayAbilitySpec(FGameplayAbilitySpecDef& InDef, FActiveGameplayEffectHandle InGameplayEffectHandle)
+	: FGameplayAbilitySpec(InDef.Ability ? InDef.Ability->GetDefaultObject<UGameplayAbility>() : nullptr, InDef.Level, InDef.InputID, InDef.SourceObject)
+{
+	InDef.AssignedHandle = Handle;
+	GameplayEffectHandle = InGameplayEffectHandle;
+}
+
+// ----------------------------------------------------
+
+FScopedAbilityListLock::FScopedAbilityListLock(UAbilitySystemComponent& InAbilitySystemComponent)
+	: AbilitySystemComponent(InAbilitySystemComponent)
+{
+	AbilitySystemComponent.IncrementAbilityListLock();
+}
+
+FScopedAbilityListLock::~FScopedAbilityListLock()
+{
+	AbilitySystemComponent.DecrementAbilityListLock();
 }
\ No newline at end of file
diff --git a/Engine/Source/Runtime/GameplayAbilities/Private/GameplayEffect.cpp b/Engine/Source/Runtime/GameplayAbilities/Private/GameplayEffect.cpp
--- a/Engine/Source/Runtime/GameplayAbilities/Private/GameplayEffect.cpp
+++ b/Engine/Source/Runtime/GameplayAbilities/Private/GameplayEffect.cpp
@@ -11,6 +11,7 @@
 #include "GameplayEffectExecutionCalculation.h"
 #if ENABLE_VISUAL_LOG
 #include "VisualLoggerTypes.h"
+//#include "VisualLogger/VisualLogger.h"
 #endif // ENABLE_VISUAL_LOG
 
 const float UGameplayEffect::INFINITE_DURATION = -1.f;
@@ -449,6 +450,9 @@ FGameplayEffectSpec::FGameplayEffectSpec(const UGameplayEffect* InDef, const FGa
 		TargetEffectSpecs.Add(FGameplayEffectSpecHandle(new FGameplayEffectSpec(TargetDef, EffectContext, Level)));
 	}
 
+	// Make Granted AbilitySpecs (caller may modify these specs after creating spec, which is why we dont just reference them from the def)
+	GrantedAbilitySpecs = InDef->GrantedAbilities;
+
 	// Everything is setup now, capture data from our source
 	CaptureDataFromSource();
 }
@@ -474,6 +478,7 @@ FGameplayEffectSpec::FGameplayEffectSpec(FGameplayEffectSpec&& Other)
 	, bCompletedSourceAttributeCapture(Other.bCompletedSourceAttributeCapture)
 	, bCompletedTargetAttributeCapture(Other.bCompletedTargetAttributeCapture)
 	, bDurationLocked(Other.bDurationLocked)
+	, GrantedAbilitySpecs(MoveTemp(Other.GrantedAbilitySpecs))
 	, SetByCallerMagnitudes(MoveTemp(Other.SetByCallerMagnitudes))
 	, EffectContext(Other.EffectContext)
 	, Level(Other.Level)
@@ -498,6 +503,7 @@ FGameplayEffectSpec& FGameplayEffectSpec::operator=(FGameplayEffectSpec&& Other)
 	bCompletedSourceAttributeCapture = Other.bCompletedSourceAttributeCapture;
 	bCompletedTargetAttributeCapture = Other.bCompletedTargetAttributeCapture;
 	bDurationLocked = Other.bDurationLocked;
+	GrantedAbilitySpecs = MoveTemp(Other.GrantedAbilitySpecs);
 	SetByCallerMagnitudes = MoveTemp(Other.SetByCallerMagnitudes);
 	EffectContext = Other.EffectContext;
 	Level = Other.Level;
@@ -521,6 +527,7 @@ FGameplayEffectSpec& FGameplayEffectSpec::operator=(const FGameplayEffectSpec& O
 	bCompletedSourceAttributeCapture = Other.bCompletedSourceAttributeCapture;
 	bCompletedTargetAttributeCapture = Other.bCompletedTargetAttributeCapture;
 	bDurationLocked = Other.bDurationLocked;
+	GrantedAbilitySpecs =  Other.GrantedAbilitySpecs;
 	SetByCallerMagnitudes = Other.SetByCallerMagnitudes;
 	EffectContext = Other.EffectContext;
 	Level = Other.Level;
@@ -1614,6 +1621,7 @@ void FActiveGameplayEffectsContainer::OnMagnitudeDependencyChange(FActiveGamepla
 
 void FActiveGameplayEffectsContainer::OnStackCountChange(FActiveGameplayEffect& ActiveEffect)
 {
+	MarkItemDirty(ActiveEffect);
 	UpdateAllAggregatorModMagnitudes(ActiveEffect);
 }
 
@@ -1663,7 +1671,7 @@ FActiveGameplayEffect* FActiveGameplayEffectsContainer::FindStackableActiveGamep
 	EGameplayEffectStackingType StackingType = GEDef->StackingType;
 
 	UAbilitySystemComponent* SourceASC = Spec.GetContext().GetInstigatorAbilitySystemComponent();
-	if (SourceASC && StackingType != EGameplayEffectStackingType::None && Spec.GetDuration() != UGameplayEffect::INSTANT_APPLICATION)
+	if (StackingType != EGameplayEffectStackingType::None && Spec.GetDuration() != UGameplayEffect::INSTANT_APPLICATION)
 	{
 		// Iterate through GameplayEffects to see if we find a match. Note that we could cache off a handle in a map but we would still
 		// do a linear search through GameplayEffects to find the actual FActiveGameplayEffect (due to unstable nature of the GameplayEffects array).
@@ -1672,7 +1680,7 @@ FActiveGameplayEffect* FActiveGameplayEffectsContainer::FindStackableActiveGamep
 		for (FActiveGameplayEffect& ActiveEffect: this)
 		{
 			// Aggregate by source stacking additionally requires the source ability component to match
-			if (ActiveEffect.Spec.Def == Spec.Def && ((StackingType == EGameplayEffectStackingType::AggregateByTarget) || (SourceASC == ActiveEffect.Spec.GetContext().GetInstigatorAbilitySystemComponent())))
+			if (ActiveEffect.Spec.Def == Spec.Def && ((StackingType == EGameplayEffectStackingType::AggregateByTarget) || (SourceASC && SourceASC == ActiveEffect.Spec.GetContext().GetInstigatorAbilitySystemComponent())))
 			{
 				StackableGE = &ActiveEffect;
 				break;
@@ -1995,10 +2003,11 @@ FActiveGameplayEffect* FActiveGameplayEffectsContainer::ApplyGameplayEffectSpec(
 		}
 
 		FGameplayEffectSpec& ExistingSpec = ExistingStackableGE->Spec;
-
-		int32 NewStackCount = ExistingSpec.StackCount;
-
-		// Process overflow if already at the stack limit; Overflow can potentially reject the application outright
+		
+		// How to apply multiple stacks at once? What if we trigger an overflow which can reject the application?
+		// We still want to apply the stacks that didnt push us over, but we also want to call HandleActiveGameplayEffectStackOverflow.
+		
+		// For now: call HandleActiveGameplayEffectStackOverflow only if we are ALREADY at the limit. Else we just clamp stack limit to max.
 		if (ExistingSpec.StackCount == ExistingSpec.Def->StackLimitCount)
 		{
 			if (!HandleActiveGameplayEffectStackOverflow(*ExistingStackableGE, ExistingSpec, Spec))
@@ -2006,9 +2015,11 @@ FActiveGameplayEffect* FActiveGameplayEffectsContainer::ApplyGameplayEffectSpec(
 				return nullptr;
 			}
 		}
-		else
+		
+		int32 NewStackCount = ExistingSpec.StackCount + Spec.StackCount;
+		if (ExistingSpec.Def->StackLimitCount > 0)
 		{
-			++NewStackCount;
+			NewStackCount = FMath::Min(NewStackCount, ExistingSpec.Def->StackLimitCount);
 		}
 
 		// Need to unregister callbacks because the source aggregators could potentially be different with the new application. They will be
@@ -2019,9 +2030,15 @@ FActiveGameplayEffect* FActiveGameplayEffectsContainer::ApplyGameplayEffectSpec(
 		// and cause a removal and add of only the ones that have changed. For now, ensure on this happening and come
 		// back to this later.
 		ensureMsgf(ExistingSpec.DynamicGrantedTags == Spec.DynamicGrantedTags, TEXT("While adding a stack of the gameplay effect: %s, the old stack and the new application had different dynamically granted tags, which is currently not resolved properly!"), *Spec.Def->GetName());
+		
+		// We only grant abilities on the first apply. So we *dont* want the new spec's GrantedAbilitySpecs list
+		TArray<FGameplayAbilitySpecDef>	GrantedSpecTempArray(MoveTemp(ExistingStackableGE->Spec.GrantedAbilitySpecs));
 
 		ExistingStackableGE->Spec = Spec;
 		ExistingStackableGE->Spec.StackCount = NewStackCount;
+
+		// Swap in old granted ability spec
+		ExistingStackableGE->Spec.GrantedAbilitySpecs = MoveTemp(GrantedSpecTempArray);
 		
 		AppliedActiveGE = ExistingStackableGE;
 
@@ -2276,7 +2293,7 @@ bool FActiveGameplayEffectsContainer::RemoveActiveGameplayEffect(FActiveGameplay
 			return true;
 		}
 	}
-	ABILITY_LOG(Warning, TEXT("RemoveActiveGameplayEffect called with invalid Handle: %s"), *Handle.ToString());
+	ABILITY_LOG(Log, TEXT("RemoveActiveGameplayEffect called with invalid Handle: %s"), *Handle.ToString());
 	return false;
 }
 
@@ -2306,6 +2323,24 @@ bool FActiveGameplayEffectsContainer::InternalRemoveActiveGameplayEffect(int32 I
 		// Mark the effect as pending removal
 		Effect.IsPendingRemove = true;
 
+		// Remove Granted Abilities
+		for (FGameplayAbilitySpecDef& AbilitySpecDef : Effect.Spec.GrantedAbilitySpecs)
+		{
+			if (AbilitySpecDef.AssignedHandle.IsValid())
+			{
+				switch(AbilitySpecDef.RemovalPolicy)
+				{
+				case EGameplayEffectGrantedAbilityRemovePolicy::CancelAbilityImmediately:
+					Owner->ClearAbility(AbilitySpecDef.AssignedHandle);
+					break;
+				case EGameplayEffectGrantedAbilityRemovePolicy::RemoveAbilityOnEnd:
+					Owner->SetRemoveAbilityOnEnd(AbilitySpecDef.AssignedHandle);
+					break;
+				}
+			}
+		}
+
+		// Invoke Remove GameplayCue event
 		bool ShouldInvokeGameplayCueEvent = true;
 		const bool bIsNetAuthority = IsNetAuthority();
 		if (!bIsNetAuthority && Effect.PredictionKey.IsValidKey() && Effect.PredictionKey.WasReceived() == false)
diff --git a/Engine/Source/Runtime/GameplayAbilities/Public/Abilities/GameplayAbility.h b/Engine/Source/Runtime/GameplayAbilities/Public/Abilities/GameplayAbility.h
--- a/Engine/Source/Runtime/GameplayAbilities/Public/Abilities/GameplayAbility.h
+++ b/Engine/Source/Runtime/GameplayAbilities/Public/Abilities/GameplayAbility.h
@@ -67,6 +67,12 @@
 /** Notification delegate definition for when the gameplay ability ends */
 DECLARE_DELEGATE_OneParam(FOnGameplayAbilityEnded, UGameplayAbility*);
 
+/** Notification delegate definition for when the gameplay ability is cancelled */
+DECLARE_MULTICAST_DELEGATE(FOnGameplayAbilityCancelled);
+
+/** Used to notify ability state tasks that a state is being ended */
+DECLARE_MULTICAST_DELEGATE_OneParam(FOnGameplayAbilityStateEnded, FName);
+
 /** TriggerData */
 USTRUCT()
 struct FAbilityTriggerData
@@ -209,6 +215,12 @@ public:
 	/** Notification that the ability has ended.  Set using TryActivateAbility. */
 	FOnGameplayAbilityEnded OnGameplayAbilityEnded;
 
+	/** Notification that the ability is being cancelled.  Called before OnGameplayAbilityEnded. */
+	FOnGameplayAbilityCancelled OnGameplayAbilityCancelled;
+
+	/** Used by the ability state task to handle when a state is ended */
+	FOnGameplayAbilityStateEnded OnGameplayAbilityStateEnded;
+
 	/** This ability has these tags */
 	UPROPERTY(EditDefaultsOnly, Category = Tags)
 	FGameplayTagContainer AbilityTags;
@@ -253,6 +265,13 @@ public:
 	UPROPERTY(EditDefaultsOnly, Category = Input)
 	bool bReplicateInputDirectly;
 
+	// --------------------------------------
+	//	CancelAbility
+	// --------------------------------------
+
+	/** Destroys instanced-per-execution abilities. Instance-per-actor abilities should 'reset'. Any active ability state tasks receive the 'OnAbilityStateInterrupted' event. Non instance abilities - what can we do? */
+	virtual void CancelAbility(const FGameplayAbilitySpecHandle Handle, const FGameplayAbilityActorInfo* ActorInfo, const FGameplayAbilityActivationInfo ActivationInfo, bool bReplicateCancelAbility);
+
 protected:
 
 	// --------------------------------------
@@ -361,13 +380,6 @@ public:
 	/** Do boilerplate init stuff and then call ActivateAbility */
 	virtual void PreActivate(const FGameplayAbilitySpecHandle Handle, const FGameplayAbilityActorInfo* ActorInfo, const FGameplayAbilityActivationInfo ActivationInfo, FOnGameplayAbilityEnded* OnGameplayAbilityEndedDelegate);
 
-	// --------------------------------------
-	//	CancelAbility
-	// --------------------------------------
-
-	/** Destroys instanced-per-execution abilities. Instance-per-actor abilities should 'reset'. Any active ability state tasks receive the 'OnAbilityStateInterrupted' event. Non instance abilities - what can we do? */
-	virtual void CancelAbility(const FGameplayAbilitySpecHandle Handle, const FGameplayAbilityActorInfo* ActorInfo, const FGameplayAbilityActivationInfo ActivationInfo);
-
 protected:
 	/** Destroys instanced-per-execution abilities. Instance-per-actor abilities should 'reset'. Non instance abilities - what can we do? */
 	UFUNCTION(BlueprintCallable, Category = Ability)
@@ -466,7 +478,7 @@ protected:
 	virtual void K2_ExecuteGameplayCue(FGameplayTag GameplayCueTag, FGameplayEffectContextHandle Context);
 
 	UFUNCTION(BlueprintCallable, Category = Ability, meta=(GameplayTagFilter="GameplayCue"), FriendlyName="AddGameplayCue")
-	virtual void K2_AddGameplayCue(FGameplayTag GameplayCueTag, FGameplayEffectContextHandle Context);
+	virtual void K2_AddGameplayCue(FGameplayTag GameplayCueTag, FGameplayEffectContextHandle Context, bool bRemoveOnAbilityEnd = true);
 
 	UFUNCTION(BlueprintCallable, Category = Ability, meta=(GameplayTagFilter="GameplayCue"), FriendlyName="RemoveGameplayCue")
 	virtual void K2_RemoveGameplayCue(FGameplayTag GameplayCueTag);
@@ -756,6 +768,9 @@ protected:
 
 	mutable FGameplayAbilitySpecHandle CurrentSpecHandle;
 
+	/** GameplayCues that were added during this ability that will get automatically removed when it ends */
+	TSet<FGameplayTag> TrackedGameplayCues;
+
 	/** Active montage being played by this ability */
 	UPROPERTY()
 	class UAnimMontage* CurrentMontage;
diff --git a/Engine/Source/Runtime/GameplayAbilities/Public/Abilities/GameplayAbilityTargetActor.h b/Engine/Source/Runtime/GameplayAbilities/Public/Abilities/GameplayAbilityTargetActor.h
--- a/Engine/Source/Runtime/GameplayAbilities/Public/Abilities/GameplayAbilityTargetActor.h
+++ b/Engine/Source/Runtime/GameplayAbilities/Public/Abilities/GameplayAbilityTargetActor.h
@@ -24,6 +24,8 @@ class GAMEPLAYABILITIES_API AGameplayAbilityTargetActor : public AActor
 
 public:
 
+	void Destroyed() override;
+
 	/** The TargetData this class produces can be entirely generated on the server. We don't require the client to send us full or partial TargetData (possibly just a 'confirm') */
 	UPROPERTY(EditAnywhere, Category=Advanced)
 	bool ShouldProduceTargetDataOnServer;
diff --git a/Engine/Source/Runtime/GameplayAbilities/Public/Abilities/GameplayAbilityTargetTypes.h b/Engine/Source/Runtime/GameplayAbilities/Public/Abilities/GameplayAbilityTargetTypes.h
--- a/Engine/Source/Runtime/GameplayAbilities/Public/Abilities/GameplayAbilityTargetTypes.h
+++ b/Engine/Source/Runtime/GameplayAbilities/Public/Abilities/GameplayAbilityTargetTypes.h
@@ -594,9 +594,9 @@ DECLARE_MULTICAST_DELEGATE_OneParam(FAbilityTargetData, FGameplayAbilityTargetDa
 /** Generic callback for returning when target data is available */
 DECLARE_MULTICAST_DELEGATE_TwoParams(FAbilityTargetDataSetDelegate, FGameplayAbilityTargetDataHandle, FGameplayTag);
 
-/** These are generic, nonpayload carrying events that the client can send to the server via ServerSetReplicatedClientEvent. */
+/** These are generic, nonpayload carrying events that are replicated between the client and server */
 UENUM()
-namespace EAbilityReplicatedClientEvent
+namespace EAbilityGenericReplicatedEvent
 {
 	enum Type
 	{	
@@ -608,6 +608,10 @@ namespace EAbilityReplicatedClientEvent
 		InputPressed,	
 		/** Input release event of the ability */
 		InputReleased,
+		/** A generic event from the client */
+		GenericSignalFromClient,
+		/** A generic event from the server */
+		GenericSignalFromServer,		
 
 		MAX
 	};
@@ -623,7 +627,7 @@ struct FAbilityReplicatedData
 };
 
 /** Struct defining the cached data for a specific gameplay ability. This data is generally syncronized client->server in a network game. */
-struct FAbilityClientDataCache
+struct FAbilityReplicatedDataCache
 {
 	/** What elements this activation is targeting */
 	FGameplayAbilityTargetDataHandle TargetData;
@@ -644,9 +648,9 @@ struct FAbilityClientDataCache
 	FSimpleMulticastDelegate TargetCancelledDelegate;
 
 	/** Generic events that contain no payload data */
-	FAbilityReplicatedData	GenericEvents[EAbilityReplicatedClientEvent::MAX];
+	FAbilityReplicatedData	GenericEvents[EAbilityGenericReplicatedEvent::MAX];
 
-	FAbilityClientDataCache() : bTargetConfirmed(false), bTargetCancelled(false) {}
+	FAbilityReplicatedDataCache() : bTargetConfirmed(false), bTargetCancelled(false) {}
 
 	/** Resets any cached data, leaves delegates up */
 	void Reset()
@@ -654,7 +658,7 @@ struct FAbilityClientDataCache
 		bTargetConfirmed = bTargetCancelled = false;
 		TargetData = FGameplayAbilityTargetDataHandle();
 		ApplicationTag = FGameplayTag();
-		for (int32 i=0; i < (int32) EAbilityReplicatedClientEvent::MAX; ++i)
+		for (int32 i=0; i < (int32) EAbilityGenericReplicatedEvent::MAX; ++i)
 		{
 			GenericEvents[i].bTriggered = false;
 		}
diff --git a/Engine/Source/Runtime/GameplayAbilities/Public/Abilities/GameplayAbilityTypes.h b/Engine/Source/Runtime/GameplayAbilities/Public/Abilities/GameplayAbilityTypes.h
--- a/Engine/Source/Runtime/GameplayAbilities/Public/Abilities/GameplayAbilityTypes.h
+++ b/Engine/Source/Runtime/GameplayAbilities/Public/Abilities/GameplayAbilityTypes.h
@@ -3,6 +3,7 @@
 #pragma once
 
 #include "GameplayEffect.h"
+#include "GameplayAbilitySpec.h"
 #include "GameplayPrediction.h"
 #include "AttributeSet.h"
 #include "GameFramework/MovementComponent.h"
@@ -80,25 +81,6 @@ namespace EGameplayAbilityReplicationPolicy
 }
 
 
-UENUM(BlueprintType)
-namespace EGameplayAbilityActivationMode
-{
-	enum Type
-	{
-		// We are the authority activating this ability
-		Authority,
-
-		// We are not the authority but aren't predicting yet. This is a mostly invalid state to be in.
-		NonAuthority,
-
-		// We are predicting the activation of this ability
-		Predicting,
-
-		// We are not the authority, but the authority has confirmed this activation
-		Confirmed,
-	};
-}
-
 UENUM(BlueprintType)
 namespace EGameplayAbilityTriggerSource
 {
@@ -120,49 +102,6 @@ namespace EGameplayAbilityTriggerSource
 
 // ----------------------------------------------------
 
-USTRUCT(BlueprintType)
-struct FGameplayAbilitySpecHandle
-{
-	GENERATED_USTRUCT_BODY()
-
-	FGameplayAbilitySpecHandle()
-	: Handle(INDEX_NONE)
-	{
-
-	}
-
-	bool IsValid() const
-	{
-		return Handle != INDEX_NONE;
-	}
-
-	void GenerateNewHandle()
-	{
-		static int32 GHandle = 1;
-		Handle = GHandle++;
-	}
-
-	bool operator==(const FGameplayAbilitySpecHandle& Other) const
-	{
-		return Handle == Other.Handle;
-	}
-
-	bool operator!=(const FGameplayAbilitySpecHandle& Other) const
-	{
-		return Handle != Other.Handle;
-	}
-
-	friend uint32 GetTypeHash(const FGameplayAbilitySpecHandle& SpecHandle)
-	{
-		return ::GetTypeHash(SpecHandle.Handle);
-	}
-
-private:
-
-	UPROPERTY()
-	int32 Handle;
-};
-
 /**
  *	FGameplayAbilityActorInfo
  *
@@ -220,180 +159,6 @@ struct GAMEPLAYABILITIES_API FGameplayAbilityActorInfo
 	virtual void ClearActorInfo();
 };
 
-/**
- *	FGameplayAbilityActivationInfo
- *
- *	Data tied to a specific activation of an ability.
- *		-Tell us whether we are the authority, if we are predicting, confirmed, etc.
- *		-Holds current and previous PredictionKey
- *		-Generally not meant to be subclassed in projects.
- *		-Passed around by value since the struct is small.
- */
-USTRUCT(BlueprintType)
-struct GAMEPLAYABILITIES_API FGameplayAbilityActivationInfo
-{
-	GENERATED_USTRUCT_BODY()
-
-	FGameplayAbilityActivationInfo()
-		: ActivationMode(EGameplayAbilityActivationMode::Authority)
-		, bCanBeEndedByOtherInstance(false)
-	{
-
-	}
-
-	FGameplayAbilityActivationInfo(AActor* InActor)
-		: bCanBeEndedByOtherInstance(false)	
-	{
-		// On Init, we are either Authority or NonAuthority. We haven't been given a PredictionKey and we haven't been confirmed.
-		// NonAuthority essentially means 'I'm not sure what how I'm going to do this yet'.
-		ActivationMode = (InActor->Role == ROLE_Authority ? EGameplayAbilityActivationMode::Authority : EGameplayAbilityActivationMode::NonAuthority);
-	}
-
-	FGameplayAbilityActivationInfo(EGameplayAbilityActivationMode::Type InType)
-		: ActivationMode(InType)
-		, bCanBeEndedByOtherInstance(false)
-	{
-	}	
-
-	UPROPERTY(BlueprintReadOnly, Category = "ActorInfo")
-	mutable TEnumAsByte<EGameplayAbilityActivationMode::Type>	ActivationMode;
-
-	/** An ability that runs on multiple game instances can be canceled by a remote instance, but only if that remote instance has already confirmed starting it. */
-	UPROPERTY()
-	bool bCanBeEndedByOtherInstance;
-
-	void SetActivationConfirmed();
-
-	/** Called on client to set this as a predicted ability */
-	void SetPredicting(FPredictionKey PredictionKey);
-
-	/** Called on the server to set the key used by the client to predict this ability */
-	void ServerSetActivationPredictionKey(FPredictionKey PredictionKey);
-
-	FPredictionKey GetActivationPredictionKey() { return PredictionKeyWhenActivated; }
-
-private:
-
-	// This was the prediction key used to activate this ability. It does not get updated
-	// if new prediction keys are generated over the course of the ability.
-	UPROPERTY()
-	FPredictionKey PredictionKeyWhenActivated;
-	
-};
-
-// ----------------------------------------------------
-
-/** An activatable ability spec, hosted on the ability system component */
-USTRUCT()
-struct GAMEPLAYABILITIES_API FGameplayAbilitySpec : public FFastArraySerializerItem
-{
-	GENERATED_USTRUCT_BODY()
-
-	FGameplayAbilitySpec()
-	: Ability(nullptr), Level(1), InputID(INDEX_NONE), SourceObject(nullptr), InputPressed(false), RemoveAfterActivation(false), ActiveCount(0)
-	{
-		
-	}
-
-	FGameplayAbilitySpec(UGameplayAbility* InAbility, int32 InLevel=1, int32 InInputID=INDEX_NONE, UObject* InSourceObject=nullptr)
-		: Ability(InAbility), Level(InLevel), InputID(InInputID), SourceObject(InSourceObject), InputPressed(false), RemoveAfterActivation(false), ActiveCount(0)
-	{
-		Handle.GenerateNewHandle();
-	}
-
-	/** Handle for outside sources to refer to this spec by */
-	UPROPERTY()
-	FGameplayAbilitySpecHandle Handle;
-	
-	/** Ability of the spec (Always the CDO. This should be const but too many things modify it currently) */
-	UPROPERTY()
-	UGameplayAbility* Ability;
-	
-	/** Level of Ability */
-	UPROPERTY()
-	int32	Level;
-
-	/** InputID, if bound */
-	UPROPERTY()
-	int32	InputID;
-
-	/** Object this ability was created from, can be an actor or static object. Useful to bind an ability to a gameplay object */
-	UPROPERTY()
-	UObject* SourceObject;
-
-	/** Is input currently pressed. Set to false when input is released */
-	UPROPERTY(NotReplicated)
-	bool	InputPressed;
-
-	/** If true, this ability should be removed as soon as it finishes executing */
-	UPROPERTY(NotReplicated)
-	bool	RemoveAfterActivation;
-
-	/** A count of the number of times this ability has been activated minus the number of times it has been ended. For instanced abilities this will be the number of currently active instances. Can't replicate until prediction accurately handles this.*/
-	UPROPERTY(NotReplicated)
-	uint8	ActiveCount;
-
-	/** Activation state of this ability. This is not replicated since it needs to be overwritten locally on clients during prediction. */
-	UPROPERTY(NotReplicated)
-	FGameplayAbilityActivationInfo	ActivationInfo;
-
-	/** Non replicating instances of this ability. */
-	UPROPERTY(NotReplicated)
-	TArray<UGameplayAbility*> NonReplicatedInstances;
-
-	/** Replicated instances of this ability.. */
-	UPROPERTY()
-	TArray<UGameplayAbility*> ReplicatedInstances;
-
-	/** Returns the primary instance, used for instance once abilities */
-	UGameplayAbility* GetPrimaryInstance() const;
-
-	/** Returns all instances, which can include instance per execution abilities */
-	TArray<UGameplayAbility*> GetAbilityInstances() const
-	{
-		TArray<UGameplayAbility*> Abilities;
-		Abilities.Append(ReplicatedInstances);
-		Abilities.Append(NonReplicatedInstances);
-		return Abilities;
-	}
-
-	/** Returns true if this ability is active in any way */
-	bool IsActive() const;
-
-	void PreReplicatedRemove(const struct FGameplayAbilitySpecContainer& InArraySerializer);
-	void PostReplicatedAdd(const struct FGameplayAbilitySpecContainer& InArraySerializer);
-};
-
-/** Fast serializer wrapper for above struct */
-USTRUCT()
-struct GAMEPLAYABILITIES_API FGameplayAbilitySpecContainer : public FFastArraySerializer
-{
-	GENERATED_USTRUCT_BODY()
-
-	/** List of activatable abilities */
-	UPROPERTY()
-	TArray<FGameplayAbilitySpec> Items;
-
-	/** Component that owns this list */
-	UAbilitySystemComponent* Owner;
-
-	void RegisterWithOwner(UAbilitySystemComponent* Owner);
-
-	bool NetDeltaSerialize(FNetDeltaSerializeInfo & DeltaParms)
-	{
-		return FFastArraySerializer::FastArrayDeltaSerialize<FGameplayAbilitySpec, FGameplayAbilitySpecContainer>(Items, DeltaParms, *this);
-	}
-};
-
-template<>
-struct TStructOpsTypeTraits< FGameplayAbilitySpecContainer > : public TStructOpsTypeTraitsBase
-{
-	enum
-	{
-		WithNetDeltaSerializer = true,
-	};
-};
-
 // ----------------------------------------------------
 
 USTRUCT()
@@ -570,10 +335,12 @@ struct GAMEPLAYABILITIES_API FGameplayEventData
  */
 struct FGameplayAbiliyInputBinds
 {
-	FGameplayAbiliyInputBinds(FString InConfirmTargetCommand, FString InCancelTargetCommand, FString InEnumName)
+	FGameplayAbiliyInputBinds(FString InConfirmTargetCommand, FString InCancelTargetCommand, FString InEnumName, int32 InConfirmTargetInputID = INDEX_NONE, int32 InCancelTargetInputID = INDEX_NONE)
 	: ConfirmTargetCommand(InConfirmTargetCommand)
 	, CancelTargetCommand(InCancelTargetCommand)
 	, EnumName(InEnumName)
+	, ConfirmTargetInputID(InConfirmTargetInputID)
+	, CancelTargetInputID(InCancelTargetInputID)
 	{ }
 
 	/** Defines command string that will be bound to Confirm Targeting */
@@ -585,6 +352,12 @@ struct FGameplayAbiliyInputBinds
 	/** Returns enum to use for ability binds. E.g., "Ability1"-"Ability9" input commands will be bound to ability activations inside the AbiltiySystemComponent */
 	FString	EnumName;
 
+	/** If >=0, Confirm is bound to an entry in the enum */
+	int32 ConfirmTargetInputID;
+
+	/** If >=0, Confirm is bound to an entry in the enum */
+	int32 CancelTargetInputID;
+
 	UEnum* GetBindEnum() { return FindObject<UEnum>(ANY_PACKAGE, *EnumName); }
 };
 
diff --git a/Engine/Source/Runtime/GameplayAbilities/Public/Abilities/GameplayAbility_CharacterJump.h b/Engine/Source/Runtime/GameplayAbilities/Public/Abilities/GameplayAbility_CharacterJump.h
--- a/Engine/Source/Runtime/GameplayAbilities/Public/Abilities/GameplayAbility_CharacterJump.h
+++ b/Engine/Source/Runtime/GameplayAbilities/Public/Abilities/GameplayAbility_CharacterJump.h
@@ -21,5 +21,5 @@ public:
 
 	virtual void InputReleased(const FGameplayAbilitySpecHandle Handle, const FGameplayAbilityActorInfo* ActorInfo, const FGameplayAbilityActivationInfo ActivationInfo) override;
 
-	virtual void CancelAbility(const FGameplayAbilitySpecHandle Handle, const FGameplayAbilityActorInfo* ActorInfo, const FGameplayAbilityActivationInfo ActivationInfo) override;
+	virtual void CancelAbility(const FGameplayAbilitySpecHandle Handle, const FGameplayAbilityActorInfo* ActorInfo, const FGameplayAbilityActivationInfo ActivationInfo, bool bReplicateCancelAbility) override;
 };
\ No newline at end of file
diff --git a/Engine/Source/Runtime/GameplayAbilities/Public/Abilities/Tasks/AbilityTask.h b/Engine/Source/Runtime/GameplayAbilities/Public/Abilities/Tasks/AbilityTask.h
--- a/Engine/Source/Runtime/GameplayAbilities/Public/Abilities/Tasks/AbilityTask.h
+++ b/Engine/Source/Runtime/GameplayAbilities/Public/Abilities/Tasks/AbilityTask.h
@@ -1,6 +1,7 @@
 // Copyright 1998-2015 Epic Games, Inc. All Rights Reserved.
 #pragma once
 #include "GameplayAbility.h"
+#include "AbilitySystemComponent.h"
 #include "AbilityTask.generated.h"
 
 
@@ -154,7 +155,7 @@ protected:
 
 	/** Helper method for registering client replicated callbacks */
 	
-	bool CallOrAddReplicatedDelegate(EAbilityReplicatedClientEvent::Type Event, FSimpleMulticastDelegate::FDelegate Delegate);
+	bool CallOrAddReplicatedDelegate(EAbilityGenericReplicatedEvent::Type Event, FSimpleMulticastDelegate::FDelegate Delegate);
 };
 
 //For searching through lists of ability instances
diff --git a/Engine/Source/Runtime/GameplayAbilities/Public/Abilities/Tasks/AbilityTask_MoveToLocation.h b/Engine/Source/Runtime/GameplayAbilities/Public/Abilities/Tasks/AbilityTask_MoveToLocation.h
--- a/Engine/Source/Runtime/GameplayAbilities/Public/Abilities/Tasks/AbilityTask_MoveToLocation.h
+++ b/Engine/Source/Runtime/GameplayAbilities/Public/Abilities/Tasks/AbilityTask_MoveToLocation.h
@@ -39,6 +39,8 @@ class UAbilityTask_MoveToLocation : public UAbilityTask
 	/** Tick function for this task, if bTickingTask == true */
 	virtual void TickTask(float DeltaTime) override;
 
+	virtual void OnDestroy(bool AbilityIsEnding) override;
+
 protected:
 
 	bool bIsFinished;
diff --git a/Engine/Source/Runtime/GameplayAbilities/Public/Abilities/Tasks/AbilityTask_NetworkSyncPoint.h b/Engine/Source/Runtime/GameplayAbilities/Public/Abilities/Tasks/AbilityTask_NetworkSyncPoint.h
new file mode 100644
--- /dev/null
+++ b/Engine/Source/Runtime/GameplayAbilities/Public/Abilities/Tasks/AbilityTask_NetworkSyncPoint.h
@@ -0,0 +1,64 @@
+// Copyright 1998-2015 Epic Games, Inc. All Rights Reserved.
+#pragma once
+#include "AbilityTask.h"
+#include "Abilities/GameplayAbilityTypes.h"
+#include "AbilitySystemComponent.h"
+#include "AbilityTask_NetworkSyncPoint.generated.h"
+
+DECLARE_DYNAMIC_MULTICAST_DELEGATE(FNetworkSyncDelegate);
+
+UENUM()
+enum class EAbilityTaskNetSyncType : uint8
+{
+	/** Both client and server wait until the other signals */
+	BothWait,
+
+	/** Only server will wait for the client signal. Client will signal and immediately continue without waiting to hear from Server. */
+	OnlyServerWait,
+
+	/** Only client will wait for the server signal. Server will signal and immediately continue without waiting to hear from Client. */
+	OnlyClientWait	
+};
+
+/** Task for providing a generic sync point for client server (one can wait for a signal from the other) */
+UCLASS(MinimalAPI)
+class UAbilityTask_NetworkSyncPoint : public UAbilityTask
+{
+	GENERATED_UCLASS_BODY()
+
+	UPROPERTY(BlueprintAssignable)
+	FNetworkSyncDelegate	OnSync;
+
+	UFUNCTION()
+	void OnSignalCallback();
+
+	virtual void Activate() override;
+
+	/**
+	 *	
+	 *	Synchronize execution flow between Client and Server. Depending on SyncType, the Client and Server will wait for the other to reach this node or another WaitNetSync node in the ability before continuing execution.  
+	 *	
+	 *	BothWait - Both Client and Server will wait until the other reaches the node. (Whoever gets their first, waits for the other before continueing).
+	 *	OnlyServerWait - Only server will wait for the client signal. Client will signal and immediately continue without waiting to hear from Server.
+	 *	OnlyClientWait - Only client will wait for the server signal. Server will signal and immediately continue without waiting to hear from Client.
+	 *	
+	 *	Note that this is "ability instance wide". These sync points never affect sync points in other abilities. 
+	 *	
+	 *	In most cases you will have both client and server execution paths connected to the same WaitNetSync node. However it is possible to use separate nodes
+	 *	for cleanliness of the graph. The "signal" is "ability instance wide".
+	 *	
+	 */
+	UFUNCTION(BlueprintCallable, Category="Ability|Tasks", meta = (HidePin = "WorldContextObject", DefaultToSelf = "WorldContextObject", BlueprintInternalUseOnly = "TRUE"))
+	static UAbilityTask_NetworkSyncPoint* WaitNetSync(UObject* WorldContextObject, EAbilityTaskNetSyncType SyncType);
+
+protected:
+
+	virtual void OnDestroy(bool AbilityEnded) override;
+
+	void SyncFinished();
+
+	/** The event we replicate */
+	EAbilityGenericReplicatedEvent::Type	ReplicatedEventToListenFor;
+
+	EAbilityTaskNetSyncType SyncType;
+};
\ No newline at end of file
diff --git a/Engine/Source/Runtime/GameplayAbilities/Public/Abilities/Tasks/AbilityTask_PlayMontageAndWait.h b/Engine/Source/Runtime/GameplayAbilities/Public/Abilities/Tasks/AbilityTask_PlayMontageAndWait.h
--- a/Engine/Source/Runtime/GameplayAbilities/Public/Abilities/Tasks/AbilityTask_PlayMontageAndWait.h
+++ b/Engine/Source/Runtime/GameplayAbilities/Public/Abilities/Tasks/AbilityTask_PlayMontageAndWait.h
@@ -22,6 +22,9 @@ class UAbilityTask_PlayMontageAndWait : public UAbilityTask
 	UFUNCTION()
 	void OnMontageEnded(UAnimMontage* Montage, bool bInterrupted);
 
+	UFUNCTION()
+	void OnMontageInterrupted();
+
 	UFUNCTION(BlueprintCallable, Category="Ability|Tasks", meta = (FriendlyName="PlayMontageAndWait",
 		HidePin = "WorldContextObject", DefaultToSelf = "WorldContextObject", BlueprintInternalUseOnly = "TRUE"))
 	static UAbilityTask_PlayMontageAndWait* CreatePlayMontageAndWaitProxy(UObject* WorldContextObject,
@@ -38,6 +41,9 @@ private:
 
 	virtual void OnDestroy(bool AbilityEnded) override;
 
+	FOnMontageBlendingOutStarted BlendingOutDelegate;
+	FDelegateHandle InterruptedHandle;
+
 	UAnimMontage* MontageToPlay;
 	float Rate;	
 	FName StartSection;
diff --git a/Engine/Source/Runtime/GameplayAbilities/Public/Abilities/Tasks/AbilityTask_StartAbilityState.h b/Engine/Source/Runtime/GameplayAbilities/Public/Abilities/Tasks/AbilityTask_StartAbilityState.h
--- a/Engine/Source/Runtime/GameplayAbilities/Public/Abilities/Tasks/AbilityTask_StartAbilityState.h
+++ b/Engine/Source/Runtime/GameplayAbilities/Public/Abilities/Tasks/AbilityTask_StartAbilityState.h
@@ -37,6 +37,8 @@ class UAbilityTask_StartAbilityState : public UAbilityTask
 
 	virtual void Activate() override;
 
+	virtual void ExternalCancel() override;
+
 	virtual FString GetDebugString() const override;
 
 	/**
@@ -53,11 +55,13 @@ private:
 	FDelegateHandle EndStateHandle;
 	FDelegateHandle InterruptStateHandle;
 
+	bool bWasEnded;
+	bool bWasInterrupted;
+
 	bool bEndCurrentState;
 
 	virtual void OnDestroy(bool AbilityEnded) override;
 
 	void OnEndState(FName StateNameToEnd);
-
 	void OnInterruptState();
 };
diff --git a/Engine/Source/Runtime/GameplayAbilities/Public/Abilities/Tasks/AbilityTask_WaitGameplayEffectApplied.h b/Engine/Source/Runtime/GameplayAbilities/Public/Abilities/Tasks/AbilityTask_WaitGameplayEffectApplied.h
--- a/Engine/Source/Runtime/GameplayAbilities/Public/Abilities/Tasks/AbilityTask_WaitGameplayEffectApplied.h
+++ b/Engine/Source/Runtime/GameplayAbilities/Public/Abilities/Tasks/AbilityTask_WaitGameplayEffectApplied.h
@@ -21,8 +21,12 @@ class UAbilityTask_WaitGameplayEffectApplied : public UAbilityTask
 	FGameplayTagRequirements TargetTagRequirements;
 	bool TriggerOnce;
 
+	void SetExternalActor(AActor* InActor);
+
 protected:
 
+	UAbilitySystemComponent* GetASC();
+
 	virtual void BroadcastDelegate(AActor* Avatar, FGameplayEffectSpecHandle SpecHandle, FActiveGameplayEffectHandle ActiveHandle) { }
 	virtual void RegisterDelegate() { }
 	virtual void RemoveDelegate() { }
@@ -30,4 +34,7 @@ protected:
 	virtual void OnDestroy(bool AbilityEnded) override;
 
 	bool RegisteredCallback;
+	bool UseExternalOwner;
+
+	TWeakObjectPtr<UAbilitySystemComponent> ExternalOwner;
 };
\ No newline at end of file
diff --git a/Engine/Source/Runtime/GameplayAbilities/Public/Abilities/Tasks/AbilityTask_WaitGameplayEffectApplied_Self.h b/Engine/Source/Runtime/GameplayAbilities/Public/Abilities/Tasks/AbilityTask_WaitGameplayEffectApplied_Self.h
--- a/Engine/Source/Runtime/GameplayAbilities/Public/Abilities/Tasks/AbilityTask_WaitGameplayEffectApplied_Self.h
+++ b/Engine/Source/Runtime/GameplayAbilities/Public/Abilities/Tasks/AbilityTask_WaitGameplayEffectApplied_Self.h
@@ -15,9 +15,12 @@ class UAbilityTask_WaitGameplayEffectApplied_Self : public UAbilityTask_WaitGame
 	UPROPERTY(BlueprintAssignable)
 	FGameplayEffectAppliedSelfDelegate OnApplied;
 
-	/** Wait until the owner *receives* a GameplayEffect from a given source (the source may be the owner too!). If TriggerOnce is true, this task will only return one time. Otherwise it will return everytime a GE is applied that meets the requirements over the life of the ability */
+	/**	 
+	 * Wait until the owner *receives* a GameplayEffect from a given source (the source may be the owner too!). If TriggerOnce is true, this task will only return one time. Otherwise it will return everytime a GE is applied that meets the requirements over the life of the ability
+	 * Optional External Owner can be used to run this task on someone else (not the owner of the ability). By default you can leave this empty.
+	 */	 	 
 	UFUNCTION(BlueprintCallable, Category="Ability|Tasks", meta = (HidePin = "WorldContextObject", DefaultToSelf = "WorldContextObject", BlueprintInternalUseOnly = "TRUE"))
-	static UAbilityTask_WaitGameplayEffectApplied_Self* WaitGameplayEffectAppliedToSelf(UObject* WorldContextObject, const FGameplayTargetDataFilterHandle SourceFilter, FGameplayTagRequirements SourceTagRequirements, FGameplayTagRequirements TargetTagRequirements, bool TriggerOnce=false );
+	static UAbilityTask_WaitGameplayEffectApplied_Self* WaitGameplayEffectAppliedToSelf(UObject* WorldContextObject, AActor* TestActor, const FGameplayTargetDataFilterHandle SourceFilter, FGameplayTagRequirements SourceTagRequirements, FGameplayTagRequirements TargetTagRequirements, bool TriggerOnce=false, AActor* OptionalExternalOwner=nullptr );
 	
 protected:
 
diff --git a/Engine/Source/Runtime/GameplayAbilities/Public/Abilities/Tasks/AbilityTask_WaitGameplayEffectApplied_Target.h b/Engine/Source/Runtime/GameplayAbilities/Public/Abilities/Tasks/AbilityTask_WaitGameplayEffectApplied_Target.h
--- a/Engine/Source/Runtime/GameplayAbilities/Public/Abilities/Tasks/AbilityTask_WaitGameplayEffectApplied_Target.h
+++ b/Engine/Source/Runtime/GameplayAbilities/Public/Abilities/Tasks/AbilityTask_WaitGameplayEffectApplied_Target.h
@@ -15,9 +15,12 @@ class UAbilityTask_WaitGameplayEffectApplied_Target : public UAbilityTask_WaitGa
 	UPROPERTY(BlueprintAssignable)
 	FGameplayEffectAppliedTargetDelegate OnApplied;
 
-	/** Wait until the owner applies a GameplayEffect to a Target (the target may be the owner too!). If TriggerOnce is true, this task will only return one time. Otherwise it will return everytime a GE is applied that meets the requirements over the life of the ability */
+	/** 
+	 * Wait until the owner (or External Owner) applies a GameplayEffect to a Target (the target may be the owner too!). If TriggerOnce is true, this task will only return one time. Otherwise it will return everytime a GE is applied that meets the requirements over the life of the ability 
+	 * Optional External Owner can be used to run this task on someone else (not the owner of the ability). By default you can leave this empty.
+	 */
 	UFUNCTION(BlueprintCallable, Category="Ability|Tasks", meta = (HidePin = "WorldContextObject", DefaultToSelf = "WorldContextObject", BlueprintInternalUseOnly = "TRUE"))
-	static UAbilityTask_WaitGameplayEffectApplied_Target* WaitGameplayEffectAppliedToTarget(UObject* WorldContextObject, const FGameplayTargetDataFilterHandle TargetFilter, FGameplayTagRequirements SourceTagRequirements, FGameplayTagRequirements TargetTagRequirements, bool TriggerOnce=false );
+	static UAbilityTask_WaitGameplayEffectApplied_Target* WaitGameplayEffectAppliedToTarget(UObject* WorldContextObject, const FGameplayTargetDataFilterHandle TargetFilter, FGameplayTagRequirements SourceTagRequirements, FGameplayTagRequirements TargetTagRequirements, bool TriggerOnce=false, AActor* OptionalExternalOwner=nullptr);
 	
 protected:
 
diff --git a/Engine/Source/Runtime/GameplayAbilities/Public/Abilities/Tasks/AbilityTask_WaitGameplayTag.h b/Engine/Source/Runtime/GameplayAbilities/Public/Abilities/Tasks/AbilityTask_WaitGameplayTag.h
new file mode 100644
--- /dev/null
+++ b/Engine/Source/Runtime/GameplayAbilities/Public/Abilities/Tasks/AbilityTask_WaitGameplayTag.h
@@ -0,0 +1,49 @@
+// Copyright 1998-2015 Epic Games, Inc. All Rights Reserved.
+#pragma once
+#include "AbilityTask.h"
+#include "Abilities/GameplayAbilityTypes.h"
+#include "GameplayEffect.h"
+#include "AbilityTask_WaitgameplayTagBase.h"
+#include "AbilityTask_WaitGameplayTag.generated.h"
+
+DECLARE_DYNAMIC_MULTICAST_DELEGATE(FWaitGameplayTagDelegate);
+
+UCLASS(MinimalAPI)
+class UAbilityTask_WaitGameplayTagAdded : public UAbilityTask_WaitGameplayTag
+{
+	GENERATED_UCLASS_BODY()
+
+	UPROPERTY(BlueprintAssignable)
+	FWaitGameplayTagDelegate	Added;
+
+	/**
+	 * 	Wait until the specified gameplay tag is Added. By default this will look at the owner of this ability. OptionalExternalTarget can be set to make this look at another actor's tags for changes. 
+	 *  If the tag is already present when this task is started, it will immediately broadcast the Added event. It will keep listening as long as OnlyTriggerOnce = false.
+	 */
+	UFUNCTION(BlueprintCallable, Category = "Ability|Tasks", meta = (HidePin = "WorldContextObject", DefaultToSelf = "WorldContextObject", BlueprintInternalUseOnly = "TRUE"))
+	static UAbilityTask_WaitGameplayTagAdded* WaitGameplayTagAdd(UObject* WorldContextObject, FGameplayTag Tag, AActor* OptionalExternalTarget=nullptr, bool OnlyTriggerOnce=false);
+
+	virtual void Activate() override;
+
+	virtual void GameplayTagCallback(const FGameplayTag Tag, int32 NewCount) override;
+};
+
+UCLASS(MinimalAPI)
+class UAbilityTask_WaitGameplayTagRemoved : public UAbilityTask_WaitGameplayTag
+{
+	GENERATED_UCLASS_BODY()
+
+	UPROPERTY(BlueprintAssignable)
+	FWaitGameplayTagDelegate	Removed;
+
+	/**
+	 * 	Wait until the specified gameplay tag is Removed. By default this will look at the owner of this ability. OptionalExternalTarget can be set to make this look at another actor's tags for changes. 
+	 *  If the tag is not present when this task is started, it will immediately broadcast the Removed event. It will keep listening as long as OnlyTriggerOnce = false.
+	 */
+	UFUNCTION(BlueprintCallable, Category = "Ability|Tasks", meta = (HidePin = "WorldContextObject", DefaultToSelf = "WorldContextObject", BlueprintInternalUseOnly = "TRUE"))
+	static UAbilityTask_WaitGameplayTagRemoved* WaitGameplayTagRemove(UObject* WorldContextObject, FGameplayTag Tag, AActor* OptionalExternalTarget=nullptr, bool OnlyTriggerOnce=false);
+
+	virtual void Activate() override;
+
+	virtual void GameplayTagCallback(const FGameplayTag Tag, int32 NewCount) override;
+};
\ No newline at end of file
diff --git a/Engine/Source/Runtime/GameplayAbilities/Public/Abilities/Tasks/AbilityTask_WaitGameplayTagBase.h b/Engine/Source/Runtime/GameplayAbilities/Public/Abilities/Tasks/AbilityTask_WaitGameplayTagBase.h
new file mode 100644
--- /dev/null
+++ b/Engine/Source/Runtime/GameplayAbilities/Public/Abilities/Tasks/AbilityTask_WaitGameplayTagBase.h
@@ -0,0 +1,35 @@
+// Copyright 1998-2015 Epic Games, Inc. All Rights Reserved.
+#pragma once
+#include "AbilityTask.h"
+#include "Abilities/GameplayAbilityTypes.h"
+#include "GameplayEffect.h"
+#include "AbilityTask_WaitGameplayTagBase.generated.h"
+
+UCLASS(MinimalAPI)
+class UAbilityTask_WaitGameplayTag : public UAbilityTask
+{
+	GENERATED_UCLASS_BODY()
+
+	virtual void Activate() override;
+
+	UFUNCTION()
+	virtual void GameplayTagCallback(const FGameplayTag Tag, int32 NewCount);
+	
+	void SetExternalTarget(AActor* Actor);
+
+	FGameplayTag	Tag;
+
+protected:
+
+	UAbilitySystemComponent* GetTargetASC();
+
+	virtual void OnDestroy(bool AbilityIsEnding) override;
+	
+	bool RegisteredCallback;
+	
+	TWeakObjectPtr<UAbilitySystemComponent> OptionalExternalTarget;
+	bool UseExternalTarget;	
+	bool OnlyTriggerOnce;
+
+	FDelegateHandle DelegateHandle;
+};
\ No newline at end of file
diff --git a/Engine/Source/Runtime/GameplayAbilities/Public/AbilitySystemBlueprintLibrary.h b/Engine/Source/Runtime/GameplayAbilities/Public/AbilitySystemBlueprintLibrary.h
--- a/Engine/Source/Runtime/GameplayAbilities/Public/AbilitySystemBlueprintLibrary.h
+++ b/Engine/Source/Runtime/GameplayAbilities/Public/AbilitySystemBlueprintLibrary.h
@@ -190,4 +190,19 @@ class GAMEPLAYABILITIES_API UAbilitySystemBlueprintLibrary : public UBlueprintFu
 	UFUNCTION(BlueprintCallable, Category = "Ability|GameplayEffect")
 	static FGameplayEffectSpecHandle AddLinkedGameplayEffectSpec(FGameplayEffectSpecHandle SpecHandle, FGameplayEffectSpecHandle LinkedGameplayEffectSpec);
 
+	/** Sets the GameplayEffectSpec's StackCount to the specified amount (prior to applying) */
+	UFUNCTION(BlueprintCallable, Category = "Ability|GameplayEffect")
+	static FGameplayEffectSpecHandle SetStackCount(FGameplayEffectSpecHandle SpecHandle, int32 StackCount);
+
+	/** Sets the GameplayEffectSpec's StackCount to the max stack count defined in the GameplayEffect definition */
+	UFUNCTION(BlueprintCallable, Category = "Ability|GameplayEffect")
+	static FGameplayEffectSpecHandle SetStackCountToMax(FGameplayEffectSpecHandle SpecHandle);
+
+	// -------------------------------------------------------------------------------
+	//		FActiveGameplayEffectHandle
+	// -------------------------------------------------------------------------------
+	
+	/** Returns current stack count of an active Gameplay Effect. Will return 0 if the GameplayEffect is no longer valid. */
+	UFUNCTION(BlueprintCallable, Category = "Ability|GameplayEffect")
+	static int32 GetActiveGameplayEffectStackCount(FActiveGameplayEffectHandle ActiveHandle);
 };
diff --git a/Engine/Source/Runtime/GameplayAbilities/Public/AbilitySystemComponent.h b/Engine/Source/Runtime/GameplayAbilities/Public/AbilitySystemComponent.h
--- a/Engine/Source/Runtime/GameplayAbilities/Public/AbilitySystemComponent.h
+++ b/Engine/Source/Runtime/GameplayAbilities/Public/AbilitySystemComponent.h
@@ -56,12 +56,6 @@ class GAMEPLAYABILITIES_API UAbilitySystemComponent : public UActorComponent, pu
 	/** Used to register callbacks to confirm/cancel input */
 	DECLARE_DYNAMIC_MULTICAST_DELEGATE(FAbilityConfirmOrCancel);
 
-	/** Used to notify ability state tasks that a state is being ended - should not be used directly */
-	DECLARE_MULTICAST_DELEGATE_OneParam(FAbilityStateEnded, FName);
-
-	/** Used to notify ability state tasks that an interrupt is occurring - should not be used directly */
-	DECLARE_MULTICAST_DELEGATE(FAbilityStateInterrupted);
-
 	friend struct FActiveGameplayEffectAction_Add;
 	friend FGameplayEffectSpec;
 	friend class AAbilitySystemDebugHUD;
@@ -277,6 +271,9 @@ class GAMEPLAYABILITIES_API UAbilitySystemComponent : public UActorComponent, pu
 	UFUNCTION(BlueprintCallable, Category = GameplayEffects)
 	float GetGameplayEffectMagnitude(FActiveGameplayEffectHandle Handle, FGameplayAttribute Attribute) const;
 
+	/** Returns current stack count of an already applied GE */
+	int32 GetCurrentStackCount(FActiveGameplayEffectHandle Handle) const;
+
 	/**
 	 * Get the source tags from the gameplay spec represented by the specified handle, if possible
 	 * 
@@ -494,11 +491,14 @@ class GAMEPLAYABILITIES_API UAbilitySystemComponent : public UActorComponent, pu
 	 */
 
 	/** Grants Ability. Returns handle that can be used in TryActivateAbility, etc. */
-	FGameplayAbilitySpecHandle GiveAbility(FGameplayAbilitySpec AbilitySpec);
+	virtual FGameplayAbilitySpecHandle GiveAbility(FGameplayAbilitySpec AbilitySpec);
 
 	/** Grants an ability and attempts to activate it exactly one time, which will cause it to be removed. Only valid on the server! */
 	FGameplayAbilitySpecHandle GiveAbilityAndActivateOnce(FGameplayAbilitySpec AbilitySpec);
 
+	/** Sets an ability spec to remove when its finished. If the spec is not currently active, it terminates it immediately. Also clears InputID of the Spec. */
+	void SetRemoveAbilityOnEnd(FGameplayAbilitySpecHandle AbilitySpecHandle);
+
 	// Gets all Activatable Gameplay Abilities that match all tags in GameplayTagContainer AND for which
 	// DoesAbilitySatisfyTagRequirements() is true.  The latter requirement allows this function to find the correct
 	// ability without requiring advanced knowledge.  For example, if there are two "Melee" abilities, one of which
@@ -532,14 +532,8 @@ class GAMEPLAYABILITIES_API UAbilitySystemComponent : public UActorComponent, pu
 	// Ability Cancelling/Interrupts
 	// --------------------------------------------
 
-	/** Used by the ability state task to handle when a state is ended - should not use directly */
-	FAbilityStateEnded OnAbilityStateEnded;
-
-	/** Used by the ability state task to handle when a state is interrupted - should not use directly */
-	FAbilityStateInterrupted OnAbilityStateInterrupted;
-
-	/** Cancels the specified ability. */
-	void CancelAbility(UGameplayAbility* Ability);
+	/** Cancels the specified ability CDO. */
+	void CancelAbility(UGameplayAbility* Ability);	
 
 	/** Cancel all abilities with the specified tags. Will not cancel the Ignore instance */
 	void CancelAbilities(const FGameplayTagContainer* WithTags=nullptr, const FGameplayTagContainer* WithoutTags=nullptr, UGameplayAbility* Ignore=nullptr);
@@ -576,7 +570,12 @@ class GAMEPLAYABILITIES_API UAbilitySystemComponent : public UActorComponent, pu
 	UPROPERTY()
 	TArray<UGameplayAbility*>	AllReplicatedInstancedAbilities;
 
-	void NotifyAbilityEnded(FGameplayAbilitySpecHandle Handle, UGameplayAbility* Ability);
+	virtual void NotifyAbilityEnded(FGameplayAbilitySpecHandle Handle, UGameplayAbility* Ability);
+
+	void IncrementAbilityListLock();
+	void DecrementAbilityListLock();
+
+protected:
 
 	/**
 	 *	The abilities we can activate. 
@@ -587,13 +586,12 @@ class GAMEPLAYABILITIES_API UAbilitySystemComponent : public UActorComponent, pu
 	 *	without an AbilitySystemComponent. For example an ability could be written to execute on a StaticMeshActor. As long as the ability doesn't require 
 	 *	instancing or anything else that the AbilitySystemComponent would provide, then it doesn't need the component to function.
 	 */
-protected:
 
 	UPROPERTY(ReplicatedUsing=OnRep_ActivateAbilities, BlueprintReadOnly, Category = "Abilities")
 	FGameplayAbilitySpecContainer	ActivatableAbilities;
 
 	/** Maps from an ability spec to the target data. Used to track replicated data and callbacks */
-	TMap<FGameplayAbilitySpecHandleAndPredictionKey, FAbilityClientDataCache> AbilityClientDataMap;
+	TMap<FGameplayAbilitySpecHandleAndPredictionKey, FAbilityReplicatedDataCache> AbilityTargetDataMap;
 
 	/** Will be called from GiveAbility or from OnRep. Initializes events (triggers and inputs) with the given ability */
 	virtual void OnGiveAbility(FGameplayAbilitySpec& AbilitySpec);
@@ -610,6 +608,10 @@ protected:
 	/** Creates a new instance of an ability, storing it in the spec */
 	UGameplayAbility* CreateNewInstanceOfAbility(FGameplayAbilitySpec& Spec, const UGameplayAbility* Ability);
 
+	int32 AbilityScopeLockCount;
+	TArray<FGameplayAbilitySpecHandle, TInlineAllocator<2> > AbilityPendingRemoves;
+	TArray<FGameplayAbilitySpec, TInlineAllocator<2> > AbilityPendingAdds;
+
 public:
 
 	/** Returns the list of all activatable abilities */
@@ -633,11 +635,11 @@ public:
 	UFUNCTION(Server, reliable, WithValidation)
 	void	ServerTryActivateAbilityWithEventData(FGameplayAbilitySpecHandle AbilityToActivate, FPredictionKey PredictionKey, FGameplayEventData TriggerEventData);
 
-	/** Replicate that an ability has ended, to the client or server as appropriate */
-	void	ReplicateEndAbility(FGameplayAbilitySpecHandle Handle, FGameplayAbilityActivationInfo ActivationInfo, UGameplayAbility* Ability);
+	/** Replicate that an ability has ended/canceled, to the client or server as appropriate */
+	void	ReplicateEndOrCancelAbility(FGameplayAbilitySpecHandle Handle, FGameplayAbilityActivationInfo ActivationInfo, UGameplayAbility* Ability, bool bWasCanceled);
 
-	/** Called by ServerEndAbility and ClientEndAbility; avoids code duplication. */
-	void	RemoteEndAbility(FGameplayAbilitySpecHandle AbilityToEnd, FGameplayAbilityActivationInfo ActivationInfo);
+	/** Called by ServerEndAbility/ServerCancelAbility and ClientEndAbility/ClientCancelAbility; avoids code duplication. */
+	void	RemoteEndOrCancelAbility(FGameplayAbilitySpecHandle AbilityToEnd, FGameplayAbilityActivationInfo ActivationInfo, bool bWasCanceled);
 
 	UFUNCTION(Server, reliable, WithValidation)
 	void	ServerEndAbility(FGameplayAbilitySpecHandle AbilityToEnd, FGameplayAbilityActivationInfo ActivationInfo, FPredictionKey PredictionKey);
@@ -645,6 +647,12 @@ public:
 	UFUNCTION(Client, reliable)
 	void	ClientEndAbility(FGameplayAbilitySpecHandle AbilityToEnd, FGameplayAbilityActivationInfo ActivationInfo);
 
+	UFUNCTION(Server, reliable, WithValidation)
+	void    ServerCancelAbility(FGameplayAbilitySpecHandle AbilityToCancel, FGameplayAbilityActivationInfo ActivationInfo);
+
+	UFUNCTION(Client, reliable)
+	void    ClientCancelAbility(FGameplayAbilitySpecHandle AbilityToCancel, FGameplayAbilityActivationInfo ActivationInfo);
+
 	UFUNCTION(Client, Reliable)
 	void	ClientActivateAbilityFailed(FGameplayAbilitySpecHandle AbilityToActivate, int16 PredictionKey);
 
@@ -677,17 +685,19 @@ public:
 	
 	virtual void BindAbilityActivationToInputComponent(UInputComponent* InputComponent, FGameplayAbiliyInputBinds BindInfo);
 
-	void AbilityLocalInputPressed(int32 InputID);
+	virtual void AbilityLocalInputPressed(int32 InputID);
 
-	void AbilityLocalInputReleased(int32 InputID);
+	virtual void AbilityLocalInputReleased(int32 InputID);
 
 	void TryActivateInputHeldAbilities();
+	
+	virtual void LocalInputConfirm();
+	
+	virtual void LocalInputCancel();
 
-	UFUNCTION(BlueprintCallable, Category="Abilities")
-	void LocalInputConfirm();
-
-	UFUNCTION(BlueprintCallable, Category="Abilities")
-	void LocalInputCancel();
+	/** InputID for binding GenericConfirm/Cancel events */
+	int32 GenericConfirmInputID;
+	int32 GenericCancelInputID;
 
 	/** Generic local callback for generic ConfirmEvent that any ability can listen to */
 	FAbilityConfirmOrCancel	GenericLocalConfirmCallbacks;
@@ -862,7 +872,7 @@ public:
 	 *	
 	 *	There are a few ways active abilities on clients communicate with the server through the AbilitySystemComponent.
 	 *	1. Through TargetData events. Either that client sends TargetData via ServerSetReplicatedTargetData or cancels it via ServerSetReplicatedTargetDataCancelled.
-	 *	2. Through Generic Client Events. These are non payload carrying events defined by the EAbilityReplicatedClientEvent enum.
+	 *	2. Through Generic Replicated Events. These are non payload carrying events defined by the EAbilityReplicatedClientEvent enum.
 	 *	
 	 *	The general pattern is:
 	 *		-Clients and server register with these events via AbilityTargetDataSetDelegate / AbilityReplicatedEventDelegate.
@@ -871,6 +881,7 @@ public:
 	 *		-The server may get the replicated event before the server is actually listening for it! We store whether the event was recevied or not, for the given ability spec/prediction key pair.
 	 *			-Use functions like CallReplicatedEventDelegateIfSet/CallReplicatedTargetDataDelegatesIfSet after registering the delegates on the server. This invokes the delegates if they were already received.
 	 *			-Alternatively, use CallOrAddReplicatedDelegate to do it in one fell swoop (not good if you need multiple callbacks though).
+	 *		-Note we also support server -> client Generic Replicated Event via ClientSetReplicatedEvent.
 	 *	
 	 *	3. Through direct input replication (assuming you are using the provided input binding code in AbilitySystemComponent). 
 	 *	   This is a per ability setting defined by UGameplayAbility::bReplicateInputDirectly. It is generally NOT a good idea too use this! It is better to use the above Generic Client Eveents.
@@ -879,16 +890,20 @@ public:
 	 *	   and ultimately pipe down in UGameplayAbility::InputPressed/Released.
 	 *	   
 	 *	   This is useful if you truly need to maintain tracked input state on the server. For example if your ability logic requires polling of input state and arbitrary times (anim events, touched event, etc).
-	 *	   If your ability just requires event based input (button was pressed, released) then it is better to use Generic Client Events which will only replicate when the ability is listening for those events.
+	 *	   If your ability just requires event based input (button was pressed, released) then it is better to use Generic Replicated Events which will only replicate when the ability is listening for those events.
 	 *		
 	 */
 
-	/** Replicates the Generic CLient Event to the server. */
+	/** Replicates the Generic Replicated Event to the server. */
 	UFUNCTION(Server, reliable, WithValidation)
-	void ServerSetReplicatedClientEvent(EAbilityReplicatedClientEvent::Type EventType, FGameplayAbilitySpecHandle AbilityHandle, FPredictionKey AbilityOriginalPredictionKey, FPredictionKey CurrentPredictionKey);
+	void ServerSetReplicatedEvent(EAbilityGenericReplicatedEvent::Type EventType, FGameplayAbilitySpecHandle AbilityHandle, FPredictionKey AbilityOriginalPredictionKey, FPredictionKey CurrentPredictionKey);
+
+	/** Replicates the Generic Replicated Event to the server. */
+	UFUNCTION(Client, reliable)
+	void ClientSetReplicatedEvent(EAbilityGenericReplicatedEvent::Type EventType, FGameplayAbilitySpecHandle AbilityHandle, FPredictionKey AbilityOriginalPredictionKey);
 
-	/** Calls local callbacks that are registered with the given Generic Client Event */
-	void InvokeReplicatedClientEvent(EAbilityReplicatedClientEvent::Type EventType, FGameplayAbilitySpecHandle AbilityHandle, FPredictionKey AbilityOriginalPredictionKey);
+	/** Calls local callbacks that are registered with the given Generic Replicated Event */
+	bool InvokeReplicatedEvent(EAbilityGenericReplicatedEvent::Type EventType, FGameplayAbilitySpecHandle AbilityHandle, FPredictionKey AbilityOriginalPredictionKey);
 
 	/**  */
 	UFUNCTION(Server, reliable, WithValidation)
@@ -905,25 +920,25 @@ public:
 	void CancelAbilityTargetData(FGameplayAbilitySpecHandle AbilityHandle, FPredictionKey AbilityOriginalPredictionKey);
 
 	/** Deletes all cached ability client data (Was: ConsumeAbilityTargetData)*/
-	void ConsumeAllClientReplicatedData(FGameplayAbilitySpecHandle AbilityHandle, FPredictionKey AbilityOriginalPredictionKey);
+	void ConsumeAllReplicatedData(FGameplayAbilitySpecHandle AbilityHandle, FPredictionKey AbilityOriginalPredictionKey);
 
 	/** Consumes cached TargetData from client (only TargetData) */
 	void ConsumeClientReplicatedTargetData(FGameplayAbilitySpecHandle AbilityHandle, FPredictionKey AbilityOriginalPredictionKey);
 
-	/** Consumes the given Generic Client Event (unsets it). */
-	void ConsumeClientReplicatedEvent(EAbilityReplicatedClientEvent::Type EventType, FGameplayAbilitySpecHandle AbilityHandle, FPredictionKey AbilityOriginalPredictionKey);
+	/** Consumes the given Generic Replicated Event (unsets it). */
+	void ConsumeGenericReplicatedEvent(EAbilityGenericReplicatedEvent::Type EventType, FGameplayAbilitySpecHandle AbilityHandle, FPredictionKey AbilityOriginalPredictionKey);
 	
-	/** Calls any ClientReplicated delegates that have been sent (TargetData or Generic Client Events). Note this can be dangerous if multiple places in an ability register events and then call this function. */
+	/** Calls any Replicated delegates that have been sent (TargetData or Generic Replicated Events). Note this can be dangerous if multiple places in an ability register events and then call this function. */
 	void CallAllReplicatedDelegatesIfSet(FGameplayAbilitySpecHandle AbilityHandle, FPredictionKey AbilityOriginalPredictionKey);
 
 	/** Calls the TargetData Confirm/Cancel events if they have been sent. */
 	bool CallReplicatedTargetDataDelegatesIfSet(FGameplayAbilitySpecHandle AbilityHandle, FPredictionKey AbilityOriginalPredictionKey);
 
-	/** Calls a given Generic Client Event delegate if the event has already been sent */
-	bool CallReplicatedEventDelegateIfSet(EAbilityReplicatedClientEvent::Type EventType, FGameplayAbilitySpecHandle AbilityHandle, FPredictionKey AbilityOriginalPredictionKey);
+	/** Calls a given Generic Replicated Event delegate if the event has already been sent */
+	bool CallReplicatedEventDelegateIfSet(EAbilityGenericReplicatedEvent::Type EventType, FGameplayAbilitySpecHandle AbilityHandle, FPredictionKey AbilityOriginalPredictionKey);
 
 	/** Calls passed in delegate if the Client Event has already been sent. If not, it adds the delegate to our multicast callback that will fire when it does. */
-	bool CallOrAddReplicatedDelegate(EAbilityReplicatedClientEvent::Type EventType, FGameplayAbilitySpecHandle AbilityHandle, FPredictionKey AbilityOriginalPredictionKey, FSimpleMulticastDelegate::FDelegate Delegate);
+	bool CallOrAddReplicatedDelegate(EAbilityGenericReplicatedEvent::Type EventType, FGameplayAbilitySpecHandle AbilityHandle, FPredictionKey AbilityOriginalPredictionKey, FSimpleMulticastDelegate::FDelegate Delegate);
 
 	/** Returns TargetDataSet delegate for a given Ability/PredictionKey pair */
 	FAbilityTargetDataSetDelegate& AbilityTargetDataSetDelegate(FGameplayAbilitySpecHandle AbilityHandle, FPredictionKey AbilityOriginalPredictionKey);
@@ -931,10 +946,10 @@ public:
 	/** Returns TargetData Cancelled delegate for a given Ability/PredictionKey pair */
 	FSimpleMulticastDelegate& AbilityTargetDataCancelledDelegate(FGameplayAbilitySpecHandle AbilityHandle, FPredictionKey AbilityOriginalPredictionKey);
 
-	/** Returns Generic Client Event for a given Ability/PredictionKey pair */
-	FSimpleMulticastDelegate& AbilityReplicatedEventDelegate(EAbilityReplicatedClientEvent::Type EventType, FGameplayAbilitySpecHandle AbilityHandle, FPredictionKey AbilityOriginalPredictionKey);
+	/** Returns Generic Replicated Event for a given Ability/PredictionKey pair */
+	FSimpleMulticastDelegate& AbilityReplicatedEventDelegate(EAbilityGenericReplicatedEvent::Type EventType, FGameplayAbilitySpecHandle AbilityHandle, FPredictionKey AbilityOriginalPredictionKey);
 
-	// Direct Input state replication. These will be called if bReplicateInputDirectly is true on the ability and is generally not a good thing to use. (Instead, prefer to use Generic Client Events).
+	// Direct Input state replication. These will be called if bReplicateInputDirectly is true on the ability and is generally not a good thing to use. (Instead, prefer to use Generic Replicated Events).
 	UFUNCTION(Server, reliable, WithValidation)
 	void ServerSetInputPressed(FGameplayAbilitySpecHandle AbilityHandle);
 
@@ -976,7 +991,8 @@ protected:
 	bool IsOwnerActorAuthoritative() const;
 
 	void OnAttributeGameplayEffectSpecExected(const FGameplayAttribute &Attribute, const struct FGameplayEffectSpec &Spec, struct FGameplayModifierEvaluatedData &Data);
-
+		
+	TArray<TWeakObjectPtr<UAbilityTask> >&	GetAbilityActiveTasks(UGameplayAbility* Ability);
 	// --------------------------------------------
 	
 	// Contains all of the gameplay effects that are currently active on this component
diff --git a/Engine/Source/Runtime/GameplayAbilities/Public/GameplayAbilitySpec.h b/Engine/Source/Runtime/GameplayAbilities/Public/GameplayAbilitySpec.h
new file mode 100644
--- /dev/null
+++ b/Engine/Source/Runtime/GameplayAbilities/Public/GameplayAbilitySpec.h
@@ -0,0 +1,318 @@
+// Copyright 1998-2015 Epic Games, Inc. All Rights Reserved.
+
+#pragma once
+
+#include "GameplayPrediction.h"
+#include "GameplayEffectTypes.h"
+#include "GameplayAbilitySpec.generated.h"
+
+/**
+ *	This file exists in addition so that GameplayEffect.h can use FGameplayAbilitySpec without having to include GameplayAbilityTypes.h which has depancies on
+ *	GameplayEffect.h
+ */
+
+USTRUCT(BlueprintType)
+struct FGameplayAbilitySpecHandle
+{
+	GENERATED_USTRUCT_BODY()
+
+	FGameplayAbilitySpecHandle()
+	: Handle(INDEX_NONE)
+	{
+
+	}
+
+	bool IsValid() const
+	{
+		return Handle != INDEX_NONE;
+	}
+
+	void GenerateNewHandle()
+	{
+		static int32 GHandle = 1;
+		Handle = GHandle++;
+	}
+
+	bool operator==(const FGameplayAbilitySpecHandle& Other) const
+	{
+		return Handle == Other.Handle;
+	}
+
+	bool operator!=(const FGameplayAbilitySpecHandle& Other) const
+	{
+		return Handle != Other.Handle;
+	}
+
+	friend uint32 GetTypeHash(const FGameplayAbilitySpecHandle& SpecHandle)
+	{
+		return ::GetTypeHash(SpecHandle.Handle);
+	}
+
+private:
+
+	UPROPERTY()
+	int32 Handle;
+};
+
+UENUM(BlueprintType)
+namespace EGameplayAbilityActivationMode
+{
+	enum Type
+	{
+		// We are the authority activating this ability
+		Authority,
+
+		// We are not the authority but aren't predicting yet. This is a mostly invalid state to be in.
+		NonAuthority,
+
+		// We are predicting the activation of this ability
+		Predicting,
+
+		// We are not the authority, but the authority has confirmed this activation
+		Confirmed,
+	};
+}
+
+/** Describes what happens when a GameplayEffect, that is granting an active ability, is removed from its owner. */
+UENUM()
+enum class EGameplayEffectGrantedAbilityRemovePolicy : uint8
+{
+	/** Active abilities are immediately canceled and the ability is removed. */
+	CancelAbilityImmediately,
+	/** Active abilities are allowed to finish, and then removed. */
+	RemoveAbilityOnEnd,
+	/** Granted abilties are left lone when the granting GameplayEffect is removed */
+	DoNothing,
+};
+
+/** This is data that can be used to create an FGameplayAbilitySpec. Has some data that is only relevant when granted by a GameplayEffect */
+USTRUCT(BlueprintType)
+struct FGameplayAbilitySpecDef
+{
+	FGameplayAbilitySpecDef()
+		: Level(1)
+		, InputID(INDEX_NONE)
+		, RemovalPolicy(EGameplayEffectGrantedAbilityRemovePolicy::CancelAbilityImmediately)
+	{
+	}
+
+	GENERATED_USTRUCT_BODY()
+
+	UPROPERTY(EditDefaultsOnly, Category="Ability Definition")
+	TSubclassOf<UGameplayAbility> Ability;
+	
+	UPROPERTY(EditDefaultsOnly, Category="Ability Definition")
+	int32 Level; 
+	
+	UPROPERTY(EditDefaultsOnly, Category="Ability Definition")
+	int32 InputID;
+
+	UPROPERTY(EditDefaultsOnly, Category="Ability Definition")
+	EGameplayEffectGrantedAbilityRemovePolicy RemovalPolicy;
+	
+	UPROPERTY()
+	UObject* SourceObject;
+
+	/** This handle can be set if the SpecDef is used to create a real FGameplaybilitySpec */
+	UPROPERTY()
+	FGameplayAbilitySpecHandle	AssignedHandle;
+};
+
+/**
+ *	FGameplayAbilityActivationInfo
+ *
+ *	Data tied to a specific activation of an ability.
+ *		-Tell us whether we are the authority, if we are predicting, confirmed, etc.
+ *		-Holds current and previous PredictionKey
+ *		-Generally not meant to be subclassed in projects.
+ *		-Passed around by value since the struct is small.
+ */
+USTRUCT(BlueprintType)
+struct GAMEPLAYABILITIES_API FGameplayAbilityActivationInfo
+{
+	GENERATED_USTRUCT_BODY()
+
+	FGameplayAbilityActivationInfo()
+		: ActivationMode(EGameplayAbilityActivationMode::Authority)
+		, bCanBeEndedByOtherInstance(false)
+	{
+
+	}
+
+	FGameplayAbilityActivationInfo(AActor* InActor)
+		: bCanBeEndedByOtherInstance(false)	
+	{
+		// On Init, we are either Authority or NonAuthority. We haven't been given a PredictionKey and we haven't been confirmed.
+		// NonAuthority essentially means 'I'm not sure what how I'm going to do this yet'.
+		ActivationMode = (InActor->Role == ROLE_Authority ? EGameplayAbilityActivationMode::Authority : EGameplayAbilityActivationMode::NonAuthority);
+	}
+
+	FGameplayAbilityActivationInfo(EGameplayAbilityActivationMode::Type InType)
+		: ActivationMode(InType)
+		, bCanBeEndedByOtherInstance(false)
+	{
+	}	
+
+	UPROPERTY(BlueprintReadOnly, Category = "ActorInfo")
+	mutable TEnumAsByte<EGameplayAbilityActivationMode::Type>	ActivationMode;
+
+	/** An ability that runs on multiple game instances can be canceled by a remote instance, but only if that remote instance has already confirmed starting it. */
+	UPROPERTY()
+	bool bCanBeEndedByOtherInstance;
+
+	void SetActivationConfirmed();
+
+	/** Called on client to set this as a predicted ability */
+	void SetPredicting(FPredictionKey PredictionKey);
+
+	/** Called on the server to set the key used by the client to predict this ability */
+	void ServerSetActivationPredictionKey(FPredictionKey PredictionKey);
+
+	FPredictionKey GetActivationPredictionKey() { return PredictionKeyWhenActivated; }
+
+private:
+
+	// This was the prediction key used to activate this ability. It does not get updated
+	// if new prediction keys are generated over the course of the ability.
+	UPROPERTY()
+	FPredictionKey PredictionKeyWhenActivated;
+	
+};
+
+
+/** An activatable ability spec, hosted on the ability system component. This defines both what the ability is (what class, what level, input binding etc)
+ *  and also holds runttime state that must be kept outside of the ability being instanced/activated.
+ */
+USTRUCT()
+struct GAMEPLAYABILITIES_API FGameplayAbilitySpec : public FFastArraySerializerItem
+{
+	GENERATED_USTRUCT_BODY()
+
+	FGameplayAbilitySpec()
+	: Ability(nullptr), Level(1), InputID(INDEX_NONE), SourceObject(nullptr), InputPressed(false), ActiveCount(0), RemoveAfterActivation(false), PendingRemove(false)
+	{
+		
+	}
+
+	FGameplayAbilitySpec(UGameplayAbility* InAbility, int32 InLevel=1, int32 InInputID=INDEX_NONE, UObject* InSourceObject=nullptr)
+		: Ability(InAbility), Level(InLevel), InputID(InInputID), SourceObject(InSourceObject), InputPressed(false), ActiveCount(0), RemoveAfterActivation(false),  PendingRemove(false)
+	{
+		Handle.GenerateNewHandle();
+	}
+	
+	FGameplayAbilitySpec(FGameplayAbilitySpecDef& InDef, FActiveGameplayEffectHandle InGameplayEffectHandle = FActiveGameplayEffectHandle());
+
+	/** Handle for outside sources to refer to this spec by */
+	UPROPERTY()
+	FGameplayAbilitySpecHandle Handle;
+	
+	/** Ability of the spec (Always the CDO. This should be const but too many things modify it currently) */
+	UPROPERTY()
+	UGameplayAbility* Ability;
+	
+	/** Level of Ability */
+	UPROPERTY()
+	int32	Level;
+
+	/** InputID, if bound */
+	UPROPERTY()
+	int32	InputID;
+
+	/** Object this ability was created from, can be an actor or static object. Useful to bind an ability to a gameplay object */
+	UPROPERTY()
+	UObject* SourceObject;
+
+	/** Is input currently pressed. Set to false when input is released */
+	UPROPERTY(NotReplicated)
+	bool	InputPressed;
+
+	/** A count of the number of times this ability has been activated minus the number of times it has been ended. For instanced abilities this will be the number of currently active instances. Can't replicate until prediction accurately handles this.*/
+	UPROPERTY(NotReplicated)
+	uint8	ActiveCount;
+
+	/** If true, this ability should be removed as soon as it finishes executing */
+	UPROPERTY(NotReplicated)
+	bool	RemoveAfterActivation;
+
+	/** Activation state of this ability. This is not replicated since it needs to be overwritten locally on clients during prediction. */
+	UPROPERTY(NotReplicated)
+	FGameplayAbilityActivationInfo	ActivationInfo;
+
+	/** Non replicating instances of this ability. */
+	UPROPERTY(NotReplicated)
+	TArray<UGameplayAbility*> NonReplicatedInstances;
+
+	/** Replicated instances of this ability.. */
+	UPROPERTY()
+	TArray<UGameplayAbility*> ReplicatedInstances;
+
+	/** Handle to GE that granted us (usually invalid) */
+	UPROPERTY(NotReplicated)
+	FActiveGameplayEffectHandle	GameplayEffectHandle;
+
+	/** Pending removal due to scope lock */
+	UPROPERTY(NotReplicated)
+	bool	PendingRemove;
+
+	/** Returns the primary instance, used for instance once abilities */
+	UGameplayAbility* GetPrimaryInstance() const;
+
+	/** Returns all instances, which can include instance per execution abilities */
+	TArray<UGameplayAbility*> GetAbilityInstances() const
+	{
+		TArray<UGameplayAbility*> Abilities;
+		Abilities.Append(ReplicatedInstances);
+		Abilities.Append(NonReplicatedInstances);
+		return Abilities;
+	}
+
+	/** Returns true if this ability is active in any way */
+	bool IsActive() const;
+
+	void PreReplicatedRemove(const struct FGameplayAbilitySpecContainer& InArraySerializer);
+	void PostReplicatedAdd(const struct FGameplayAbilitySpecContainer& InArraySerializer);
+};
+
+
+/** Fast serializer wrapper for above struct */
+USTRUCT()
+struct GAMEPLAYABILITIES_API FGameplayAbilitySpecContainer : public FFastArraySerializer
+{
+	GENERATED_USTRUCT_BODY()
+
+	/** List of activatable abilities */
+	UPROPERTY()
+	TArray<FGameplayAbilitySpec> Items;
+
+	/** Component that owns this list */
+	UAbilitySystemComponent* Owner;
+
+	void RegisterWithOwner(UAbilitySystemComponent* Owner);
+
+	bool NetDeltaSerialize(FNetDeltaSerializeInfo & DeltaParms)
+	{
+		return FFastArraySerializer::FastArrayDeltaSerialize<FGameplayAbilitySpec, FGameplayAbilitySpecContainer>(Items, DeltaParms, *this);
+	}
+};
+
+template<>
+struct TStructOpsTypeTraits< FGameplayAbilitySpecContainer > : public TStructOpsTypeTraitsBase
+{
+	enum
+	{
+		WithNetDeltaSerializer = true,
+	};
+};
+
+// ----------------------------------------------------
+
+struct GAMEPLAYABILITIES_API FScopedAbilityListLock
+{
+	FScopedAbilityListLock(UAbilitySystemComponent& InContainer);
+	~FScopedAbilityListLock();
+
+private:
+	UAbilitySystemComponent& AbilitySystemComponent;
+};
+
+#define ABILITYLIST_SCOPE_LOCK()	FScopedAbilityListLock ActiveScopeLock(*this);
\ No newline at end of file
diff --git a/Engine/Source/Runtime/GameplayAbilities/Public/GameplayEffect.h b/Engine/Source/Runtime/GameplayAbilities/Public/GameplayEffect.h
--- a/Engine/Source/Runtime/GameplayAbilities/Public/GameplayEffect.h
+++ b/Engine/Source/Runtime/GameplayAbilities/Public/GameplayEffect.h
@@ -7,6 +7,7 @@
 #include "GameplayTagContainer.h"
 #include "TimerManager.h"
 #include "GameplayEffectTypes.h"
+#include "GameplayAbilitySpec.h"
 #include "GameplayEffectAggregator.h"
 #include "GameplayEffectCalculation.h"
 #include "ActiveGameplayEffectIterator.h"
@@ -525,11 +526,11 @@ public:
 
 #if WITH_EDITORONLY_DATA
 	/** Template to derive starting values and editing customization from */
-	UPROPERTY(EditDefaultsOnly, BlueprintReadOnly, Category = Template)
+	UPROPERTY()
 	UGameplayEffectTemplate*	Template;
 
 	/** When false, show a limited set of properties for editing, based on the template we are derived from */
-	UPROPERTY(EditDefaultsOnly, BlueprintReadOnly, Category = Template)
+	UPROPERTY()
 	bool ShowAllProperties;
 #endif
 
@@ -680,6 +681,13 @@ public:
 	/** Policy for how the effect period should be reset (or not) while stacking */
 	UPROPERTY(EditDefaultsOnly, BlueprintReadOnly, Category = Stacking)
 	EGameplayEffectStackingPeriodPolicy StackPeriodResetPolicy;
+
+	// ----------------------------------------------------------------------
+	//	Granted abilities
+	// ----------------------------------------------------------------------
+	
+	UPROPERTY(EditDefaultsOnly, BlueprintReadOnly, Category = "Granted Abilities")
+	TArray<FGameplayAbilitySpecDef>	GrantedAbilities;
 };
 
 /** Holds evaluated magnitude from a GameplayEffect modifier */
@@ -1088,6 +1096,9 @@ public:
 	UPROPERTY(NotReplicated)
 	uint32 bDurationLocked : 1;
 
+	UPROPERTY(NotReplicated)
+	TArray<FGameplayAbilitySpecDef> GrantedAbilitySpecs;
+
 private:
 
 	/** Map of set by caller magnitudes */
diff --git a/Engine/Source/Runtime/GameplayAbilities/Public/GameplayEffectTypes.h b/Engine/Source/Runtime/GameplayAbilities/Public/GameplayEffectTypes.h
--- a/Engine/Source/Runtime/GameplayAbilities/Public/GameplayEffectTypes.h
+++ b/Engine/Source/Runtime/GameplayAbilities/Public/GameplayEffectTypes.h
@@ -18,6 +18,7 @@
 
 class UAbilitySystemComponent;
 class UGameplayEffect;
+class UGameplayAbility;
 
 struct FGameplayEffectSpec;
 struct FGameplayEffectModCallbackData;
@@ -734,7 +735,7 @@ DECLARE_DELEGATE_RetVal_OneParam(FOnGameplayEffectTagCountChanged&, FRegisterGam
  * while simultaneously tracking the count of parent tags as well. Events/delegates are fired whenever the tag counts
  * of any tag (explicit or parent) are modified.
  */
-struct FGameplayTagCountContainer
+struct GAMEPLAYABILITIES_API FGameplayTagCountContainer
 {
 	// Constructor
 	FGameplayTagCountContainer()
@@ -834,11 +835,11 @@ struct GAMEPLAYABILITIES_API FGameplayTagRequirements
 	GENERATED_USTRUCT_BODY()
 
 	/** All of these tags must be present */
-	UPROPERTY(EditAnywhere, Category = GameplayModifier)
+	UPROPERTY(EditAnywhere, BlueprintReadWrite, Category = GameplayModifier)
 	FGameplayTagContainer RequireTags;
 
 	/** None of these tags may be present */
-	UPROPERTY(EditAnywhere, Category = GameplayModifier)
+	UPROPERTY(EditAnywhere, BlueprintReadWrite, Category = GameplayModifier)
 	FGameplayTagContainer IgnoreTags;
 
 	bool	RequirementsMet(const FGameplayTagContainer& Container) const;
@@ -976,4 +977,4 @@ struct TStructOpsTypeTraits<FGameplayEffectSpecHandle> : public TStructOpsTypeTr
 		WithNetSerializer = true,
 		WithIdenticalViaEquality = true,
 	};
-};
+};
\ No newline at end of file
