Merging

//depot/UE4-Orion/Engine/Source/Runtime/GameplayAbilities/Public/...

to //depot/UE4/Engine/Source/Runtime/GameplayAbilities/Public/...

diff --git a/Engine/Source/Runtime/GameplayAbilities/Public/GameplayAbilitySpec.h b/Engine/Source/Runtime/GameplayAbilities/Public/GameplayAbilitySpec.h
--- a/Engine/Source/Runtime/GameplayAbilities/Public/GameplayAbilitySpec.h
+++ b/Engine/Source/Runtime/GameplayAbilities/Public/GameplayAbilitySpec.h
@@ -98,19 +98,19 @@ struct FGameplayAbilitySpecDef
 
 	GENERATED_USTRUCT_BODY()
 
-	UPROPERTY(EditDefaultsOnly, Category="Ability Definition")
+	UPROPERTY(EditDefaultsOnly, Category="Ability Definition", NotReplicated)
 	TSubclassOf<UGameplayAbility> Ability;
 	
-	UPROPERTY(EditDefaultsOnly, Category="Ability Definition")
+	UPROPERTY(EditDefaultsOnly, Category="Ability Definition", NotReplicated)
 	int32 Level; 
 	
-	UPROPERTY(EditDefaultsOnly, Category="Ability Definition")
+	UPROPERTY(EditDefaultsOnly, Category="Ability Definition", NotReplicated)
 	int32 InputID;
 
-	UPROPERTY(EditDefaultsOnly, Category="Ability Definition")
+	UPROPERTY(EditDefaultsOnly, Category="Ability Definition", NotReplicated)
 	EGameplayEffectGrantedAbilityRemovePolicy RemovalPolicy;
 	
-	UPROPERTY()
+	UPROPERTY(NotReplicated)
 	UObject* SourceObject;
 
 	/** This handle can be set if the SpecDef is used to create a real FGameplaybilitySpec */
diff --git a/Engine/Source/Runtime/GameplayAbilities/Public/GameplayEffect.h b/Engine/Source/Runtime/GameplayAbilities/Public/GameplayEffect.h
--- a/Engine/Source/Runtime/GameplayAbilities/Public/GameplayEffect.h
+++ b/Engine/Source/Runtime/GameplayAbilities/Public/GameplayEffect.h
@@ -1096,7 +1096,7 @@ public:
 	UPROPERTY(NotReplicated)
 	uint32 bDurationLocked : 1;
 
-	UPROPERTY(NotReplicated)
+	UPROPERTY()
 	TArray<FGameplayAbilitySpecDef> GrantedAbilitySpecs;
 
 private:
