Integration(2467881): Removed ability input held support

* Removed support for reactivating an ability if the input button is held down.  This was game specific and moved to orion code.

* Added utility functions for checking confirm and cancel ability input.

diff --git a/Engine/Source/Runtime/GameplayAbilities/Private/Abilities/GameplayAbility.cpp b/Engine/Source/Runtime/GameplayAbilities/Private/Abilities/GameplayAbility.cpp
--- a/Engine/Source/Runtime/GameplayAbilities/Private/Abilities/GameplayAbility.cpp
+++ b/Engine/Source/Runtime/GameplayAbilities/Private/Abilities/GameplayAbility.cpp
@@ -52,7 +52,6 @@ UGameplayAbility::UGameplayAbility(const FObjectInitializer& ObjectInitializer)
 #endif
 
 	bServerRespectsRemoteAbilityCancellation = true;
-	bActivateOnInputHeld = false;
 	bReplicateInputDirectly = false;
 }
 
diff --git a/Engine/Source/Runtime/GameplayAbilities/Private/AbilitySystemComponent.cpp b/Engine/Source/Runtime/GameplayAbilities/Private/AbilitySystemComponent.cpp
--- a/Engine/Source/Runtime/GameplayAbilities/Private/AbilitySystemComponent.cpp
+++ b/Engine/Source/Runtime/GameplayAbilities/Private/AbilitySystemComponent.cpp
@@ -770,14 +770,38 @@ float UAbilitySystemComponent::GetGameplayEffectMagnitude(FActiveGameplayEffectH
 
 int32 UAbilitySystemComponent::GetCurrentStackCount(FActiveGameplayEffectHandle Handle) const
 {
-	const FActiveGameplayEffect* ActiveGE = ActiveGameplayEffects.GetActiveGameplayEffect(Handle);
-	if (ActiveGE)
+	if (const FActiveGameplayEffect* ActiveGE = ActiveGameplayEffects.GetActiveGameplayEffect(Handle))
 	{
 		return ActiveGE->Spec.StackCount;
 	}
 	return 0;
 }
 
+int32 UAbilitySystemComponent::GetCurrentStackCount(FGameplayAbilitySpecHandle Handle) const
+{
+	FActiveGameplayEffectHandle GEHandle = FindActiveGameplayEffectHandle(Handle);
+	if (GEHandle.IsValid())
+	{
+		return GetCurrentStackCount(GEHandle);
+	}
+	return 0;
+}
+
+FActiveGameplayEffectHandle UAbilitySystemComponent::FindActiveGameplayEffectHandle(FGameplayAbilitySpecHandle Handle) const
+{
+	for (const FActiveGameplayEffect& ActiveGE : &ActiveGameplayEffects)
+	{
+		for (const FGameplayAbilitySpecDef& AbiilitySpecDef : ActiveGE.Spec.GrantedAbilitySpecs)
+		{
+			if (AbiilitySpecDef.AssignedHandle == Handle)
+			{
+				return ActiveGE.Handle;
+			}
+		}
+	}
+	return FActiveGameplayEffectHandle();
+}
+
 void UAbilitySystemComponent::InvokeGameplayCueEvent(const FGameplayEffectSpecForRPC &Spec, EGameplayCueEvent::Type EventType)
 {
 	AActor* ActorAvatar = AbilityActorInfo->AvatarActor.Get();
diff --git a/Engine/Source/Runtime/GameplayAbilities/Private/AbilitySystemComponent_Abilities.cpp b/Engine/Source/Runtime/GameplayAbilities/Private/AbilitySystemComponent_Abilities.cpp
--- a/Engine/Source/Runtime/GameplayAbilities/Private/AbilitySystemComponent_Abilities.cpp
+++ b/Engine/Source/Runtime/GameplayAbilities/Private/AbilitySystemComponent_Abilities.cpp
@@ -1719,13 +1719,13 @@ void UAbilitySystemComponent::BindAbilityActivationToInputComponent(UInputCompon
 void UAbilitySystemComponent::AbilityLocalInputPressed(int32 InputID)
 {
 	// Consume the input if this InputID is overloaded with GenericConfirm/Cancel and the GenericConfim/Cancel callback is bound
-	if (InputID == GenericConfirmInputID && GenericLocalConfirmCallbacks.IsBound())
+	if (IsGenericConfirmInputBound(InputID))
 	{
 		LocalInputConfirm();
 		return;
 	}
 
-	if (InputID == GenericCancelInputID && GenericLocalCancelCallbacks.IsBound())
+	if (IsGenericCancelInputBound(InputID))
 	{
 		LocalInputCancel();
 		return;
@@ -1857,28 +1857,6 @@ void UAbilitySystemComponent::AbilitySpecInputReleased(FGameplayAbilitySpec& Spe
 	}
 }
 
-void UAbilitySystemComponent::TryActivateInputHeldAbilities()
-{
-	APlayerController* PC = AbilityActorInfo.IsValid() ? AbilityActorInfo->PlayerController.Get() : nullptr;
-
-	if (PC && PC->PlayerInput)
-	{
-		ABILITYLIST_SCOPE_LOCK();
-		for (FGameplayAbilitySpec& Spec : ActivatableAbilities.Items)
-		{
-			if ((Spec.InputID == INDEX_NONE) || !Spec.InputPressed || Spec.IsActive() || !Spec.Ability)
-			{
-				continue;
-			}
-
-			if (Spec.Ability->bActivateOnInputHeld)
-			{
-				TryActivateAbility(Spec.Handle);
-			}
-		}
-	}
-}
-
 void UAbilitySystemComponent::LocalInputConfirm()
 {
 	GenericLocalConfirmCallbacks.Broadcast();
diff --git a/Engine/Source/Runtime/GameplayAbilities/Public/Abilities/GameplayAbility.h b/Engine/Source/Runtime/GameplayAbilities/Public/Abilities/GameplayAbility.h
--- a/Engine/Source/Runtime/GameplayAbilities/Public/Abilities/GameplayAbility.h
+++ b/Engine/Source/Runtime/GameplayAbilities/Public/Abilities/GameplayAbility.h
@@ -257,10 +257,6 @@ public:
 	/** Input binding stub. */
 	virtual void InputReleased(const FGameplayAbilitySpecHandle Handle, const FGameplayAbilityActorInfo* ActorInfo, const FGameplayAbilityActivationInfo ActivationInfo) {};
 
-	/** If this is set, the ability will try to activate if the input is held. */
-	UPROPERTY(EditDefaultsOnly, Category = Input)
-	bool bActivateOnInputHeld;
-
 	/** If true, this ability will always replicate input press/release events to the server. */
 	UPROPERTY(EditDefaultsOnly, Category = Input)
 	bool bReplicateInputDirectly;
diff --git a/Engine/Source/Runtime/GameplayAbilities/Public/AbilitySystemComponent.h b/Engine/Source/Runtime/GameplayAbilities/Public/AbilitySystemComponent.h
--- a/Engine/Source/Runtime/GameplayAbilities/Public/AbilitySystemComponent.h
+++ b/Engine/Source/Runtime/GameplayAbilities/Public/AbilitySystemComponent.h
@@ -274,6 +274,12 @@ class GAMEPLAYABILITIES_API UAbilitySystemComponent : public UActorComponent, pu
 	/** Returns current stack count of an already applied GE */
 	int32 GetCurrentStackCount(FActiveGameplayEffectHandle Handle) const;
 
+	/** Returns current stack count of an already applied GE, but given the ability spec handle that was granted by the GE */
+	int32 GetCurrentStackCount(FGameplayAbilitySpecHandle Handle) const;
+
+	/** Gets the GE Handle of the GE that granted the passed in Ability */
+	FActiveGameplayEffectHandle FindActiveGameplayEffectHandle(FGameplayAbilitySpecHandle Handle) const;
+
 	/**
 	 * Get the source tags from the gameplay spec represented by the specified handle, if possible
 	 * 
@@ -686,19 +692,18 @@ public:
 	virtual void BindAbilityActivationToInputComponent(UInputComponent* InputComponent, FGameplayAbiliyInputBinds BindInfo);
 
 	virtual void AbilityLocalInputPressed(int32 InputID);
-
 	virtual void AbilityLocalInputReleased(int32 InputID);
-
-	void TryActivateInputHeldAbilities();
 	
 	virtual void LocalInputConfirm();
-	
 	virtual void LocalInputCancel();
 
 	/** InputID for binding GenericConfirm/Cancel events */
 	int32 GenericConfirmInputID;
 	int32 GenericCancelInputID;
 
+	bool IsGenericConfirmInputBound(int32 InputID) const	{ return ((InputID == GenericConfirmInputID) && GenericLocalConfirmCallbacks.IsBound()); }
+	bool IsGenericCancelInputBound(int32 InputID) const		{ return ((InputID == GenericCancelInputID) && GenericLocalCancelCallbacks.IsBound()); }
+
 	/** Generic local callback for generic ConfirmEvent that any ability can listen to */
 	FAbilityConfirmOrCancel	GenericLocalConfirmCallbacks;
 
