Fix VS2012 build (no chained constructors allowed!).

diff --git a/Engine/Source/Runtime/GameplayAbilities/Private/GameplayAbilityTypes.cpp b/Engine/Source/Runtime/GameplayAbilities/Private/GameplayAbilityTypes.cpp
--- a/Engine/Source/Runtime/GameplayAbilities/Private/GameplayAbilityTypes.cpp
+++ b/Engine/Source/Runtime/GameplayAbilities/Private/GameplayAbilityTypes.cpp
@@ -163,7 +163,14 @@ void FGameplayAbilitySpecContainer::RegisterWithOwner(UAbilitySystemComponent* O
 // ----------------------------------------------------
 
 FGameplayAbilitySpec:: FGameplayAbilitySpec(FGameplayAbilitySpecDef& InDef, FActiveGameplayEffectHandle InGameplayEffectHandle)
-	: FGameplayAbilitySpec(InDef.Ability ? InDef.Ability->GetDefaultObject<UGameplayAbility>() : nullptr, InDef.Level, InDef.InputID, InDef.SourceObject)
+	: Ability(InDef.Ability ? InDef.Ability->GetDefaultObject<UGameplayAbility>() : nullptr)
+	, Level(InDef.Level)
+	, InputID(InDef.InputID)
+	, SourceObject(InDef.SourceObject)
+	, InputPressed(false)
+	, ActiveCount(0)
+	, RemoveAfterActivation(false)
+	, PendingRemove(false)
 {
 	InDef.AssignedHandle = Handle;
 	GameplayEffectHandle = InGameplayEffectHandle;
