Integration(2468968): Fixed PlayMontageAndWait hang

* Fixed PlayMontageAndWait to end when the montage can't be played.

diff --git a/Engine/Source/Runtime/GameplayAbilities/Private/Abilities/Tasks/AbilityTask_PlayMontageAndWait.cpp b/Engine/Source/Runtime/GameplayAbilities/Private/Abilities/Tasks/AbilityTask_PlayMontageAndWait.cpp
--- a/Engine/Source/Runtime/GameplayAbilities/Private/Abilities/Tasks/AbilityTask_PlayMontageAndWait.cpp
+++ b/Engine/Source/Runtime/GameplayAbilities/Private/Abilities/Tasks/AbilityTask_PlayMontageAndWait.cpp
@@ -66,6 +66,8 @@ UAbilityTask_PlayMontageAndWait* UAbilityTask_PlayMontageAndWait::CreatePlayMont
 
 void UAbilityTask_PlayMontageAndWait::Activate()
 {
+	bool bPlayedMontage = false;
+
 	if (AbilitySystemComponent.IsValid() && Ability.IsValid())
 	{
 		const FGameplayAbilityActorInfo* ActorInfo = Ability->GetCurrentActorInfo();
@@ -76,6 +78,7 @@ void UAbilityTask_PlayMontageAndWait::Activate()
 				// Playing a montage could potentially fire off a callback into game code which could kill this ability! Early out if we are  pending kill.
 				if (IsPendingKill())
 				{
+					OnCancelled.Broadcast();
 					return;
 				}
 
@@ -83,13 +86,17 @@ void UAbilityTask_PlayMontageAndWait::Activate()
 
 				BlendingOutDelegate.BindUObject(this, &UAbilityTask_PlayMontageAndWait::OnMontageEnded);
 				ActorInfo->AnimInstance->Montage_SetBlendingOutDelegate(BlendingOutDelegate);
-			}
-			else
-			{
-				ABILITY_LOG(Warning, TEXT("UAbilityTask_PlayMontageAndWait called in Ability %s with null montage; Task Instance Name %s."), *Ability->GetName(), *InstanceName.ToString());
+
+				bPlayedMontage = true;
 			}
 		}
 	}
+
+	if (!bPlayedMontage)
+	{
+		ABILITY_LOG(Warning, TEXT("UAbilityTask_PlayMontageAndWait called in Ability %s failed to play montage; Task Instance Name %s."), *Ability->GetName(), *InstanceName.ToString());
+		OnCancelled.Broadcast();
+	}
 }
 
 void UAbilityTask_PlayMontageAndWait::ExternalCancel()
