Make InitAbilityActorInfo virtual

diff --git a/Engine/Source/Runtime/GameplayAbilities/Public/AbilitySystemComponent.h b/Engine/Source/Runtime/GameplayAbilities/Public/AbilitySystemComponent.h
--- a/Engine/Source/Runtime/GameplayAbilities/Public/AbilitySystemComponent.h
+++ b/Engine/Source/Runtime/GameplayAbilities/Public/AbilitySystemComponent.h
@@ -854,7 +854,7 @@ public:
 	 *      OwnerActor is the actor that logically owns this component.
 	 *		AvatarActor is what physical actor in the world we are acting on. Usually a Pawn but it could be a Tower, Building, Turret, etc, may be the same as Owner
 	 */
-	void InitAbilityActorInfo(AActor* InOwnerActor, AActor* InAvatarActor);
+	virtual void InitAbilityActorInfo(AActor* InOwnerActor, AActor* InAvatarActor);
 
 	/** Changes the avatar actor, leaves the owner actor the same */
 	void SetAvatarActor(AActor* InAvatarActor);
