Merging using UE4-To-UE4-Orion

- Adding ability to get Ability Spec from the handle of the GE that granted the ability

diff --git a/Engine/Source/Runtime/GameplayAbilities/Private/AbilitySystemComponent_Abilities.cpp b/Engine/Source/Runtime/GameplayAbilities/Private/AbilitySystemComponent_Abilities.cpp
--- a/Engine/Source/Runtime/GameplayAbilities/Private/AbilitySystemComponent_Abilities.cpp
+++ b/Engine/Source/Runtime/GameplayAbilities/Private/AbilitySystemComponent_Abilities.cpp
@@ -538,6 +538,18 @@ FGameplayAbilitySpec* UAbilitySystemComponent::FindAbilitySpecFromHandle(FGamepl
 	return nullptr;
 }
 
+FGameplayAbilitySpec* UAbilitySystemComponent::FindAbilitySpecFromGEHandle(FActiveGameplayEffectHandle Handle)
+{
+	for (FGameplayAbilitySpec& Spec : ActivatableAbilities.Items)
+	{
+		if (Spec.GameplayEffectHandle == Handle)
+		{
+			return &Spec;
+		}
+	}
+	return nullptr;
+}
+
 void UAbilitySystemComponent::MarkAbilitySpecDirty(FGameplayAbilitySpec& Spec)
 {
 	if (IsOwnerActorAuthoritative())
diff --git a/Engine/Source/Runtime/GameplayAbilities/Public/AbilitySystemComponent.h b/Engine/Source/Runtime/GameplayAbilities/Public/AbilitySystemComponent.h
--- a/Engine/Source/Runtime/GameplayAbilities/Public/AbilitySystemComponent.h
+++ b/Engine/Source/Runtime/GameplayAbilities/Public/AbilitySystemComponent.h
@@ -626,6 +626,9 @@ public:
 	/** Returns an ability spec from a handle. If modifying call MarkAbilitySpecDirty */
 	FGameplayAbilitySpec* FindAbilitySpecFromHandle(FGameplayAbilitySpecHandle Handle);
 
+	/** Returns an ability spec from a GE handle. If modifying call MarkAbilitySpecDirty */
+	FGameplayAbilitySpec* FindAbilitySpecFromGEHandle(FActiveGameplayEffectHandle Handle);
+
 	/** Returns an ability spec from a handle. If modifying call MarkAbilitySpecDirty */
 	FGameplayAbilitySpec* FindAbilitySpecFromInputID(int32 InputID);
 
