Fixup function name mismatch on delegate removal

diff --git a/Engine/Source/Runtime/GameplayAbilities/Private/Abilities/Tasks/AbilityTask_WaitConfirmCancel.cpp b/Engine/Source/Runtime/GameplayAbilities/Private/Abilities/Tasks/AbilityTask_WaitConfirmCancel.cpp
--- a/Engine/Source/Runtime/GameplayAbilities/Private/Abilities/Tasks/AbilityTask_WaitConfirmCancel.cpp
+++ b/Engine/Source/Runtime/GameplayAbilities/Private/Abilities/Tasks/AbilityTask_WaitConfirmCancel.cpp
@@ -93,8 +93,8 @@ void UAbilityTask_WaitConfirmCancel::OnDestroy(bool AbilityEnding)
 {
 	if (RegisteredCallbacks && AbilitySystemComponent.IsValid())
 	{
-		AbilitySystemComponent->GenericLocalConfirmCallbacks.RemoveDynamic(this, &UAbilityTask_WaitConfirmCancel::OnConfirmCallback);
-		AbilitySystemComponent->GenericLocalCancelCallbacks.RemoveDynamic(this, &UAbilityTask_WaitConfirmCancel::OnCancelCallback);
+		AbilitySystemComponent->GenericLocalConfirmCallbacks.RemoveDynamic(this, &UAbilityTask_WaitConfirmCancel::OnLocalConfirmCallback);
+		AbilitySystemComponent->GenericLocalCancelCallbacks.RemoveDynamic(this, &UAbilityTask_WaitConfirmCancel::OnLocalCancelCallback);
 	}
 
 	Super::OnDestroy(AbilityEnding);
