Never set replicated prediction key to 0. Fixes PL issues where confirmations were getting dropped.

diff --git a/Engine/Source/Runtime/GameplayAbilities/Private/GameplayPrediction.cpp b/Engine/Source/Runtime/GameplayAbilities/Private/GameplayPrediction.cpp
--- a/Engine/Source/Runtime/GameplayAbilities/Private/GameplayPrediction.cpp
+++ b/Engine/Source/Runtime/GameplayAbilities/Private/GameplayPrediction.cpp
@@ -245,7 +245,14 @@ FScopedPredictionWindow::~FScopedPredictionWindow()
 	{
 		if (SetReplicatedPredictionKey)
 		{
-			Owner->ReplicatedPredictionKey = Owner->ScopedPredictionKey;
+			// It is important to not set the ReplicatedPredictionKey unless it is valid (>0).
+			// If we werent given a new prediction key for this scope from the client, then setting the
+			// replicated prediction key back to 0 could cause OnReps to be missed on the client during high PL.
+			// (for example, predict w/ key 100 -> prediction key replication dropped -> predict w/ invalid key -> next rep of prediction key is 0).
+			if (Owner->ScopedPredictionKey.IsValidKey())
+			{
+				Owner->ReplicatedPredictionKey = Owner->ScopedPredictionKey;
+			}
 		}
 		if (ClearScopedPredictionKey)
 		{
