Ability System:
-Add 'passed tags' field to executions. Executions can elect to have a passed in tag field (Defined by the execution class itself).
-Added basic support for changing duration of active gameplay effect
-Added a few helper methods for finding stack counts of active gameplay effects
-Minor change to LocalConfirm/Cancel delegates: it is now possible to reregister for these delegates during callback
(Merge from dev branch)

diff --git a/Engine/Source/Editor/GameplayAbilitiesEditor/Private/GameplayEffectExecutionDefinitionDetails.cpp b/Engine/Source/Editor/GameplayAbilitiesEditor/Private/GameplayEffectExecutionDefinitionDetails.cpp
--- a/Engine/Source/Editor/GameplayAbilitiesEditor/Private/GameplayEffectExecutionDefinitionDetails.cpp
+++ b/Engine/Source/Editor/GameplayAbilitiesEditor/Private/GameplayEffectExecutionDefinitionDetails.cpp
@@ -30,6 +30,7 @@ void FGameplayEffectExecutionDefinitionDetails::CustomizeHeader(TSharedRef<IProp
 void FGameplayEffectExecutionDefinitionDetails::CustomizeChildren(TSharedRef<IPropertyHandle> StructPropertyHandle, IDetailChildrenBuilder& StructBuilder, IPropertyTypeCustomizationUtils& StructCustomizationUtils)
 {
 	bShowCalculationModifiers = false;
+	bShowPassedInTags = false;
 
 	// @todo: For now, only allow single editing
 	if (StructPropertyHandle->GetNumOuterObjects() == 1)
@@ -37,6 +38,7 @@ void FGameplayEffectExecutionDefinitionDetails::CustomizeChildren(TSharedRef<IPr
 		CalculationClassPropHandle = StructPropertyHandle->GetChildHandle(GET_MEMBER_NAME_CHECKED(FGameplayEffectExecutionDefinition, CalculationClass));
 		TSharedPtr<IPropertyHandle> LinkedEffectsPropHandle = StructPropertyHandle->GetChildHandle(GET_MEMBER_NAME_CHECKED(FGameplayEffectExecutionDefinition, ConditionalGameplayEffectClasses));
 		TSharedPtr<IPropertyHandle> CalcModPropHandle = StructPropertyHandle->GetChildHandle(GET_MEMBER_NAME_CHECKED(FGameplayEffectExecutionDefinition, CalculationModifiers));
+		TSharedPtr<IPropertyHandle> PassedInTagsHandle = StructPropertyHandle->GetChildHandle(GET_MEMBER_NAME_CHECKED(FGameplayEffectExecutionDefinition, PassedInTags));
 		CalculationModifiersArrayPropHandle = CalcModPropHandle.IsValid() ? CalcModPropHandle->AsArray() : nullptr;
 
 		if (CalculationClassPropHandle.IsValid())
@@ -56,6 +58,12 @@ void FGameplayEffectExecutionDefinitionDetails::CustomizeChildren(TSharedRef<IPr
 		{
 			StructBuilder.AddChildProperty(LinkedEffectsPropHandle.ToSharedRef());
 		}
+
+		if (PassedInTagsHandle.IsValid())
+		{
+			IDetailPropertyRow& PropRow = StructBuilder.AddChildProperty(PassedInTagsHandle.ToSharedRef());
+			PropRow.Visibility(TAttribute<EVisibility>::Create(TAttribute<EVisibility>::FGetter::CreateSP(this, &FGameplayEffectExecutionDefinitionDetails::GetPassedInTagsVisibility)));
+		}
 	}
 }
 
@@ -81,6 +89,9 @@ void FGameplayEffectExecutionDefinitionDetails::UpdateCalculationModifiers()
 			if (ExecutionCDO)
 			{
 				ExecutionCDO->GetValidScopedModifierAttributeCaptureDefinitions(ValidCaptureDefinitions);
+
+				// Grab this while we are at it so we know if we should show the 'Passed In Tags' property
+				bShowPassedInTags = ExecutionCDO->DoesRequirePassedInTags();
 			}
 		}
 	}
@@ -126,4 +137,9 @@ EVisibility FGameplayEffectExecutionDefinitionDetails::GetCalculationModifierVis
 	return (bShowCalculationModifiers ? EVisibility::Visible : EVisibility::Collapsed);
 }
 
+EVisibility FGameplayEffectExecutionDefinitionDetails::GetPassedInTagsVisibility() const
+{
+	return (bShowPassedInTags ? EVisibility::Visible : EVisibility::Collapsed);
+}
+
 #undef LOCTEXT_NAMESPACE
diff --git a/Engine/Source/Editor/GameplayAbilitiesEditor/Private/GameplayEffectExecutionDefinitionDetails.h b/Engine/Source/Editor/GameplayAbilitiesEditor/Private/GameplayEffectExecutionDefinitionDetails.h
--- a/Engine/Source/Editor/GameplayAbilitiesEditor/Private/GameplayEffectExecutionDefinitionDetails.h
+++ b/Engine/Source/Editor/GameplayAbilitiesEditor/Private/GameplayEffectExecutionDefinitionDetails.h
@@ -32,6 +32,9 @@ private:
 	/** Visibility delegate for calculation modifiers */
 	EVisibility GetCalculationModifierVisibility() const;
 
+	/** Visibility delegate for Passed In Tags field */
+	EVisibility GetPassedInTagsVisibility() const;
+
 	/** Property handle for the calculation class */
 	TSharedPtr<IPropertyHandle> CalculationClassPropHandle;
 
@@ -40,4 +43,7 @@ private:
 
 	/** If true, the calculation modifiers array should be shown */
 	bool bShowCalculationModifiers;
+
+	/** If true, the Passed In Tags field will be shown */
+	bool bShowPassedInTags;
 };
\ No newline at end of file
diff --git a/Engine/Source/Runtime/GameplayAbilities/Private/AbilitySystemComponent.cpp b/Engine/Source/Runtime/GameplayAbilities/Private/AbilitySystemComponent.cpp
--- a/Engine/Source/Runtime/GameplayAbilities/Private/AbilitySystemComponent.cpp
+++ b/Engine/Source/Runtime/GameplayAbilities/Private/AbilitySystemComponent.cpp
@@ -1014,6 +1014,16 @@ TArray<float> UAbilitySystemComponent::GetActiveEffectsDuration(const FActiveGam
 	return ActiveGameplayEffects.GetActiveEffectsDuration(Query);
 }
 
+TArray<FActiveGameplayEffectHandle> UAbilitySystemComponent::GetActiveEffects(const FActiveGameplayEffectQuery Query) const
+{
+	return ActiveGameplayEffects.GetActiveEffects(Query);
+}
+
+void UAbilitySystemComponent::ModifyActiveEffectStartTime(FActiveGameplayEffectHandle Handle, float StartTimeDiff)
+{
+	ActiveGameplayEffects.ModifyActiveEffectStartTime(Handle, StartTimeDiff);
+}
+
 void UAbilitySystemComponent::RemoveActiveEffectsWithTags(const FGameplayTagContainer Tags)
 {
 	AActor *OwningActor = GetOwner();
diff --git a/Engine/Source/Runtime/GameplayAbilities/Private/AbilitySystemComponent_Abilities.cpp b/Engine/Source/Runtime/GameplayAbilities/Private/AbilitySystemComponent_Abilities.cpp
--- a/Engine/Source/Runtime/GameplayAbilities/Private/AbilitySystemComponent_Abilities.cpp
+++ b/Engine/Source/Runtime/GameplayAbilities/Private/AbilitySystemComponent_Abilities.cpp
@@ -1871,14 +1871,16 @@ void UAbilitySystemComponent::AbilitySpecInputReleased(FGameplayAbilitySpec& Spe
 
 void UAbilitySystemComponent::LocalInputConfirm()
 {
-	GenericLocalConfirmCallbacks.Broadcast();
+	FAbilityConfirmOrCancel Temp = GenericLocalConfirmCallbacks;
 	GenericLocalConfirmCallbacks.Clear();
+	Temp.Broadcast();
 }
 
 void UAbilitySystemComponent::LocalInputCancel()
 {	
-	GenericLocalCancelCallbacks.Broadcast();
+	FAbilityConfirmOrCancel Temp = GenericLocalCancelCallbacks;
 	GenericLocalCancelCallbacks.Clear();
+	Temp.Broadcast();
 }
 
 void UAbilitySystemComponent::TargetConfirm()
diff --git a/Engine/Source/Runtime/GameplayAbilities/Private/AbilitySystemStats.cpp b/Engine/Source/Runtime/GameplayAbilities/Private/AbilitySystemStats.cpp
--- a/Engine/Source/Runtime/GameplayAbilities/Private/AbilitySystemStats.cpp
+++ b/Engine/Source/Runtime/GameplayAbilities/Private/AbilitySystemStats.cpp
@@ -9,6 +9,8 @@ DEFINE_STAT(STAT_GameplayEffectsTick);
 DEFINE_STAT(STAT_GameplayEffectsCanApplyAttributeModifiers);
 DEFINE_STAT(STAT_GameplayEffectsGetActiveEffectsTimeRemaining);
 DEFINE_STAT(STAT_GameplayEffectsGetActiveEffectsDuration);
+DEFINE_STAT(STAT_GameplayEffectsGetActiveEffects);
+DEFINE_STAT(STAT_GameplayEffectsModifyActiveEffectStartTime);
 DEFINE_STAT(STAT_GameplayAbilityGetCooldownTimeRemaining);
 DEFINE_STAT(STAT_GameplayAbilityGetCooldownTimeRemainingAndDuration);
 DEFINE_STAT(STAT_RemoveActiveGameplayEffect);
diff --git a/Engine/Source/Runtime/GameplayAbilities/Private/GameplayEffect.cpp b/Engine/Source/Runtime/GameplayAbilities/Private/GameplayEffect.cpp
--- a/Engine/Source/Runtime/GameplayAbilities/Private/GameplayEffect.cpp
+++ b/Engine/Source/Runtime/GameplayAbilities/Private/GameplayEffect.cpp
@@ -1379,7 +1379,7 @@ void FActiveGameplayEffectsContainer::ExecuteActiveEffectsFrom(FGameplayEffectSp
 			check(ExecCDO);
 
 			// Run the custom execution
-			FGameplayEffectCustomExecutionParameters ExecutionParams(SpecToUse, CurExecDef.CalculationModifiers, Owner);
+			FGameplayEffectCustomExecutionParameters ExecutionParams(SpecToUse, CurExecDef.CalculationModifiers, Owner, CurExecDef.PassedInTags);
 			FGameplayEffectCustomExecutionOutput ExecutionOutput;
 			ExecCDO->Execute(ExecutionParams, ExecutionOutput);
 
@@ -2641,14 +2641,9 @@ void FActiveGameplayEffectsContainer::CheckDuration(FActiveGameplayEffectHandle
 			}
 			else
 			{
-				// check the time remaining for the current gameplay effect duration timer
-				// if it is less than zero create a new callback with the correct time remaining
-				float TimeRemaining = TimerManager.GetTimerRemaining(Effect.DurationHandle);
-				if (TimeRemaining <= 0.f)
-				{
-					FTimerDelegate Delegate = FTimerDelegate::CreateUObject(Owner, &UAbilitySystemComponent::CheckDurationExpired, Effect.Handle);
-					TimerManager.SetTimer(Effect.DurationHandle, Delegate, (Effect.StartWorldTime + Duration) - CurrentTime, false);
-				}
+				// Always reset the timer, since the duration might have been modified
+				FTimerDelegate Delegate = FTimerDelegate::CreateUObject(Owner, &UAbilitySystemComponent::CheckDurationExpired, Effect.Handle);
+				TimerManager.SetTimer(Effect.DurationHandle, Delegate, (Effect.StartWorldTime + Duration) - CurrentTime, false);
 			}
 
 			break;
@@ -2734,8 +2729,47 @@ TArray<float> FActiveGameplayEffectsContainer::GetActiveEffectsDuration(const FA
 	return ReturnList;
 }
 
+TArray<FActiveGameplayEffectHandle> FActiveGameplayEffectsContainer::GetActiveEffects(const FActiveGameplayEffectQuery Query) const
+{
+	SCOPE_CYCLE_COUNTER(STAT_GameplayEffectsGetActiveEffects);
+
+	TArray<FActiveGameplayEffectHandle> ReturnList;
+
+	for (const FActiveGameplayEffect& Effect : this)
+	{
+		if (!Query.Matches(Effect))
+		{
+			continue;
+		}
+
+		ReturnList.Add(Effect.Handle);
+	}
+
+	return ReturnList;
+}
+
+void FActiveGameplayEffectsContainer::ModifyActiveEffectStartTime(FActiveGameplayEffectHandle Handle, float StartTimeDiff)
+{
+	SCOPE_CYCLE_COUNTER(STAT_GameplayEffectsModifyActiveEffectStartTime);
+
+	FActiveGameplayEffect* Effect = GetActiveGameplayEffect(Handle);
+
+	if (Effect)
+	{
+		Effect->StartWorldTime += StartTimeDiff;
+		Effect->StartGameStateTime = FMath::RoundToInt((float)Effect->StartGameStateTime + StartTimeDiff);
+
+		CheckDuration(Handle);
+
+		MarkItemDirty(*Effect);
+	}
+}
+
 void FActiveGameplayEffectsContainer::RemoveActiveEffects(const FActiveGameplayEffectQuery Query, int32 StacksToRemove)
 {
+	// Force a lock because the removals could cause other removals earlier in the array, so iterating backwards is not safe all by itself
+	GAMEPLAYEFFECT_SCOPE_LOCK();
+
 	// Manually iterating through in reverse because this is a removal operation
 	for (int32 idx=GetNumGameplayEffects()-1; idx >= 0; --idx)
 	{
diff --git a/Engine/Source/Runtime/GameplayAbilities/Private/GameplayEffectExecutionCalculation.cpp b/Engine/Source/Runtime/GameplayAbilities/Private/GameplayEffectExecutionCalculation.cpp
--- a/Engine/Source/Runtime/GameplayAbilities/Private/GameplayEffectExecutionCalculation.cpp
+++ b/Engine/Source/Runtime/GameplayAbilities/Private/GameplayEffectExecutionCalculation.cpp
@@ -7,12 +7,14 @@
 FGameplayEffectCustomExecutionParameters::FGameplayEffectCustomExecutionParameters()
 	: OwningSpec(nullptr)
 	, TargetAbilitySystemComponent(nullptr)
+	, PassedInTags(FGameplayTagContainer())
 {
 }
 
-FGameplayEffectCustomExecutionParameters::FGameplayEffectCustomExecutionParameters(FGameplayEffectSpec& InOwningSpec, const TArray<FGameplayEffectExecutionScopedModifierInfo>& InScopedMods, UAbilitySystemComponent* InTargetAbilityComponent)
+FGameplayEffectCustomExecutionParameters::FGameplayEffectCustomExecutionParameters(FGameplayEffectSpec& InOwningSpec, const TArray<FGameplayEffectExecutionScopedModifierInfo>& InScopedMods, UAbilitySystemComponent* InTargetAbilityComponent, const FGameplayTagContainer& InPassedInTags)
 	: OwningSpec(&InOwningSpec)
 	, TargetAbilitySystemComponent(InTargetAbilityComponent)
+	, PassedInTags(InPassedInTags)
 {
 	check(InOwningSpec.Def);
 
@@ -61,6 +63,11 @@ UAbilitySystemComponent* FGameplayEffectCustomExecutionParameters::GetSourceAbil
 	return OwningSpec->GetContext().GetInstigatorAbilitySystemComponent();
 }
 
+const FGameplayTagContainer& FGameplayEffectCustomExecutionParameters::GetPassedInTags() const
+{
+	return PassedInTags;
+}
+
 bool FGameplayEffectCustomExecutionParameters::AttemptCalculateCapturedAttributeMagnitude(const FGameplayEffectAttributeCaptureDefinition& InCaptureDef, const FAggregatorEvaluateParameters& InEvalParams, OUT float& OutMagnitude) const
 {
 	check(OwningSpec);
@@ -226,6 +233,7 @@ void FGameplayEffectCustomExecutionOutput::GetOutputModifiers(OUT TArray<FGamepl
 UGameplayEffectExecutionCalculation::UGameplayEffectExecutionCalculation(const FObjectInitializer& ObjectInitializer)
 	: Super(ObjectInitializer)
 {
+	bRequiresPassedInTags = false;
 }
 
 #if WITH_EDITORONLY_DATA
@@ -242,6 +250,11 @@ void UGameplayEffectExecutionCalculation::GetValidScopedModifierAttributeCapture
 		}
 	}
 }
+
+bool UGameplayEffectExecutionCalculation::DoesRequirePassedInTags() const
+{
+	return bRequiresPassedInTags;
+}
 #endif // #if WITH_EDITORONLY_DATA
 
 void UGameplayEffectExecutionCalculation::Execute_Implementation(const FGameplayEffectCustomExecutionParameters& ExecutionParams, OUT FGameplayEffectCustomExecutionOutput& OutExecutionOutput) const
diff --git a/Engine/Source/Runtime/GameplayAbilities/Public/AbilitySystemComponent.h b/Engine/Source/Runtime/GameplayAbilities/Public/AbilitySystemComponent.h
--- a/Engine/Source/Runtime/GameplayAbilities/Public/AbilitySystemComponent.h
+++ b/Engine/Source/Runtime/GameplayAbilities/Public/AbilitySystemComponent.h
@@ -421,6 +421,10 @@ class GAMEPLAYABILITIES_API UAbilitySystemComponent : public UActorComponent, pu
 
 	TArray<float> GetActiveEffectsDuration(const FActiveGameplayEffectQuery Query) const;
 
+	TArray<FActiveGameplayEffectHandle> GetActiveEffects(const FActiveGameplayEffectQuery Query) const;
+
+	void ModifyActiveEffectStartTime(FActiveGameplayEffectHandle Handle, float StartTimeDiff);
+
 	UFUNCTION(BlueprintCallable, Category = GameplayEffects)
 	void RemoveActiveEffectsWithTags(FGameplayTagContainer Tags);
 
diff --git a/Engine/Source/Runtime/GameplayAbilities/Public/AbilitySystemStats.h b/Engine/Source/Runtime/GameplayAbilities/Public/AbilitySystemStats.h
--- a/Engine/Source/Runtime/GameplayAbilities/Public/AbilitySystemStats.h
+++ b/Engine/Source/Runtime/GameplayAbilities/Public/AbilitySystemStats.h
@@ -13,6 +13,8 @@ DECLARE_CYCLE_STAT_EXTERN(TEXT("GameplayEffectsTick"), STAT_GameplayEffectsTick,
 DECLARE_CYCLE_STAT_EXTERN(TEXT("CanApplyAttributeModifiers"), STAT_GameplayEffectsCanApplyAttributeModifiers, STATGROUP_AbilitySystem, );
 DECLARE_CYCLE_STAT_EXTERN(TEXT("GetActiveEffectsData"), STAT_GameplayEffectsGetActiveEffectsTimeRemaining, STATGROUP_AbilitySystem, );
 DECLARE_CYCLE_STAT_EXTERN(TEXT("GetActiveEffectsData"), STAT_GameplayEffectsGetActiveEffectsDuration, STATGROUP_AbilitySystem, );
+DECLARE_CYCLE_STAT_EXTERN(TEXT("GetActiveEffectsData"), STAT_GameplayEffectsGetActiveEffects, STATGROUP_AbilitySystem, );
+DECLARE_CYCLE_STAT_EXTERN(TEXT("SetActiveEffectsData"), STAT_GameplayEffectsModifyActiveEffectStartTime, STATGROUP_AbilitySystem, );
 DECLARE_CYCLE_STAT_EXTERN(TEXT("GetCooldownTimeRemaining"), STAT_GameplayAbilityGetCooldownTimeRemaining, STATGROUP_AbilitySystem, );
 DECLARE_CYCLE_STAT_EXTERN(TEXT("GetCooldownTimeRemaining"), STAT_GameplayAbilityGetCooldownTimeRemainingAndDuration, STATGROUP_AbilitySystem, );
 DECLARE_CYCLE_STAT_EXTERN(TEXT("RemoveActiveGameplayEffect"), STAT_RemoveActiveGameplayEffect, STATGROUP_AbilitySystem, );
diff --git a/Engine/Source/Runtime/GameplayAbilities/Public/GameplayEffect.h b/Engine/Source/Runtime/GameplayAbilities/Public/GameplayEffect.h
--- a/Engine/Source/Runtime/GameplayAbilities/Public/GameplayEffect.h
+++ b/Engine/Source/Runtime/GameplayAbilities/Public/GameplayEffect.h
@@ -334,6 +334,10 @@ struct FGameplayEffectExecutionDefinition
 	UPROPERTY(EditDefaultsOnly, Category=Execution)
 	TSubclassOf<UGameplayEffectExecutionCalculation> CalculationClass;
 	
+	/** These tags are passed into the execution as is, and may be used to do conditional logic */
+	UPROPERTY(EditDefaultsOnly, Category = Execution)
+	FGameplayTagContainer PassedInTags;
+
 	/** Modifiers that are applied "in place" during the execution calculation */
 	UPROPERTY(EditDefaultsOnly, Category = Execution)
 	TArray<FGameplayEffectExecutionScopedModifierInfo> CalculationModifiers;
@@ -1446,6 +1450,10 @@ struct GAMEPLAYABILITIES_API FActiveGameplayEffectsContainer : public FFastArray
 
 	TArray<float> GetActiveEffectsDuration(const FActiveGameplayEffectQuery Query) const;
 
+	TArray<FActiveGameplayEffectHandle> GetActiveEffects(const FActiveGameplayEffectQuery Query) const;
+
+	void ModifyActiveEffectStartTime(FActiveGameplayEffectHandle Handle, float StartTimeDiff);
+
 	void RemoveActiveEffects(const FActiveGameplayEffectQuery Query, int32 StacksToRemove);
 
 	int32 GetGameStateTime() const;
diff --git a/Engine/Source/Runtime/GameplayAbilities/Public/GameplayEffectExecutionCalculation.h b/Engine/Source/Runtime/GameplayAbilities/Public/GameplayEffectExecutionCalculation.h
--- a/Engine/Source/Runtime/GameplayAbilities/Public/GameplayEffectExecutionCalculation.h
+++ b/Engine/Source/Runtime/GameplayAbilities/Public/GameplayEffectExecutionCalculation.h
@@ -20,7 +20,7 @@ public:
 
 	// Constructors
 	FGameplayEffectCustomExecutionParameters();
-	FGameplayEffectCustomExecutionParameters(FGameplayEffectSpec& InOwningSpec, const TArray<FGameplayEffectExecutionScopedModifierInfo>& InScopedMods, UAbilitySystemComponent* InTargetAbilityComponent);
+	FGameplayEffectCustomExecutionParameters(FGameplayEffectSpec& InOwningSpec, const TArray<FGameplayEffectExecutionScopedModifierInfo>& InScopedMods, UAbilitySystemComponent* InTargetAbilityComponent, const FGameplayTagContainer& InPassedIntags);
 
 	/** Simple accessor to owning gameplay spec */
 	const FGameplayEffectSpec& GetOwningSpec() const;
@@ -31,6 +31,9 @@ public:
 	/** Simple accessor to source ability system component (could be null!) */
 	UAbilitySystemComponent* GetSourceAbilitySystemComponent() const;
 	
+	/** Simple accessor to the Passed In Tags to this execution */
+	const FGameplayTagContainer& GetPassedInTags() const;
+
 	/**
 	 * Attempts to calculate the magnitude of a captured attribute given the specified parameters. Can fail if the gameplay spec doesn't have
 	 * a valid capture for the attribute.
@@ -100,6 +103,9 @@ private:
 
 	/** Target ability system component of the execution */
 	TWeakObjectPtr<UAbilitySystemComponent> TargetAbilitySystemComponent;
+
+	/** The extra tags that were passed in to this execution */
+	FGameplayTagContainer PassedInTags;
 };
 
 /** Struct representing the output of a custom gameplay effect execution. */
@@ -164,6 +170,12 @@ class GAMEPLAYABILITIES_API UGameplayEffectExecutionCalculation : public UGamepl
 {
 	GENERATED_UCLASS_BODY()
 
+protected:
+
+	/** Used to indicate if this execution uses Passed In Tags */
+	UPROPERTY(EditDefaultsOnly, BlueprintReadOnly, Category = Attributes)
+	bool bRequiresPassedInTags;
+
 #if WITH_EDITORONLY_DATA
 
 protected:
@@ -180,8 +192,13 @@ public:
 	 */
 	virtual void GetValidScopedModifierAttributeCaptureDefinitions(OUT TArray<FGameplayEffectAttributeCaptureDefinition>& OutScopableModifiers) const;
 
+	/** Returns if this execution requires passed in tags */
+	virtual bool DoesRequirePassedInTags() const;
+
 #endif // #if WITH_EDITORONLY_DATA
 
+public:
+
 	/**
 	 * Called whenever the owning gameplay effect is executed. Allowed to do essentially whatever is desired, including generating new
 	 * modifiers to instantly execute as well.
