GameplayCue perf stats

diff --git a/Engine/Source/Runtime/GameplayAbilities/Private/AbilitySystemStats.cpp b/Engine/Source/Runtime/GameplayAbilities/Private/AbilitySystemStats.cpp
--- a/Engine/Source/Runtime/GameplayAbilities/Private/AbilitySystemStats.cpp
+++ b/Engine/Source/Runtime/GameplayAbilities/Private/AbilitySystemStats.cpp
@@ -23,5 +23,5 @@ DEFINE_STAT(STAT_FindAbilitySpecFromHandle);
 DEFINE_STAT(STAT_AggregatorEvaluate);
 DEFINE_STAT(STAT_HasApplicationImmunityToSpec);
 DEFINE_STAT(STAT_HasMatchingGameplayTag);
-
-
+DEFINE_STAT(STAT_HandleGameplayCueNotifyStatic);
+DEFINE_STAT(STAT_HandleGameplayCueNotifyActor);
\ No newline at end of file
diff --git a/Engine/Source/Runtime/GameplayAbilities/Private/GameplayCueNotify_Actor.cpp b/Engine/Source/Runtime/GameplayAbilities/Private/GameplayCueNotify_Actor.cpp
--- a/Engine/Source/Runtime/GameplayAbilities/Private/GameplayCueNotify_Actor.cpp
+++ b/Engine/Source/Runtime/GameplayAbilities/Private/GameplayCueNotify_Actor.cpp
@@ -62,6 +62,8 @@ bool AGameplayCueNotify_Actor::HandlesEvent(EGameplayCueEvent::Type EventType) c
 
 void AGameplayCueNotify_Actor::HandleGameplayCue(AActor* MyTarget, EGameplayCueEvent::Type EventType, FGameplayCueParameters Parameters)
 {
+	SCOPE_CYCLE_COUNTER(STAT_HandleGameplayCueNotifyActor);
+
 	if (MyTarget && !MyTarget->IsPendingKill())
 	{
 		K2_HandleGameplayCue(MyTarget, EventType, Parameters);
diff --git a/Engine/Source/Runtime/GameplayAbilities/Private/GameplayCueNotify_Static.cpp b/Engine/Source/Runtime/GameplayAbilities/Private/GameplayCueNotify_Static.cpp
--- a/Engine/Source/Runtime/GameplayAbilities/Private/GameplayCueNotify_Static.cpp
+++ b/Engine/Source/Runtime/GameplayAbilities/Private/GameplayCueNotify_Static.cpp
@@ -60,6 +60,8 @@ bool UGameplayCueNotify_Static::HandlesEvent(EGameplayCueEvent::Type EventType)
 
 void UGameplayCueNotify_Static::HandleGameplayCue(AActor* MyTarget, EGameplayCueEvent::Type EventType, FGameplayCueParameters Parameters)
 {
+	SCOPE_CYCLE_COUNTER(STAT_HandleGameplayCueNotifyStatic);
+
 	if (MyTarget && !MyTarget->IsPendingKill())
 	{
 		K2_HandleGameplayCue(MyTarget, EventType, Parameters);
diff --git a/Engine/Source/Runtime/GameplayAbilities/Public/AbilitySystemStats.h b/Engine/Source/Runtime/GameplayAbilities/Public/AbilitySystemStats.h
--- a/Engine/Source/Runtime/GameplayAbilities/Public/AbilitySystemStats.h
+++ b/Engine/Source/Runtime/GameplayAbilities/Public/AbilitySystemStats.h
@@ -26,4 +26,6 @@ DECLARE_CYCLE_STAT_EXTERN(TEXT("TickAbilityTasks"), STAT_TickAbilityTasks, STATG
 DECLARE_CYCLE_STAT_EXTERN(TEXT("FindAbilitySpecFromHandle"), STAT_FindAbilitySpecFromHandle, STATGROUP_AbilitySystem, );
 DECLARE_CYCLE_STAT_EXTERN(TEXT("Aggregator Evaluate"), STAT_AggregatorEvaluate, STATGROUP_AbilitySystem, );
 DECLARE_CYCLE_STAT_EXTERN(TEXT("Has Application Immunity To Spec"), STAT_HasApplicationImmunityToSpec, STATGROUP_AbilitySystem, );
-DECLARE_CYCLE_STAT_EXTERN(TEXT("Has Matching GameplayTag"), STAT_HasMatchingGameplayTag, STATGROUP_AbilitySystem, );
\ No newline at end of file
+DECLARE_CYCLE_STAT_EXTERN(TEXT("Has Matching GameplayTag"), STAT_HasMatchingGameplayTag, STATGROUP_AbilitySystem, );
+DECLARE_CYCLE_STAT_EXTERN(TEXT("GameplayCueNotify Static"), STAT_HandleGameplayCueNotifyStatic, STATGROUP_AbilitySystem, );
+DECLARE_CYCLE_STAT_EXTERN(TEXT("GameplayCueNotify Actor"), STAT_HandleGameplayCueNotifyActor, STATGROUP_AbilitySystem, );
\ No newline at end of file
