Missed a file during merge

diff --git a/Engine/Source/Runtime/GameplayAbilities/Private/GameplayEffect.cpp b/Engine/Source/Runtime/GameplayAbilities/Private/GameplayEffect.cpp
--- a/Engine/Source/Runtime/GameplayAbilities/Private/GameplayEffect.cpp
+++ b/Engine/Source/Runtime/GameplayAbilities/Private/GameplayEffect.cpp
@@ -9,6 +9,7 @@
 #include "AbilitySystemInterface.h"
 #include "GameplayModMagnitudeCalculation.h"
 #include "GameplayEffectExecutionCalculation.h"
+#include "GameplayCueManager.h"
 #if ENABLE_VISUAL_LOG
 #include "VisualLoggerTypes.h"
 //#include "VisualLogger/VisualLogger.h"
@@ -1305,7 +1306,6 @@ FActiveGameplayEffectsContainer::FActiveGameplayEffectsContainer()
 	: Owner(nullptr)
 	, ScopedLockCount(0)
 	, PendingRemoves(0)
-	, PendingAdds(0)
 	, PendingGameplayEffectHead(nullptr)
 {
 	PendingGameplayEffectNext = &PendingGameplayEffectHead;
@@ -1435,8 +1435,8 @@ void FActiveGameplayEffectsContainer::ExecuteActiveEffectsFrom(FGameplayEffectSp
 		// TODO: check replication policy. Right now we will replicate every execute via a multicast RPC
 
 		ABILITY_LOG(Log, TEXT("Invoking Execute GameplayCue for %s"), *SpecToUse.ToSimpleString());
-		Owner->ForceReplication();
-		Owner->NetMulticast_InvokeGameplayCueExecuted_FromSpec(SpecToUse, PredictionKey);
+
+		UAbilitySystemGlobals::Get().GetGameplayCueManager()->InvokeGameplayCueExecuted_FromSpec(Owner, SpecToUse, PredictionKey);
 	}
 
 	// Apply any conditional linked effects
@@ -1670,13 +1670,12 @@ FActiveGameplayEffect* FActiveGameplayEffectsContainer::FindStackableActiveGamep
 	const UGameplayEffect* GEDef = Spec.Def;
 	EGameplayEffectStackingType StackingType = GEDef->StackingType;
 
-	UAbilitySystemComponent* SourceASC = Spec.GetContext().GetInstigatorAbilitySystemComponent();
 	if (StackingType != EGameplayEffectStackingType::None && Spec.GetDuration() != UGameplayEffect::INSTANT_APPLICATION)
 	{
 		// Iterate through GameplayEffects to see if we find a match. Note that we could cache off a handle in a map but we would still
 		// do a linear search through GameplayEffects to find the actual FActiveGameplayEffect (due to unstable nature of the GameplayEffects array).
 		// If this becomes a slow point in the profiler, the map may still be useful as an early out to avoid an unnecessary sweep.
-
+		UAbilitySystemComponent* SourceASC = Spec.GetContext().GetInstigatorAbilitySystemComponent();
 		for (FActiveGameplayEffect& ActiveEffect: this)
 		{
 			// Aggregate by source stacking additionally requires the source ability component to match
@@ -2073,7 +2072,7 @@ FActiveGameplayEffect* FActiveGameplayEffectsContainer::ApplyGameplayEffectSpec(
 			 *	onto the real active GE list once the scope lock is over.
 			 *	
 			 *	To avoid extra heap allocations, each active gameplayeffect container keeps a linked list of pending GEs. This list is allocated
-			 *	on demand and resused in subsequent pending adds. The code below will either 1) Alloc a new pending GE 2) reuse an existing pending GE.
+			 *	on demand and re-used in subsequent pending adds. The code below will either 1) Alloc a new pending GE 2) reuse an existing pending GE.
 			 *	The move operator is used to move stuff to and from these pending GEs to avoid deep copies.
 			 */
 
@@ -2095,7 +2094,6 @@ FActiveGameplayEffect* FActiveGameplayEffectsContainer::ApplyGameplayEffectSpec(
 
 			// The next pending GameplayEffect goes to where our PendingNext points
 			PendingGameplayEffectNext = &AppliedActiveGE->PendingNext;
-			PendingAdds++;
 		}
 		else
 		{
@@ -3081,28 +3079,20 @@ void FActiveGameplayEffectsContainer::DecrementLock()
 		// ------------------------------------------
 		// Move any pending effects onto the real list
 		// ------------------------------------------
-		if (PendingAdds > 0)
+		FActiveGameplayEffect* PendingGameplayEffect = PendingGameplayEffectHead;
+		FActiveGameplayEffect* Stop = *PendingGameplayEffectNext;
+		while (PendingGameplayEffect != Stop)
 		{
-			FActiveGameplayEffect* PendingGameplayEffect = PendingGameplayEffectHead;
-			FActiveGameplayEffect* Stop = *PendingGameplayEffectNext;
-			while (PendingGameplayEffect != Stop)
+			if (!PendingGameplayEffect->IsPendingRemove)
 			{
 				GameplayEffects_Internal.Add(MoveTemp(*PendingGameplayEffect));
-				PendingGameplayEffect = PendingGameplayEffect->PendingNext;
-				PendingAdds--;
-			}
-
-			
-			if (!ensure(PendingAdds == 0))
-			{
-				ABILITY_LOG(Error, TEXT("~FScopedActiveGameplayEffectLock has %d pending adds after a scope lock removal"), PendingAdds);
-				PendingAdds = 0;
 			}
-
-			// Reset our pending GameplayEffect linked list
-			PendingGameplayEffectNext = &PendingGameplayEffectHead;
+			PendingGameplayEffect = PendingGameplayEffect->PendingNext;
 		}
 
+		// Reset our pending GameplayEffect linked list
+		PendingGameplayEffectNext = &PendingGameplayEffectHead;
+
 		// -----------------------------------------
 		// Delete any pending remove effects
 		// -----------------------------------------
