UE4-Orion-To-UE4(2475245): Ability task searching; Ability custom events;

* Added 'FAbilityInstanceClassPredicate' for searching through ability task instances by class.

* Added 5 game specific custom events to EAbilityGenericReplicatedEvent.

diff --git a/Engine/Source/Runtime/GameplayAbilities/Public/Abilities/GameplayAbilityTargetTypes.h b/Engine/Source/Runtime/GameplayAbilities/Public/Abilities/GameplayAbilityTargetTypes.h
--- a/Engine/Source/Runtime/GameplayAbilities/Public/Abilities/GameplayAbilityTargetTypes.h
+++ b/Engine/Source/Runtime/GameplayAbilities/Public/Abilities/GameplayAbilityTargetTypes.h
@@ -611,7 +611,13 @@ namespace EAbilityGenericReplicatedEvent
 		/** A generic event from the client */
 		GenericSignalFromClient,
 		/** A generic event from the server */
-		GenericSignalFromServer,		
+		GenericSignalFromServer,
+		/** Custom events for game use */
+		GameCustom1,
+		GameCustom2,
+		GameCustom3,
+		GameCustom4,
+		GameCustom5,
 
 		MAX
 	};
diff --git a/Engine/Source/Runtime/GameplayAbilities/Public/Abilities/Tasks/AbilityTask.h b/Engine/Source/Runtime/GameplayAbilities/Public/Abilities/Tasks/AbilityTask.h
--- a/Engine/Source/Runtime/GameplayAbilities/Public/Abilities/Tasks/AbilityTask.h
+++ b/Engine/Source/Runtime/GameplayAbilities/Public/Abilities/Tasks/AbilityTask.h
@@ -173,3 +173,19 @@ struct FAbilityInstanceNamePredicate
 
 	FName InstanceName;
 };
+
+
+struct FAbilityInstanceClassPredicate
+{
+	FAbilityInstanceClassPredicate(TSubclassOf<UAbilityTask> Class)
+	{
+		TaskClass = Class;
+	}
+
+	bool operator()(const TWeakObjectPtr<UAbilityTask> A) const
+	{
+		return (A.IsValid() && (A.Get()->GetClass() == TaskClass));
+	}
+
+	TSubclassOf<UAbilityTask> TaskClass;
+};
