UE4-Orion-To-UE4(2477460): Added UAbilityTask::IsLocallyControlled();

diff --git a/Engine/Source/Runtime/GameplayAbilities/Private/Abilities/Tasks/AbilityTask.cpp b/Engine/Source/Runtime/GameplayAbilities/Private/Abilities/Tasks/AbilityTask.cpp
--- a/Engine/Source/Runtime/GameplayAbilities/Private/Abilities/Tasks/AbilityTask.cpp
+++ b/Engine/Source/Runtime/GameplayAbilities/Private/Abilities/Tasks/AbilityTask.cpp
@@ -184,6 +184,17 @@ bool UAbilityTask::IsForRemoteClient() const
 	return false;
 }
 
+bool UAbilityTask::IsLocallyControlled() const
+{
+	UGameplayAbility* MyAbility = Ability.Get();
+	if (MyAbility && MyAbility->GetCurrentActorInfo()->OwnerActor.IsValid())
+	{
+		return MyAbility->GetCurrentActorInfo()->IsLocallyControlled();
+	}
+
+	return false;
+}
+
 bool UAbilityTask::CallOrAddReplicatedDelegate(EAbilityGenericReplicatedEvent::Type Event, FSimpleMulticastDelegate::FDelegate Delegate)
 {
 	return AbilitySystemComponent->CallOrAddReplicatedDelegate(Event, GetAbilitySpecHandle(), GetActivationPredictionKey(), Delegate);
diff --git a/Engine/Source/Runtime/GameplayAbilities/Public/Abilities/Tasks/AbilityTask.h b/Engine/Source/Runtime/GameplayAbilities/Public/Abilities/Tasks/AbilityTask.h
--- a/Engine/Source/Runtime/GameplayAbilities/Public/Abilities/Tasks/AbilityTask.h
+++ b/Engine/Source/Runtime/GameplayAbilities/Public/Abilities/Tasks/AbilityTask.h
@@ -148,7 +148,10 @@ class GAMEPLAYABILITIES_API UAbilityTask : public UObject
 	/** Returns true if we are executing the ability on the server for a non locally controlled client */
 	bool IsForRemoteClient() const;
 
-protected:	
+	/** Returns true if we are executing the ability on the locally controlled client */
+	bool IsLocallyControlled() const;
+
+protected:
 
 	/** End and CleanUp the task - may be called by the task itself or by the owning ability if the ability is ending. Do NOT call directly! Call EndTask() or AbilityEnded() */
 	virtual void OnDestroy(bool AbilityIsEnding);
