WaitInputPress/WaitInputRelease initial state;

* Added option to WaitInputPress for checking if the button is already pressed when the task starts.

* Added option to WaitInputRelease for checking if the button is already released when the task starts.

UE4-Orion-To-UE4(2477944)

diff --git a/Engine/Source/Runtime/GameplayAbilities/Private/Abilities/Tasks/AbilityTask_WaitInputPress.cpp b/Engine/Source/Runtime/GameplayAbilities/Private/Abilities/Tasks/AbilityTask_WaitInputPress.cpp
--- a/Engine/Source/Runtime/GameplayAbilities/Private/Abilities/Tasks/AbilityTask_WaitInputPress.cpp
+++ b/Engine/Source/Runtime/GameplayAbilities/Private/Abilities/Tasks/AbilityTask_WaitInputPress.cpp
@@ -8,6 +8,7 @@ UAbilityTask_WaitInputPress::UAbilityTask_WaitInputPress(const FObjectInitialize
 	: Super(ObjectInitializer)
 {
 	StartTime = 0.f;
+	bTestInitialState = false;
 }
 
 void UAbilityTask_WaitInputPress::OnPressCallback()
@@ -36,9 +37,11 @@ void UAbilityTask_WaitInputPress::OnPressCallback()
 	EndTask();
 }
 
-UAbilityTask_WaitInputPress* UAbilityTask_WaitInputPress::WaitInputPress(class UObject* WorldContextObject)
+UAbilityTask_WaitInputPress* UAbilityTask_WaitInputPress::WaitInputPress(class UObject* WorldContextObject, bool bTestAlreadyPressed)
 {
-	return NewTask<UAbilityTask_WaitInputPress>(WorldContextObject);
+	UAbilityTask_WaitInputPress* Task = NewTask<UAbilityTask_WaitInputPress>(WorldContextObject);
+	Task->bTestInitialState = bTestAlreadyPressed;
+	return Task;
 }
 
 void UAbilityTask_WaitInputPress::Activate()
@@ -46,6 +49,16 @@ void UAbilityTask_WaitInputPress::Activate()
 	StartTime = GetWorld()->GetTimeSeconds();
 	if (Ability.IsValid())
 	{
+		if (bTestInitialState && IsLocallyControlled())
+		{
+			FGameplayAbilitySpec *Spec = Ability->GetCurrentAbilitySpec();
+			if (Spec && Spec->InputPressed)
+			{
+				OnPressCallback();
+				return;
+			}
+		}
+
 		DelegateHandle = AbilitySystemComponent->AbilityReplicatedEventDelegate(EAbilityGenericReplicatedEvent::InputPressed, GetAbilitySpecHandle(), GetActivationPredictionKey()).AddUObject(this, &UAbilityTask_WaitInputPress::OnPressCallback);
 		if (IsForRemoteClient())
 		{
diff --git a/Engine/Source/Runtime/GameplayAbilities/Private/Abilities/Tasks/AbilityTask_WaitInputRelease.cpp b/Engine/Source/Runtime/GameplayAbilities/Private/Abilities/Tasks/AbilityTask_WaitInputRelease.cpp
--- a/Engine/Source/Runtime/GameplayAbilities/Private/Abilities/Tasks/AbilityTask_WaitInputRelease.cpp
+++ b/Engine/Source/Runtime/GameplayAbilities/Private/Abilities/Tasks/AbilityTask_WaitInputRelease.cpp
@@ -8,6 +8,7 @@ UAbilityTask_WaitInputRelease::UAbilityTask_WaitInputRelease(const FObjectInitia
 	: Super(ObjectInitializer)
 {
 	StartTime = 0.f;
+	bTestInitialState = false;
 }
 
 void UAbilityTask_WaitInputRelease::OnReleaseCallback()
@@ -36,9 +37,11 @@ void UAbilityTask_WaitInputRelease::OnReleaseCallback()
 	EndTask();
 }
 
-UAbilityTask_WaitInputRelease* UAbilityTask_WaitInputRelease::WaitInputRelease(class UObject* WorldContextObject)
+UAbilityTask_WaitInputRelease* UAbilityTask_WaitInputRelease::WaitInputRelease(class UObject* WorldContextObject, bool bTestAlreadyReleased)
 {
-	return NewTask<UAbilityTask_WaitInputRelease>(WorldContextObject);
+	UAbilityTask_WaitInputRelease* Task = NewTask<UAbilityTask_WaitInputRelease>(WorldContextObject);
+	Task->bTestInitialState = bTestAlreadyReleased;
+	return Task;
 }
 
 void UAbilityTask_WaitInputRelease::Activate()
@@ -46,6 +49,16 @@ void UAbilityTask_WaitInputRelease::Activate()
 	StartTime = GetWorld()->GetTimeSeconds();
 	if (Ability.IsValid())
 	{
+		if (bTestInitialState && IsLocallyControlled())
+		{
+			FGameplayAbilitySpec *Spec = Ability->GetCurrentAbilitySpec();
+			if (Spec && !Spec->InputPressed)
+			{
+				OnReleaseCallback();
+				return;
+			}
+		}
+
 		DelegateHandle = AbilitySystemComponent->AbilityReplicatedEventDelegate(EAbilityGenericReplicatedEvent::InputReleased, GetAbilitySpecHandle(), GetActivationPredictionKey()).AddUObject(this, &UAbilityTask_WaitInputRelease::OnReleaseCallback);
 		if (IsForRemoteClient())
 		{
diff --git a/Engine/Source/Runtime/GameplayAbilities/Public/Abilities/Tasks/AbilityTask_WaitInputPress.h b/Engine/Source/Runtime/GameplayAbilities/Public/Abilities/Tasks/AbilityTask_WaitInputPress.h
--- a/Engine/Source/Runtime/GameplayAbilities/Public/Abilities/Tasks/AbilityTask_WaitInputPress.h
+++ b/Engine/Source/Runtime/GameplayAbilities/Public/Abilities/Tasks/AbilityTask_WaitInputPress.h
@@ -9,7 +9,7 @@ DECLARE_DYNAMIC_MULTICAST_DELEGATE_OneParam(FInputPressDelegate, float, TimeWait
 
 /**
  *	Waits until the input is pressed from activating an ability. This should be true immediately upon starting the ability, since the key was pressed to activate it.
- *	We expect server to execute this task in parrallel and keep its own time. We do not keep track of 
+ *	We expect server to execute this task in parallel and keep its own time. We do not keep track of 
  */
 UCLASS(MinimalAPI)
 class UAbilityTask_WaitInputPress : public UAbilityTask
@@ -26,12 +26,13 @@ class UAbilityTask_WaitInputPress : public UAbilityTask
 
 	/** Wait until the user presses the input button for this ability's activation. Returns time this node spent waiting for the press. Will return 0 if input was already down. */
 	UFUNCTION(BlueprintCallable, Category = "Ability|Tasks", meta = (HidePin = "WorldContextObject", DefaultToSelf = "WorldContextObject", BlueprintInternalUseOnly = "TRUE"))
-	static UAbilityTask_WaitInputPress* WaitInputPress(UObject* WorldContextObject);
+	static UAbilityTask_WaitInputPress* WaitInputPress(UObject* WorldContextObject, bool bTestAlreadyPressed=false);
 
 protected:
 
 	virtual void OnDestroy(bool AbilityEnded) override;
 
 	float StartTime;
+	bool bTestInitialState;
 	FDelegateHandle DelegateHandle;
 };
\ No newline at end of file
diff --git a/Engine/Source/Runtime/GameplayAbilities/Public/Abilities/Tasks/AbilityTask_WaitInputRelease.h b/Engine/Source/Runtime/GameplayAbilities/Public/Abilities/Tasks/AbilityTask_WaitInputRelease.h
--- a/Engine/Source/Runtime/GameplayAbilities/Public/Abilities/Tasks/AbilityTask_WaitInputRelease.h
+++ b/Engine/Source/Runtime/GameplayAbilities/Public/Abilities/Tasks/AbilityTask_WaitInputRelease.h
@@ -9,7 +9,7 @@ DECLARE_DYNAMIC_MULTICAST_DELEGATE_OneParam(FInputReleaseDelegate, float, TimeHe
 
 /**
  *	Waits until the input is released from activating an ability. Clients will replicate a 'release input' event to the server, but not the exact time it was held locally.
- *	We expect server to execute this task in parrallel and keep its own time.
+ *	We expect server to execute this task in parallel and keep its own time.
  */
 UCLASS(MinimalAPI)
 class UAbilityTask_WaitInputRelease : public UAbilityTask
@@ -26,12 +26,13 @@ class UAbilityTask_WaitInputRelease : public UAbilityTask
 
 	/** Wait until the user releases the input button for this ability's activation. Returns time from hitting this node, till release. Will return 0 if input was already released. */
 	UFUNCTION(BlueprintCallable, Category="Ability|Tasks", meta = (HidePin = "WorldContextObject", DefaultToSelf = "WorldContextObject", BlueprintInternalUseOnly = "TRUE"))
-	static UAbilityTask_WaitInputRelease* WaitInputRelease(UObject* WorldContextObject);
+	static UAbilityTask_WaitInputRelease* WaitInputRelease(UObject* WorldContextObject, bool bTestAlreadyReleased=false);
 
 protected:
 
 	virtual void OnDestroy(bool AbilityEnded) override;
 
 	float StartTime;
+	bool bTestInitialState;
 	FDelegateHandle DelegateHandle;
 };
\ No newline at end of file
