Deprecate unnecessary AActor::GetActorClass

diff --git a/Engine/Source/Editor/Kismet/Private/SSCSEditor.cpp b/Engine/Source/Editor/Kismet/Private/SSCSEditor.cpp
--- a/Engine/Source/Editor/Kismet/Private/SSCSEditor.cpp
+++ b/Engine/Source/Editor/Kismet/Private/SSCSEditor.cpp
@@ -3847,7 +3847,7 @@ FSCSEditorTreeNodePtrType SSCSEditor::GetNodeFromActorComponent(const UActorComp
 			{
 				// Get the component owner's class object
 				check(ActorComponent->GetOwner() != NULL);
-				UClass* OwnerClass = ActorComponent->GetOwner()->GetActorClass();
+				UClass* OwnerClass = ActorComponent->GetOwner()->GetClass();
 
 				// If the given component is one that's created during Blueprint construction
 				if (ActorComponent->IsCreatedByConstructionScript())
diff --git a/Engine/Source/Editor/UnrealEd/Private/Kismet2/ComponentEditorUtils.cpp b/Engine/Source/Editor/UnrealEd/Private/Kismet2/ComponentEditorUtils.cpp
--- a/Engine/Source/Editor/UnrealEd/Private/Kismet2/ComponentEditorUtils.cpp
+++ b/Engine/Source/Editor/UnrealEd/Private/Kismet2/ComponentEditorUtils.cpp
@@ -118,7 +118,7 @@ bool FComponentEditorUtils::CanEditNativeComponent(const UActorComponent* Native
 
 	bool bCanEdit = false;
 	
-	UClass* OwnerClass = (NativeComponent && NativeComponent->GetOwner()) ? NativeComponent->GetOwner()->GetActorClass() : nullptr;
+	UClass* OwnerClass = (NativeComponent && NativeComponent->GetOwner()) ? NativeComponent->GetOwner()->GetClass() : nullptr;
 	if (OwnerClass != nullptr)
 	{
 		// If the owner is a blueprint generated class, use the BP parent class
diff --git a/Engine/Source/Runtime/Engine/Classes/GameFramework/Actor.h b/Engine/Source/Runtime/Engine/Classes/GameFramework/Actor.h
--- a/Engine/Source/Runtime/Engine/Classes/GameFramework/Actor.h
+++ b/Engine/Source/Runtime/Engine/Classes/GameFramework/Actor.h
@@ -2373,7 +2373,7 @@ public:
 // DEPRECATED FUNCTIONS
 
 	/** Get the class of this Actor. */
-	UFUNCTION(BlueprintCallable, meta=(DeprecatedFunction, FriendlyName = "GetActorClass"), Category="Class")
+	DEPRECATED(4.8, "Use GetClass() instead of GetActorClass()")
 	class UClass* GetActorClass() const;
 };
 
diff --git a/Engine/Source/Runtime/GameplayAbilities/Private/GameplayCueManager.cpp b/Engine/Source/Runtime/GameplayAbilities/Private/GameplayCueManager.cpp
--- a/Engine/Source/Runtime/GameplayAbilities/Private/GameplayCueManager.cpp
+++ b/Engine/Source/Runtime/GameplayAbilities/Private/GameplayCueManager.cpp
@@ -121,7 +121,7 @@ void UGameplayCueManager::HandleGameplayCueNotify_Internal(AActor* TargetActor,
 				{
 					// We don't have an instance for this, and we need one, so make one
 					//SpawnedInstancedCue = static_cast<AGameplayCueNotify_Actor*>(StaticDuplicateObject(InstancedCue, this, TEXT("None"), ~RF_RootSet));
-					SpawnedInstancedCue = TargetActor->GetWorld()->SpawnActor<AGameplayCueNotify_Actor>(InstancedCue->GetActorClass(), TargetActor->GetActorLocation(), TargetActor->GetActorRotation());
+					SpawnedInstancedCue = TargetActor->GetWorld()->SpawnActor<AGameplayCueNotify_Actor>(InstancedCue->GetClass(), TargetActor->GetActorLocation(), TargetActor->GetActorRotation());
 					auto& InnerMap = NotifyMapActor.FindOrAdd(TargetActor);
 					InnerMap.Add(DataIdx) = SpawnedInstancedCue;
 				}
