Enabled UHT to digest GENERATED_BODY instead of GENERATED_UCLASS_BODY, GENERATED_USTRUCT_BODY, GENERATED_UINTERFACE_BODY or GENERATED_IINTERFACE_BODY, changed every occurence to the new syntax and fix

diff --git a/Engine/Source/Runtime/GameplayAbilities/Public/AbilitySystemComponent.h b/Engine/Source/Runtime/GameplayAbilities/Public/AbilitySystemComponent.h
--- a/Engine/Source/Runtime/GameplayAbilities/Public/AbilitySystemComponent.h
+++ b/Engine/Source/Runtime/GameplayAbilities/Public/AbilitySystemComponent.h
@@ -51,7 +51,9 @@ DECLARE_MULTICAST_DELEGATE_TwoParams(FAbilityFailedDelegate, const UGameplayAbil
 UCLASS(ClassGroup=AbilitySystem, hidecategories=(Object,LOD,Lighting,Transform,Sockets,TextureStreaming), editinlinenew, meta=(BlueprintSpawnableComponent))
 class GAMEPLAYABILITIES_API UAbilitySystemComponent : public UActorComponent, public IGameplayTagAssetInterface
 {
-	GENERATED_UCLASS_BODY()
+	GENERATED_BODY()
+public:
+	UAbilitySystemComponent(const FObjectInitializer& ObjectInitializer = FObjectInitializer::Get());
 
 	/** Used to register callbacks to confirm/cancel input */
 	DECLARE_DYNAMIC_MULTICAST_DELEGATE(FAbilityConfirmOrCancel);
@@ -457,15 +459,19 @@ class GAMEPLAYABILITIES_API UAbilitySystemComponent : public UActorComponent, pu
 	// Do not call these functions directly, call the wrappers on GameplayCueManager instead
 	UFUNCTION(NetMulticast, unreliable)
 	void NetMulticast_InvokeGameplayCueExecuted_FromSpec(const FGameplayEffectSpecForRPC Spec, FPredictionKey PredictionKey);
+	virtual void NetMulticast_InvokeGameplayCueExecuted_FromSpec_Implementation(const FGameplayEffectSpecForRPC Spec, FPredictionKey PredictionKey);
 
 	UFUNCTION(NetMulticast, unreliable)
 	void NetMulticast_InvokeGameplayCueExecuted(const FGameplayTag GameplayCueTag, FPredictionKey PredictionKey, FGameplayEffectContextHandle EffectContext);
+	virtual void NetMulticast_InvokeGameplayCueExecuted_Implementation(const FGameplayTag GameplayCueTag, FPredictionKey PredictionKey, FGameplayEffectContextHandle EffectContext);
 
 	UFUNCTION(NetMulticast, unreliable)
 	void NetMulticast_InvokeGameplayCueExecuted_WithParams(const FGameplayTag GameplayCueTag, FPredictionKey PredictionKey, FGameplayCueParameters GameplayCueParameters);
+	virtual void NetMulticast_InvokeGameplayCueExecuted_WithParams_Implementation(const FGameplayTag GameplayCueTag, FPredictionKey PredictionKey, FGameplayCueParameters GameplayCueParameters);
 
 	UFUNCTION(NetMulticast, unreliable)
 	void NetMulticast_InvokeGameplayCueAdded(const FGameplayTag GameplayCueTag, FPredictionKey PredictionKey, FGameplayEffectContextHandle EffectContext);
+	virtual void NetMulticast_InvokeGameplayCueAdded_Implementation(const FGameplayTag GameplayCueTag, FPredictionKey PredictionKey, FGameplayEffectContextHandle EffectContext);
 
 	// GameplayCues can also come on their own. These take an optional effect context to pass through hit result, etc
 	void ExecuteGameplayCue(const FGameplayTag GameplayCueTag, FGameplayEffectContextHandle EffectContext = FGameplayEffectContextHandle());
@@ -650,11 +656,15 @@ public:
 	UFUNCTION()
 	void	OnRep_ActivateAbilities();
 
-	UFUNCTION(Server, reliable, WithValidation)
+	UFUNCTION(Server="ServerTryActivateAbility_Implementation", reliable, WithValidation="ServerTryActivateAbility_Validate")
 	void	ServerTryActivateAbility(FGameplayAbilitySpecHandle AbilityToActivate, FPredictionKey PredictionKey);
+	virtual void ServerTryActivateAbility_Implementation(FGameplayAbilitySpecHandle AbilityToActivate, FPredictionKey PredictionKey);
+	virtual bool ServerTryActivateAbility_Validate(FGameplayAbilitySpecHandle , FPredictionKey );
 
-	UFUNCTION(Server, reliable, WithValidation)
+	UFUNCTION(Server="ServerTryActivateAbilityWithEventData_Implementation", reliable, WithValidation="ServerTryActivateAbilityWithEventData_Validate")
 	void	ServerTryActivateAbilityWithEventData(FGameplayAbilitySpecHandle AbilityToActivate, FPredictionKey PredictionKey, FGameplayEventData TriggerEventData);
+	virtual void ServerTryActivateAbilityWithEventData_Implementation(FGameplayAbilitySpecHandle AbilityToActivate, FPredictionKey PredictionKey, FGameplayEventData TriggerEventData);
+	virtual bool ServerTryActivateAbilityWithEventData_Validate(FGameplayAbilitySpecHandle , FPredictionKey , FGameplayEventData );
 
 	/** Replicate that an ability has ended/canceled, to the client or server as appropriate */
 	void	ReplicateEndOrCancelAbility(FGameplayAbilitySpecHandle Handle, FGameplayAbilityActivationInfo ActivationInfo, UGameplayAbility* Ability, bool bWasCanceled);
@@ -662,28 +672,37 @@ public:
 	/** Called by ServerEndAbility/ServerCancelAbility and ClientEndAbility/ClientCancelAbility; avoids code duplication. */
 	void	RemoteEndOrCancelAbility(FGameplayAbilitySpecHandle AbilityToEnd, FGameplayAbilityActivationInfo ActivationInfo, bool bWasCanceled);
 
-	UFUNCTION(Server, reliable, WithValidation)
+	UFUNCTION(Server="ServerEndAbility_Implementation", reliable, WithValidation="ServerEndAbility_Validate")
 	void	ServerEndAbility(FGameplayAbilitySpecHandle AbilityToEnd, FGameplayAbilityActivationInfo ActivationInfo, FPredictionKey PredictionKey);
+	virtual void ServerEndAbility_Implementation(FGameplayAbilitySpecHandle AbilityToEnd, FGameplayAbilityActivationInfo ActivationInfo, FPredictionKey PredictionKey);
+	virtual bool ServerEndAbility_Validate(FGameplayAbilitySpecHandle , FGameplayAbilityActivationInfo , FPredictionKey );
 
-	UFUNCTION(Client, reliable)
+	UFUNCTION(Client="ClientEndAbility_Implementation", reliable)
 	void	ClientEndAbility(FGameplayAbilitySpecHandle AbilityToEnd, FGameplayAbilityActivationInfo ActivationInfo);
+	virtual void ClientEndAbility_Implementation(FGameplayAbilitySpecHandle AbilityToEnd, FGameplayAbilityActivationInfo ActivationInfo);
 
-	UFUNCTION(Server, reliable, WithValidation)
+	UFUNCTION(Server="ServerCancelAbility_Implementation", reliable, WithValidation="ServerCancelAbility_Validate")
 	void    ServerCancelAbility(FGameplayAbilitySpecHandle AbilityToCancel, FGameplayAbilityActivationInfo ActivationInfo);
+	virtual void ServerCancelAbility_Implementation(FGameplayAbilitySpecHandle AbilityToCancel, FGameplayAbilityActivationInfo ActivationInfo);
+	virtual bool ServerCancelAbility_Validate(FGameplayAbilitySpecHandle , FGameplayAbilityActivationInfo );
 
-	UFUNCTION(Client, reliable)
+	UFUNCTION(Client="ClientCancelAbility_Implementation", reliable)
 	void    ClientCancelAbility(FGameplayAbilitySpecHandle AbilityToCancel, FGameplayAbilityActivationInfo ActivationInfo);
+	virtual void ClientCancelAbility_Implementation(FGameplayAbilitySpecHandle AbilityToCancel, FGameplayAbilityActivationInfo ActivationInfo);
 
-	UFUNCTION(Client, Reliable)
+	UFUNCTION(Client="ClientActivateAbilityFailed_Implementation", Reliable)
 	void	ClientActivateAbilityFailed(FGameplayAbilitySpecHandle AbilityToActivate, int16 PredictionKey);
+	virtual void ClientActivateAbilityFailed_Implementation(FGameplayAbilitySpecHandle AbilityToActivate, int16 PredictionKey);
 
 	void	OnClientActivateAbilityFailed(FGameplayAbilitySpecHandle AbilityToActivate, FPredictionKey::KeyType PredictionKey);
 
-	UFUNCTION(Client, Reliable)
+	UFUNCTION(Client="ClientActivateAbilitySucceed_Implementation", Reliable)
 	void	ClientActivateAbilitySucceed(FGameplayAbilitySpecHandle AbilityToActivate, int16 PredictionKey);
+	virtual void ClientActivateAbilitySucceed_Implementation(FGameplayAbilitySpecHandle AbilityToActivate, int16 PredictionKey);
 
-	UFUNCTION(Client, Reliable)
+	UFUNCTION(Client="ClientActivateAbilitySucceedWithEventData_Implementation", Reliable)
 	void	ClientActivateAbilitySucceedWithEventData(FGameplayAbilitySpecHandle AbilityToActivate, int16 PredictionKey, FGameplayEventData TriggerEventData);
+	virtual void ClientActivateAbilitySucceedWithEventData_Implementation(FGameplayAbilitySpecHandle AbilityToActivate, int16 PredictionKey, FGameplayEventData TriggerEventData);
 
 	// ----------------------------------------------------------------------------------------------------------------
 
@@ -828,12 +847,16 @@ protected:
 	void OnRep_ReplicatedAnimMontage();
 
 	/** RPC function called from CurrentMontageSetNextSectopnName, replicates to other clients */
-	UFUNCTION(reliable, server, WithValidation)
+	UFUNCTION(reliable, server="ServerCurrentMontageSetNextSectionName_Implementation", WithValidation="ServerCurrentMontageSetNextSectionName_Validate")
 	void ServerCurrentMontageSetNextSectionName(UAnimMontage* ClientAnimMontage, float ClientPosition, FName SectionName, FName NextSectionName);
+	virtual void ServerCurrentMontageSetNextSectionName_Implementation(UAnimMontage* ClientAnimMontage, float ClientPosition, FName SectionName, FName NextSectionName);
+	virtual bool ServerCurrentMontageSetNextSectionName_Validate(UAnimMontage* , float , FName , FName );
 
 	/** RPC function called from CurrentMontageJumpToSection, replicates to other clients */
-	UFUNCTION(reliable, server, WithValidation)
+	UFUNCTION(reliable, server="ServerCurrentMontageJumpToSectionName_Implementation", WithValidation="ServerCurrentMontageJumpToSectionName_Validate")
 	void ServerCurrentMontageJumpToSectionName(UAnimMontage* ClientAnimMontage, FName SectionName);
+	virtual void ServerCurrentMontageJumpToSectionName_Implementation(UAnimMontage* ClientAnimMontage, FName SectionName);
+	virtual bool ServerCurrentMontageJumpToSectionName_Validate(UAnimMontage* , FName );
 
 	/** Abilities that are triggered from a gameplay event */
 	TMap<FGameplayTag, TArray<FGameplayAbilitySpecHandle > > GameplayEventTriggeredAbilities;
@@ -915,23 +938,30 @@ public:
 	 */
 
 	/** Replicates the Generic Replicated Event to the server. */
-	UFUNCTION(Server, reliable, WithValidation)
+	UFUNCTION(Server="ServerSetReplicatedEvent_Implementation", reliable, WithValidation="ServerSetReplicatedEvent_Validate")
 	void ServerSetReplicatedEvent(EAbilityGenericReplicatedEvent::Type EventType, FGameplayAbilitySpecHandle AbilityHandle, FPredictionKey AbilityOriginalPredictionKey, FPredictionKey CurrentPredictionKey);
+	virtual void ServerSetReplicatedEvent_Implementation(EAbilityGenericReplicatedEvent::Type EventType, FGameplayAbilitySpecHandle AbilityHandle, FPredictionKey AbilityOriginalPredictionKey, FPredictionKey CurrentPredictionKey);
+	virtual bool ServerSetReplicatedEvent_Validate(EAbilityGenericReplicatedEvent::Type , FGameplayAbilitySpecHandle , FPredictionKey , FPredictionKey );
 
 	/** Replicates the Generic Replicated Event to the server. */
-	UFUNCTION(Client, reliable)
+	UFUNCTION(Client="ClientSetReplicatedEvent_Implementation", reliable)
 	void ClientSetReplicatedEvent(EAbilityGenericReplicatedEvent::Type EventType, FGameplayAbilitySpecHandle AbilityHandle, FPredictionKey AbilityOriginalPredictionKey);
+	virtual void ClientSetReplicatedEvent_Implementation(EAbilityGenericReplicatedEvent::Type EventType, FGameplayAbilitySpecHandle AbilityHandle, FPredictionKey AbilityOriginalPredictionKey);
 
 	/** Calls local callbacks that are registered with the given Generic Replicated Event */
 	bool InvokeReplicatedEvent(EAbilityGenericReplicatedEvent::Type EventType, FGameplayAbilitySpecHandle AbilityHandle, FPredictionKey AbilityOriginalPredictionKey);
 
 	/**  */
-	UFUNCTION(Server, reliable, WithValidation)
+	UFUNCTION(Server="ServerSetReplicatedTargetData_Implementation", reliable, WithValidation="ServerSetReplicatedTargetData_Validate")
 	void ServerSetReplicatedTargetData(FGameplayAbilitySpecHandle AbilityHandle, FPredictionKey AbilityOriginalPredictionKey, FGameplayAbilityTargetDataHandle ReplicatedTargetData, FGameplayTag ApplicationTag, FPredictionKey CurrentPredictionKey);
+	virtual void ServerSetReplicatedTargetData_Implementation(FGameplayAbilitySpecHandle AbilityHandle, FPredictionKey AbilityOriginalPredictionKey, FGameplayAbilityTargetDataHandle ReplicatedTargetData, FGameplayTag ApplicationTag, FPredictionKey CurrentPredictionKey);
+	virtual bool ServerSetReplicatedTargetData_Validate(FGameplayAbilitySpecHandle , FPredictionKey , FGameplayAbilityTargetDataHandle , FGameplayTag , FPredictionKey );
 
 	/** Replicates to the server that targeting has been cancelled */
-	UFUNCTION(Server, reliable, WithValidation)
+	UFUNCTION(Server="ServerSetReplicatedTargetDataCancelled_Implementation", reliable, WithValidation="ServerSetReplicatedTargetDataCancelled_Validate")
 	void ServerSetReplicatedTargetDataCancelled(FGameplayAbilitySpecHandle AbilityHandle, FPredictionKey AbilityOriginalPredictionKey, FPredictionKey CurrentPredictionKey);
+	virtual void ServerSetReplicatedTargetDataCancelled_Implementation(FGameplayAbilitySpecHandle AbilityHandle, FPredictionKey AbilityOriginalPredictionKey, FPredictionKey CurrentPredictionKey);
+	virtual bool ServerSetReplicatedTargetDataCancelled_Validate(FGameplayAbilitySpecHandle , FPredictionKey , FPredictionKey );
 
 	/** Sets the current target data and calls applicable callbacks */
 	void ConfirmAbilityTargetData(FGameplayAbilitySpecHandle AbilityHandle, FPredictionKey AbilityOriginalPredictionKey, const FGameplayAbilityTargetDataHandle& TargetData, const FGameplayTag& ApplicationTag);
@@ -970,11 +1000,15 @@ public:
 	FSimpleMulticastDelegate& AbilityReplicatedEventDelegate(EAbilityGenericReplicatedEvent::Type EventType, FGameplayAbilitySpecHandle AbilityHandle, FPredictionKey AbilityOriginalPredictionKey);
 
 	// Direct Input state replication. These will be called if bReplicateInputDirectly is true on the ability and is generally not a good thing to use. (Instead, prefer to use Generic Replicated Events).
-	UFUNCTION(Server, reliable, WithValidation)
+	UFUNCTION(Server="ServerSetInputPressed_Implementation", reliable, WithValidation="ServerSetInputPressed_Validate")
 	void ServerSetInputPressed(FGameplayAbilitySpecHandle AbilityHandle);
+	virtual void ServerSetInputPressed_Implementation(FGameplayAbilitySpecHandle AbilityHandle);
+	virtual bool ServerSetInputPressed_Validate(FGameplayAbilitySpecHandle );
 
-	UFUNCTION(Server, reliable, WithValidation)
+	UFUNCTION(Server="ServerSetInputReleased_Implementation", reliable, WithValidation="ServerSetInputReleased_Validate")
 	void ServerSetInputReleased(FGameplayAbilitySpecHandle AbilityHandle);
+	virtual void ServerSetInputReleased_Implementation(FGameplayAbilitySpecHandle AbilityHandle);
+	virtual bool ServerSetInputReleased_Validate(FGameplayAbilitySpecHandle );
 
 	/** Called on local player always. Called on server only if bReplicateInputDirectly is set on the GameplayAbility. */
 	void AbilitySpecInputReleased(FGameplayAbilitySpec& Spec);
diff --git a/Engine/Source/Runtime/GameplayAbilities/Public/AbilitySystemDebugHUD.h b/Engine/Source/Runtime/GameplayAbilities/Public/AbilitySystemDebugHUD.h
--- a/Engine/Source/Runtime/GameplayAbilities/Public/AbilitySystemDebugHUD.h
+++ b/Engine/Source/Runtime/GameplayAbilities/Public/AbilitySystemDebugHUD.h
@@ -27,7 +27,9 @@ namespace EAlignVertical
 UCLASS()
 class AAbilitySystemDebugHUD : public AHUD
 {
-	GENERATED_UCLASS_BODY()
+	GENERATED_BODY()
+public:
+	AAbilitySystemDebugHUD(const FObjectInitializer& ObjectInitializer = FObjectInitializer::Get());
 
 	/** main HUD update loop */
 	void DrawWithBackground(UFont* InFont, const FString& Text, const FColor& TextColor, EAlignHorizontal::Type HAlign, float& OffsetX, EAlignVertical::Type VAlign, float& OffsetY, float Alpha=1.f);
@@ -38,4 +40,4 @@ private:
 
 	void DrawDebugAbilitySystemComponent(UAbilitySystemComponent *Component);
 
-};
\ No newline at end of file
+};
diff --git a/Engine/Source/Runtime/GameplayAbilities/Public/AbilitySystemGlobals.h b/Engine/Source/Runtime/GameplayAbilities/Public/AbilitySystemGlobals.h
--- a/Engine/Source/Runtime/GameplayAbilities/Public/AbilitySystemGlobals.h
+++ b/Engine/Source/Runtime/GameplayAbilities/Public/AbilitySystemGlobals.h
@@ -21,7 +21,9 @@ struct FGameplayEffectSpec;
 UCLASS(config=Game)
 class GAMEPLAYABILITIES_API UAbilitySystemGlobals : public UObject
 {
-	GENERATED_UCLASS_BODY()
+	GENERATED_BODY()
+public:
+	UAbilitySystemGlobals(const FObjectInitializer& ObjectInitializer = FObjectInitializer::Get());
 
 	/** Gets the single instance of the globals object, will create it as necessary */
 	static UAbilitySystemGlobals& Get();
diff --git a/Engine/Source/Runtime/GameplayAbilities/Public/AbilitySystemInterface.h b/Engine/Source/Runtime/GameplayAbilities/Public/AbilitySystemInterface.h
--- a/Engine/Source/Runtime/GameplayAbilities/Public/AbilitySystemInterface.h
+++ b/Engine/Source/Runtime/GameplayAbilities/Public/AbilitySystemInterface.h
@@ -12,12 +12,15 @@ class UAbilitySystemComponent;
 UINTERFACE(MinimalAPI, meta = (CannotImplementInterfaceInBlueprint))
 class UAbilitySystemInterface : public UInterface
 {
-	GENERATED_UINTERFACE_BODY()
+	GENERATED_BODY()
+public:
+	GAMEPLAYABILITIES_API UAbilitySystemInterface(const FObjectInitializer& ObjectInitializer = FObjectInitializer::Get());
 };
 
 class GAMEPLAYABILITIES_API IAbilitySystemInterface
 {
-	GENERATED_IINTERFACE_BODY()
+	GENERATED_BODY()
+public:
 
 	virtual UAbilitySystemComponent* GetAbilitySystemComponent() const = 0;
-};
\ No newline at end of file
+};
diff --git a/Engine/Source/Runtime/GameplayAbilities/Public/AbilitySystemTestAttributeSet.h b/Engine/Source/Runtime/GameplayAbilities/Public/AbilitySystemTestAttributeSet.h
--- a/Engine/Source/Runtime/GameplayAbilities/Public/AbilitySystemTestAttributeSet.h
+++ b/Engine/Source/Runtime/GameplayAbilities/Public/AbilitySystemTestAttributeSet.h
@@ -7,7 +7,9 @@
 UCLASS(Blueprintable, BlueprintType, meta=(HideInDetailsView))
 class GAMEPLAYABILITIES_API UAbilitySystemTestAttributeSet : public UAttributeSet
 {
-	GENERATED_UCLASS_BODY()
+	GENERATED_BODY()
+public:
+	UAbilitySystemTestAttributeSet(const FObjectInitializer& ObjectInitializer = FObjectInitializer::Get());
 
 
 
diff --git a/Engine/Source/Runtime/GameplayAbilities/Public/AbilitySystemTestPawn.h b/Engine/Source/Runtime/GameplayAbilities/Public/AbilitySystemTestPawn.h
--- a/Engine/Source/Runtime/GameplayAbilities/Public/AbilitySystemTestPawn.h
+++ b/Engine/Source/Runtime/GameplayAbilities/Public/AbilitySystemTestPawn.h
@@ -8,7 +8,9 @@
 UCLASS(Blueprintable, BlueprintType, notplaceable)
 class GAMEPLAYABILITIES_API AAbilitySystemTestPawn : public ADefaultPawn, public IGameplayCueInterface, public IAbilitySystemInterface
 {
-	GENERATED_UCLASS_BODY()
+	GENERATED_BODY()
+public:
+	AAbilitySystemTestPawn(const FObjectInitializer& ObjectInitializer = FObjectInitializer::Get());
 
 	virtual void PostInitializeComponents() override;
 
diff --git a/Engine/Source/Runtime/GameplayAbilities/Public/AttributeSet.h b/Engine/Source/Runtime/GameplayAbilities/Public/AttributeSet.h
--- a/Engine/Source/Runtime/GameplayAbilities/Public/AttributeSet.h
+++ b/Engine/Source/Runtime/GameplayAbilities/Public/AttributeSet.h
@@ -12,7 +12,7 @@ struct FGameplayAbilityActorInfo;
 USTRUCT(BlueprintType)
 struct GAMEPLAYABILITIES_API FGameplayAttribute
 {
-	GENERATED_USTRUCT_BODY()
+	GENERATED_BODY()
 
 	FGameplayAttribute()
 		: Attribute(nullptr)
@@ -74,7 +74,9 @@ private:
 UCLASS(DefaultToInstanced, Blueprintable)
 class GAMEPLAYABILITIES_API UAttributeSet : public UObject
 {
-	GENERATED_UCLASS_BODY()
+	GENERATED_BODY()
+public:
+	UAttributeSet(const FObjectInitializer& ObjectInitializer = FObjectInitializer::Get());
 
 public:
 	bool IsNameStableForNetworking() const override;
@@ -144,7 +146,7 @@ protected:
 USTRUCT()
 struct GAMEPLAYABILITIES_API FGlobalCurveDataOverride
 {
-	GENERATED_USTRUCT_BODY()
+	GENERATED_BODY()
 
 	TArray<UCurveTable*>	Overrides;
 };
@@ -156,7 +158,7 @@ struct GAMEPLAYABILITIES_API FGlobalCurveDataOverride
 USTRUCT()
 struct GAMEPLAYABILITIES_API FScalableFloat
 {
-	GENERATED_USTRUCT_BODY()
+	GENERATED_BODY()
 
 	FScalableFloat()
 		: Value(0.f)
@@ -237,7 +239,7 @@ private:
 USTRUCT(BlueprintType)
 struct GAMEPLAYABILITIES_API FAttributeMetaData : public FTableRowBase
 {
-	GENERATED_USTRUCT_BODY()
+	GENERATED_BODY()
 
 public:
 
diff --git a/Engine/Source/Runtime/GameplayAbilities/Public/GameplayAbilityBlueprintGeneratedClass.h b/Engine/Source/Runtime/GameplayAbilities/Public/GameplayAbilityBlueprintGeneratedClass.h
--- a/Engine/Source/Runtime/GameplayAbilities/Public/GameplayAbilityBlueprintGeneratedClass.h
+++ b/Engine/Source/Runtime/GameplayAbilities/Public/GameplayAbilityBlueprintGeneratedClass.h
@@ -7,5 +7,7 @@
 UCLASS()
 class  GAMEPLAYABILITIES_API UGameplayAbilityBlueprintGeneratedClass : public UBlueprintGeneratedClass
 {
-	GENERATED_UCLASS_BODY()
-};
\ No newline at end of file
+	GENERATED_BODY()
+public:
+	UGameplayAbilityBlueprintGeneratedClass(const FObjectInitializer& ObjectInitializer = FObjectInitializer::Get());
+};
diff --git a/Engine/Source/Runtime/GameplayAbilities/Public/GameplayAbilitySet.h b/Engine/Source/Runtime/GameplayAbilities/Public/GameplayAbilitySet.h
--- a/Engine/Source/Runtime/GameplayAbilities/Public/GameplayAbilitySet.h
+++ b/Engine/Source/Runtime/GameplayAbilities/Public/GameplayAbilitySet.h
@@ -38,7 +38,7 @@ namespace EGameplayAbilityInputBinds
 USTRUCT()
 struct FGameplayAbilityBindInfo
 {
-	GENERATED_USTRUCT_BODY()
+	GENERATED_BODY()
 
 	UPROPERTY(EditAnywhere, Category = BindInfo)
 	TEnumAsByte<EGameplayAbilityInputBinds::Type>	Command;
@@ -56,11 +56,13 @@ struct FGameplayAbilityBindInfo
 UCLASS()
 class GAMEPLAYABILITIES_API UGameplayAbilitySet : public UDataAsset
 {
-	GENERATED_UCLASS_BODY()
+	GENERATED_BODY()
+public:
+	UGameplayAbilitySet(const FObjectInitializer& ObjectInitializer = FObjectInitializer::Get());
 
 	
 	UPROPERTY(EditAnywhere, Category = AbilitySet)
 	TArray<FGameplayAbilityBindInfo>	Abilities;
 
 	void GiveAbilities(UAbilitySystemComponent* AbilitySystemComponent) const;
-};
\ No newline at end of file
+};
diff --git a/Engine/Source/Runtime/GameplayAbilities/Public/GameplayAbilitySpec.h b/Engine/Source/Runtime/GameplayAbilities/Public/GameplayAbilitySpec.h
--- a/Engine/Source/Runtime/GameplayAbilities/Public/GameplayAbilitySpec.h
+++ b/Engine/Source/Runtime/GameplayAbilities/Public/GameplayAbilitySpec.h
@@ -14,7 +14,7 @@
 USTRUCT(BlueprintType)
 struct FGameplayAbilitySpecHandle
 {
-	GENERATED_USTRUCT_BODY()
+	GENERATED_BODY()
 
 	FGameplayAbilitySpecHandle()
 	: Handle(INDEX_NONE)
@@ -89,6 +89,7 @@ enum class EGameplayEffectGrantedAbilityRemovePolicy : uint8
 USTRUCT(BlueprintType)
 struct FGameplayAbilitySpecDef
 {
+	GENERATED_BODY()
 	FGameplayAbilitySpecDef()
 		: Level(1)
 		, InputID(INDEX_NONE)
@@ -96,7 +97,6 @@ struct FGameplayAbilitySpecDef
 	{
 	}
 
-	GENERATED_USTRUCT_BODY()
 
 	UPROPERTY(EditDefaultsOnly, Category="Ability Definition", NotReplicated)
 	TSubclassOf<UGameplayAbility> Ability;
@@ -130,7 +130,7 @@ struct FGameplayAbilitySpecDef
 USTRUCT(BlueprintType)
 struct GAMEPLAYABILITIES_API FGameplayAbilityActivationInfo
 {
-	GENERATED_USTRUCT_BODY()
+	GENERATED_BODY()
 
 	FGameplayAbilityActivationInfo()
 		: ActivationMode(EGameplayAbilityActivationMode::Authority)
@@ -186,7 +186,7 @@ private:
 USTRUCT()
 struct GAMEPLAYABILITIES_API FGameplayAbilitySpec : public FFastArraySerializerItem
 {
-	GENERATED_USTRUCT_BODY()
+	GENERATED_BODY()
 
 	FGameplayAbilitySpec()
 	: Ability(nullptr), Level(1), InputID(INDEX_NONE), SourceObject(nullptr), InputPressed(false), ActiveCount(0), RemoveAfterActivation(false), PendingRemove(false)
@@ -278,7 +278,7 @@ struct GAMEPLAYABILITIES_API FGameplayAbilitySpec : public FFastArraySerializerI
 USTRUCT()
 struct GAMEPLAYABILITIES_API FGameplayAbilitySpecContainer : public FFastArraySerializer
 {
-	GENERATED_USTRUCT_BODY()
+	GENERATED_BODY()
 
 	/** List of activatable abilities */
 	UPROPERTY()
@@ -315,4 +315,4 @@ private:
 	UAbilitySystemComponent& AbilitySystemComponent;
 };
 
-#define ABILITYLIST_SCOPE_LOCK()	FScopedAbilityListLock ActiveScopeLock(*this);
\ No newline at end of file
+#define ABILITYLIST_SCOPE_LOCK()	FScopedAbilityListLock ActiveScopeLock(*this);
diff --git a/Engine/Source/Runtime/GameplayAbilities/Public/GameplayCueInterface.h b/Engine/Source/Runtime/GameplayAbilities/Public/GameplayCueInterface.h
--- a/Engine/Source/Runtime/GameplayAbilities/Public/GameplayCueInterface.h
+++ b/Engine/Source/Runtime/GameplayAbilities/Public/GameplayCueInterface.h
@@ -11,12 +11,15 @@
 UINTERFACE(MinimalAPI, meta = (CannotImplementInterfaceInBlueprint))
 class UGameplayCueInterface: public UInterface
 {
-	GENERATED_UINTERFACE_BODY()
+	GENERATED_BODY()
+public:
+	GAMEPLAYABILITIES_API UGameplayCueInterface(const FObjectInitializer& ObjectInitializer = FObjectInitializer::Get());
 };
 
 class GAMEPLAYABILITIES_API IGameplayCueInterface
 {
-	GENERATED_IINTERFACE_BODY()
+	GENERATED_BODY()
+public:
 
 	virtual void HandleGameplayCue(AActor *Self, FGameplayTag GameplayCueTag, EGameplayCueEvent::Type EventType, FGameplayCueParameters Parameters);
 
@@ -58,7 +61,7 @@ struct FActiveGameplayCueContainer;
 USTRUCT(BlueprintType)
 struct FActiveGameplayCue : public FFastArraySerializerItem
 {
-	GENERATED_USTRUCT_BODY()
+	GENERATED_BODY()
 
 	UPROPERTY()
 	FGameplayTag GameplayCueTag;
@@ -74,7 +77,7 @@ struct FActiveGameplayCue : public FFastArraySerializerItem
 USTRUCT(BlueprintType)
 struct FActiveGameplayCueContainer : public FFastArraySerializer
 {
-	GENERATED_USTRUCT_BODY()
+	GENERATED_BODY()
 
 	UPROPERTY()
 	TArray< FActiveGameplayCue >	GameplayCues;
@@ -102,4 +105,4 @@ struct TStructOpsTypeTraits< FActiveGameplayCueContainer > : public TStructOpsTy
 	{
 		WithNetDeltaSerializer = true,
 	};
-};
\ No newline at end of file
+};
diff --git a/Engine/Source/Runtime/GameplayAbilities/Public/GameplayCueManager.h b/Engine/Source/Runtime/GameplayAbilities/Public/GameplayCueManager.h
--- a/Engine/Source/Runtime/GameplayAbilities/Public/GameplayCueManager.h
+++ b/Engine/Source/Runtime/GameplayAbilities/Public/GameplayCueManager.h
@@ -48,7 +48,7 @@
 USTRUCT()
 struct FGameplayCueNotifyData
 {
-	GENERATED_USTRUCT_BODY()
+	GENERATED_BODY()
 
 	FGameplayCueNotifyData()
 	: LoadedGameplayCueClass(nullptr)
@@ -80,7 +80,7 @@ enum class EGameplayCuePayloadType : uint8
 USTRUCT()
 struct FGameplayCuePendingExecute
 {
-	GENERATED_USTRUCT_BODY()
+	GENERATED_BODY()
 
 	FGameplayCuePendingExecute()
 	: PayloadType(EGameplayCuePayloadType::EffectContext)
@@ -129,7 +129,9 @@ struct GAMEPLAYABILITIES_API FScopedGameplayCueSendContext
 UCLASS()
 class GAMEPLAYABILITIES_API UGameplayCueManager : public UDataAsset
 {
-	GENERATED_UCLASS_BODY()
+	GENERATED_BODY()
+public:
+	UGameplayCueManager(const FObjectInitializer& ObjectInitializer = FObjectInitializer::Get());
 
 	// -------------------------------------------------------------
 	// Wrappers to handle replicating executed cues
diff --git a/Engine/Source/Runtime/GameplayAbilities/Public/GameplayCueNotify_Actor.h b/Engine/Source/Runtime/GameplayAbilities/Public/GameplayCueNotify_Actor.h
--- a/Engine/Source/Runtime/GameplayAbilities/Public/GameplayCueNotify_Actor.h
+++ b/Engine/Source/Runtime/GameplayAbilities/Public/GameplayCueNotify_Actor.h
@@ -22,7 +22,9 @@
 UCLASS(Blueprintable, meta = (ShowWorldContextPin), hidecategories = (Replication))
 class GAMEPLAYABILITIES_API AGameplayCueNotify_Actor : public AActor
 {
-	GENERATED_UCLASS_BODY()
+	GENERATED_BODY()
+public:
+	AGameplayCueNotify_Actor(const FObjectInitializer& ObjectInitializer = FObjectInitializer::Get());
 
 	/** Does this GameplayCueNotify handle this type of GameplayCueEvent? */
 	virtual bool HandlesEvent(EGameplayCueEvent::Type EventType) const;
@@ -45,15 +47,19 @@ class GAMEPLAYABILITIES_API AGameplayCueNotify_Actor : public AActor
 
 	UFUNCTION(BlueprintNativeEvent, Category = "GameplayCueNotify")
 	bool OnExecute(AActor* MyTarget, FGameplayCueParameters Parameters);
+	virtual bool OnExecute_Implementation(AActor* MyTarget, FGameplayCueParameters Parameters);
 
 	UFUNCTION(BlueprintNativeEvent, Category = "GameplayCueNotify")
 	bool OnActive(AActor* MyTarget, FGameplayCueParameters Parameters);
+	virtual bool OnActive_Implementation(AActor* MyTarget, FGameplayCueParameters Parameters);
 
 	UFUNCTION(BlueprintNativeEvent, Category = "GameplayCueNotify")
 	bool WhileActive(AActor* MyTarget, FGameplayCueParameters Parameters);
+	virtual bool WhileActive_Implementation(AActor* MyTarget, FGameplayCueParameters Parameters);
 
 	UFUNCTION(BlueprintNativeEvent, Category = "GameplayCueNotify")
 	bool OnRemove(AActor* MyTarget, FGameplayCueParameters Parameters);
+	virtual bool OnRemove_Implementation(AActor* MyTarget, FGameplayCueParameters Parameters);
 
 	UPROPERTY(EditDefaultsOnly, Category = GameplayCue)
 	FGameplayTag	GameplayCueTag;
@@ -72,4 +78,4 @@ class GAMEPLAYABILITIES_API AGameplayCueNotify_Actor : public AActor
 
 private:
 	virtual void DeriveGameplayCueTagFromAssetName();
-};
\ No newline at end of file
+};
diff --git a/Engine/Source/Runtime/GameplayAbilities/Public/GameplayCueNotify_HitImpact.h b/Engine/Source/Runtime/GameplayAbilities/Public/GameplayCueNotify_HitImpact.h
--- a/Engine/Source/Runtime/GameplayAbilities/Public/GameplayCueNotify_HitImpact.h
+++ b/Engine/Source/Runtime/GameplayAbilities/Public/GameplayCueNotify_HitImpact.h
@@ -15,7 +15,9 @@
 UCLASS()
 class GAMEPLAYABILITIES_API UGameplayCueNotify_HitImpact : public UGameplayCueNotify_Static
 {
-	GENERATED_UCLASS_BODY()
+	GENERATED_BODY()
+public:
+	UGameplayCueNotify_HitImpact(const FObjectInitializer& ObjectInitializer = FObjectInitializer::Get());
 
 	/** Does this GameplayCueNotify handle this type of GameplayCueEvent? */
 	virtual bool HandlesEvent(EGameplayCueEvent::Type EventType) const override;
diff --git a/Engine/Source/Runtime/GameplayAbilities/Public/GameplayCueNotify_Static.h b/Engine/Source/Runtime/GameplayAbilities/Public/GameplayCueNotify_Static.h
--- a/Engine/Source/Runtime/GameplayAbilities/Public/GameplayCueNotify_Static.h
+++ b/Engine/Source/Runtime/GameplayAbilities/Public/GameplayCueNotify_Static.h
@@ -22,7 +22,9 @@
 UCLASS(Blueprintable, meta = (ShowWorldContextPin), hidecategories = (Replication))
 class GAMEPLAYABILITIES_API UGameplayCueNotify_Static : public UObject
 {
-	GENERATED_UCLASS_BODY()
+	GENERATED_BODY()
+public:
+	UGameplayCueNotify_Static(const FObjectInitializer& ObjectInitializer = FObjectInitializer::Get());
 
 	/** Does this GameplayCueNotify handle this type of GameplayCueEvent? */
 	virtual bool HandlesEvent(EGameplayCueEvent::Type EventType) const;
@@ -45,15 +47,19 @@ class GAMEPLAYABILITIES_API UGameplayCueNotify_Static : public UObject
 
 	UFUNCTION(BlueprintNativeEvent, BlueprintPure, Category = "GameplayCueNotify")
 	bool OnExecute(AActor* MyTarget, FGameplayCueParameters Parameters) const;
+	virtual bool OnExecute_Implementation(AActor* MyTarget, FGameplayCueParameters Parameters) const;
 
 	UFUNCTION(BlueprintNativeEvent, BlueprintPure, Category = "GameplayCueNotify")
 	bool OnActive(AActor* MyTarget, FGameplayCueParameters Parameters) const;
+	virtual bool OnActive_Implementation(AActor* MyTarget, FGameplayCueParameters Parameters) const;
 
 	UFUNCTION(BlueprintNativeEvent, BlueprintPure, Category = "GameplayCueNotify")
 	bool WhileActive(AActor* MyTarget, FGameplayCueParameters Parameters) const;
+	virtual bool WhileActive_Implementation(AActor* MyTarget, FGameplayCueParameters Parameters) const;
 
 	UFUNCTION(BlueprintNativeEvent, BlueprintPure, Category = "GameplayCueNotify")
 	bool OnRemove(AActor* MyTarget, FGameplayCueParameters Parameters) const;
+	virtual bool OnRemove_Implementation(AActor* MyTarget, FGameplayCueParameters Parameters) const;
 
 	UPROPERTY(EditDefaultsOnly, Category = GameplayCue)
 	FGameplayTag	GameplayCueTag;
@@ -68,4 +74,4 @@ class GAMEPLAYABILITIES_API UGameplayCueNotify_Static : public UObject
 
 private:
 	virtual void DeriveGameplayCueTagFromAssetName();
-};
\ No newline at end of file
+};
diff --git a/Engine/Source/Runtime/GameplayAbilities/Public/GameplayEffectCalculation.h b/Engine/Source/Runtime/GameplayAbilities/Public/GameplayEffectCalculation.h
--- a/Engine/Source/Runtime/GameplayAbilities/Public/GameplayEffectCalculation.h
+++ b/Engine/Source/Runtime/GameplayAbilities/Public/GameplayEffectCalculation.h
@@ -12,7 +12,9 @@ class UAbilitySystemComponent;
 UCLASS(BlueprintType, Blueprintable, Abstract)
 class GAMEPLAYABILITIES_API UGameplayEffectCalculation : public UObject
 {
-	GENERATED_UCLASS_BODY()
+	GENERATED_BODY()
+public:
+	UGameplayEffectCalculation(const FObjectInitializer& ObjectInitializer = FObjectInitializer::Get());
 
 public:
 
diff --git a/Engine/Source/Runtime/GameplayAbilities/Public/GameplayEffectExecutionCalculation.h b/Engine/Source/Runtime/GameplayAbilities/Public/GameplayEffectExecutionCalculation.h
--- a/Engine/Source/Runtime/GameplayAbilities/Public/GameplayEffectExecutionCalculation.h
+++ b/Engine/Source/Runtime/GameplayAbilities/Public/GameplayEffectExecutionCalculation.h
@@ -14,7 +14,7 @@ class UAbilitySystemComponent;
 USTRUCT()
 struct GAMEPLAYABILITIES_API FGameplayEffectCustomExecutionParameters
 {
-	GENERATED_USTRUCT_BODY()
+	GENERATED_BODY()
 
 public:
 
@@ -112,7 +112,7 @@ private:
 USTRUCT()
 struct GAMEPLAYABILITIES_API FGameplayEffectCustomExecutionOutput
 {
-	GENERATED_USTRUCT_BODY()
+	GENERATED_BODY()
 
 public:
 
@@ -168,7 +168,9 @@ private:
 UCLASS(BlueprintType, Blueprintable, Abstract)
 class GAMEPLAYABILITIES_API UGameplayEffectExecutionCalculation : public UGameplayEffectCalculation
 {
-	GENERATED_UCLASS_BODY()
+	GENERATED_BODY()
+public:
+	UGameplayEffectExecutionCalculation(const FObjectInitializer& ObjectInitializer = FObjectInitializer::Get());
 
 protected:
 
@@ -210,4 +212,5 @@ public:
 	 */
 	UFUNCTION(BlueprintNativeEvent, Category="Calculation")
 	void Execute(const FGameplayEffectCustomExecutionParameters& ExecutionParams, FGameplayEffectCustomExecutionOutput& OutExecutionOutput) const;
+	virtual void Execute_Implementation(const FGameplayEffectCustomExecutionParameters& ExecutionParams, FGameplayEffectCustomExecutionOutput& OutExecutionOutput) const;
 };
diff --git a/Engine/Source/Runtime/GameplayAbilities/Public/GameplayEffectExtension.h b/Engine/Source/Runtime/GameplayAbilities/Public/GameplayEffectExtension.h
--- a/Engine/Source/Runtime/GameplayAbilities/Public/GameplayEffectExtension.h
+++ b/Engine/Source/Runtime/GameplayAbilities/Public/GameplayEffectExtension.h
@@ -27,7 +27,9 @@ struct FGameplayEffectModCallbackData
 UCLASS(BlueprintType)
 class GAMEPLAYABILITIES_API UGameplayEffectExtension: public UObject
 {
-	GENERATED_UCLASS_BODY()
+	GENERATED_BODY()
+public:
+	UGameplayEffectExtension(const FObjectInitializer& ObjectInitializer = FObjectInitializer::Get());
 
 	/** Attributes on the source instigator that are relevant to calculating modifications using this extension */
 	UPROPERTY(EditDefaultsOnly, Category=Calculation)
@@ -48,4 +50,4 @@ public:
 	{
 
 	}
-};
\ No newline at end of file
+};
diff --git a/Engine/Source/Runtime/GameplayAbilities/Public/GameplayEffectExtension_LifestealTest.h b/Engine/Source/Runtime/GameplayAbilities/Public/GameplayEffectExtension_LifestealTest.h
--- a/Engine/Source/Runtime/GameplayAbilities/Public/GameplayEffectExtension_LifestealTest.h
+++ b/Engine/Source/Runtime/GameplayAbilities/Public/GameplayEffectExtension_LifestealTest.h
@@ -11,7 +11,9 @@
 UCLASS(BlueprintType)
 class GAMEPLAYABILITIES_API UGameplayEffectExtension_LifestealTest : public UGameplayEffectExtension
 {
-	GENERATED_UCLASS_BODY()
+	GENERATED_BODY()
+public:
+	UGameplayEffectExtension_LifestealTest(const FObjectInitializer& ObjectInitializer = FObjectInitializer::Get());
 
 public:
 
@@ -21,4 +23,4 @@ public:
 	/** The GameplayEffect to apply when restoring health to the instigator */
 	UPROPERTY(EditDefaultsOnly, Category = Lifesteal)
 	UGameplayEffect *	HealthRestoreGameplayEffect;
-};
\ No newline at end of file
+};
diff --git a/Engine/Source/Runtime/GameplayAbilities/Public/GameplayEffectExtension_ShieldTest.h b/Engine/Source/Runtime/GameplayAbilities/Public/GameplayEffectExtension_ShieldTest.h
--- a/Engine/Source/Runtime/GameplayAbilities/Public/GameplayEffectExtension_ShieldTest.h
+++ b/Engine/Source/Runtime/GameplayAbilities/Public/GameplayEffectExtension_ShieldTest.h
@@ -11,7 +11,9 @@
 UCLASS(BlueprintType)
 class GAMEPLAYABILITIES_API UGameplayEffectExtension_ShieldTest : public UGameplayEffectExtension
 {
-	GENERATED_UCLASS_BODY()
+	GENERATED_BODY()
+public:
+	UGameplayEffectExtension_ShieldTest(const FObjectInitializer& ObjectInitializer = FObjectInitializer::Get());
 
 public:
 
@@ -20,4 +22,4 @@ public:
 
 	UPROPERTY(EditDefaultsOnly, Category = Lifesteal)
 	UGameplayEffect * ShieldRemoveGameplayEffect;
-};
\ No newline at end of file
+};
diff --git a/Engine/Source/Runtime/GameplayAbilities/Public/GameplayEffectTemplate.h b/Engine/Source/Runtime/GameplayAbilities/Public/GameplayEffectTemplate.h
--- a/Engine/Source/Runtime/GameplayAbilities/Public/GameplayEffectTemplate.h
+++ b/Engine/Source/Runtime/GameplayAbilities/Public/GameplayEffectTemplate.h
@@ -27,11 +27,13 @@
 UCLASS()
 class UGameplayEffectTemplate : public UGameplayEffect
 {
-	GENERATED_UCLASS_BODY()
+	GENERATED_BODY()
+public:
+	UGameplayEffectTemplate(const FObjectInitializer& ObjectInitializer = FObjectInitializer::Get());
 
 #if WITH_EDITORONLY_DATA
 	/** Names of the properties that should be exposed in the limited/basic view when editing a GameplayEffect based on this template */
 	UPROPERTY(EditDefaultsOnly, Category="Template")
 	TArray<FString>	EditableProperties;
 #endif
-};
\ No newline at end of file
+};
diff --git a/Engine/Source/Runtime/GameplayAbilities/Public/GameplayEffectUIData.h b/Engine/Source/Runtime/GameplayAbilities/Public/GameplayEffectUIData.h
--- a/Engine/Source/Runtime/GameplayAbilities/Public/GameplayEffectUIData.h
+++ b/Engine/Source/Runtime/GameplayAbilities/Public/GameplayEffectUIData.h
@@ -11,10 +11,12 @@
 UCLASS(Blueprintable, Abstract, EditInlineNew, CollapseCategories)
 class GAMEPLAYABILITIES_API UGameplayEffectUIData : public UObject
 {
-	GENERATED_UCLASS_BODY()
+	GENERATED_BODY()
+public:
+	UGameplayEffectUIData(const FObjectInitializer& ObjectInitializer = FObjectInitializer::Get());
 
 	virtual bool NeedsLoadForServer() const override
 	{
 		return false;
 	}
-};
\ No newline at end of file
+};
diff --git a/Engine/Source/Runtime/GameplayAbilities/Public/GameplayEffectUIData_TextOnly.h b/Engine/Source/Runtime/GameplayAbilities/Public/GameplayEffectUIData_TextOnly.h
--- a/Engine/Source/Runtime/GameplayAbilities/Public/GameplayEffectUIData_TextOnly.h
+++ b/Engine/Source/Runtime/GameplayAbilities/Public/GameplayEffectUIData_TextOnly.h
@@ -11,8 +11,10 @@
 UCLASS()
 class UGameplayEffectUIData_TextOnly : public UGameplayEffectUIData
 {
-	GENERATED_UCLASS_BODY()
+	GENERATED_BODY()
+public:
+	UGameplayEffectUIData_TextOnly(const FObjectInitializer& ObjectInitializer = FObjectInitializer::Get());
 
 	UPROPERTY(EditDefaultsOnly, BlueprintReadOnly, Category=Data)
 	FText Description;
-};
\ No newline at end of file
+};
diff --git a/Engine/Source/Runtime/GameplayAbilities/Public/GameplayModMagnitudeCalculation.h b/Engine/Source/Runtime/GameplayAbilities/Public/GameplayModMagnitudeCalculation.h
--- a/Engine/Source/Runtime/GameplayAbilities/Public/GameplayModMagnitudeCalculation.h
+++ b/Engine/Source/Runtime/GameplayAbilities/Public/GameplayModMagnitudeCalculation.h
@@ -14,9 +14,11 @@ class UAbilitySystemComponent;
 UCLASS(BlueprintType, Blueprintable, Abstract)
 class GAMEPLAYABILITIES_API UGameplayModMagnitudeCalculation : public UGameplayEffectCalculation
 {
+	GENERATED_BODY()
+public:
+	UGameplayModMagnitudeCalculation(const FObjectInitializer& ObjectInitializer = FObjectInitializer::Get());
 
 public:
-	GENERATED_UCLASS_BODY()
 
 	/**
 	 * Calculate the base magnitude of the gameplay effect modifier, given the specified spec. Note that the owngin spec def can still modify this base value
@@ -28,10 +30,11 @@ public:
 	 */
 	UFUNCTION(BlueprintNativeEvent, Category="Calculation")
 	float CalculateBaseMagnitude(const FGameplayEffectSpec& Spec) const;
+	virtual float CalculateBaseMagnitude_Implementation(const FGameplayEffectSpec& Spec) const;
 
 protected:
 	
 
 	/** Convenience method to get attribute magnitude during a CalculateMagnitude call */
 	bool GetCapturedAttributeMagnitude(const FGameplayEffectAttributeCaptureDefinition& Def, const FGameplayEffectSpec& Spec, const FAggregatorEvaluateParameters& EvaluationParameters, OUT float& Magnitude) const;
-};
\ No newline at end of file
+};
diff --git a/Engine/Source/Runtime/GameplayAbilities/Public/TickableAttributeSetInterface.h b/Engine/Source/Runtime/GameplayAbilities/Public/TickableAttributeSetInterface.h
--- a/Engine/Source/Runtime/GameplayAbilities/Public/TickableAttributeSetInterface.h
+++ b/Engine/Source/Runtime/GameplayAbilities/Public/TickableAttributeSetInterface.h
@@ -8,12 +8,15 @@
 UINTERFACE(MinimalAPI, meta = (CannotImplementInterfaceInBlueprint))
 class UTickableAttributeSetInterface : public UInterface
 {
-	GENERATED_UINTERFACE_BODY()
+	GENERATED_BODY()
+public:
+	GAMEPLAYABILITIES_API UTickableAttributeSetInterface(const FObjectInitializer& ObjectInitializer = FObjectInitializer::Get());
 };
 
 class ITickableAttributeSetInterface
 {
-	GENERATED_IINTERFACE_BODY()
+	GENERATED_BODY()
+public:
 
 public:
 	/**
diff --git a/Engine/Source/Runtime/GameplayTags/Classes/BlueprintGameplayTagLibrary.h b/Engine/Source/Runtime/GameplayTags/Classes/BlueprintGameplayTagLibrary.h
--- a/Engine/Source/Runtime/GameplayTags/Classes/BlueprintGameplayTagLibrary.h
+++ b/Engine/Source/Runtime/GameplayTags/Classes/BlueprintGameplayTagLibrary.h
@@ -13,7 +13,9 @@ struct FGameplayTagContainer;
 UCLASS(MinimalAPI)
 class UBlueprintGameplayTagLibrary : public UBlueprintFunctionLibrary
 {
-	GENERATED_UCLASS_BODY()
+	GENERATED_BODY()
+public:
+	GAMEPLAYTAGS_API UBlueprintGameplayTagLibrary(const FObjectInitializer& ObjectInitializer = FObjectInitializer::Get());
 
 	/**
 	 * Determine if the specified gameplay tags match, given the specified match types
diff --git a/Engine/Source/Runtime/GameplayTags/Classes/GameplayTagAssetInterface.h b/Engine/Source/Runtime/GameplayTags/Classes/GameplayTagAssetInterface.h
--- a/Engine/Source/Runtime/GameplayTags/Classes/GameplayTagAssetInterface.h
+++ b/Engine/Source/Runtime/GameplayTags/Classes/GameplayTagAssetInterface.h
@@ -9,12 +9,15 @@
 UINTERFACE(Blueprintable, MinimalAPI, meta=(CannotImplementInterfaceInBlueprint))
 class UGameplayTagAssetInterface : public UInterface
 {
-	GENERATED_UINTERFACE_BODY()
+	GENERATED_BODY()
+public:
+	GAMEPLAYTAGS_API UGameplayTagAssetInterface(const FObjectInitializer& ObjectInitializer = FObjectInitializer::Get());
 };
 
 class GAMEPLAYTAGS_API IGameplayTagAssetInterface
 {
-	GENERATED_IINTERFACE_BODY()
+	GENERATED_BODY()
+public:
 
 	/**
 	 * Get any owned gameplay tags on the asset
diff --git a/Engine/Source/Runtime/GameplayTags/Classes/GameplayTagsManager.h b/Engine/Source/Runtime/GameplayTags/Classes/GameplayTagsManager.h
--- a/Engine/Source/Runtime/GameplayTags/Classes/GameplayTagsManager.h
+++ b/Engine/Source/Runtime/GameplayTags/Classes/GameplayTagsManager.h
@@ -9,7 +9,7 @@
 USTRUCT()
 struct FGameplayTagTableRow : public FTableRowBase
 {
-	GENERATED_USTRUCT_BODY()
+	GENERATED_BODY()
 
 	/** Tag specified in the table */
 	UPROPERTY(EditAnywhere, BlueprintReadWrite, Category=GameplayTag)
@@ -33,7 +33,7 @@ struct FGameplayTagTableRow : public FTableRowBase
 USTRUCT()
 struct FGameplayTagNode
 {
-	GENERATED_USTRUCT_BODY()
+	GENERATED_BODY()
 	FGameplayTagNode(){};
 
 	/** Simple constructor */
@@ -117,7 +117,9 @@ private:
 UCLASS(config=Engine)
 class GAMEPLAYTAGS_API UGameplayTagsManager : public UObject
 {
-	GENERATED_UCLASS_BODY()
+	GENERATED_BODY()
+public:
+	UGameplayTagsManager(const FObjectInitializer& ObjectInitializer = FObjectInitializer::Get());
 
 	// Destructor
 	~UGameplayTagsManager();
diff --git a/Engine/Source/Runtime/GameplayTags/Classes/GameplayTagsSettings.h b/Engine/Source/Runtime/GameplayTags/Classes/GameplayTagsSettings.h
--- a/Engine/Source/Runtime/GameplayTags/Classes/GameplayTagsSettings.h
+++ b/Engine/Source/Runtime/GameplayTags/Classes/GameplayTagsSettings.h
@@ -22,7 +22,9 @@
 UCLASS(config=Engine, defaultconfig, notplaceable)
 class GAMEPLAYTAGS_API UGameplayTagsSettings : public UObject
 {
-	GENERATED_UCLASS_BODY()
+	GENERATED_BODY()
+public:
+	UGameplayTagsSettings(const FObjectInitializer& ObjectInitializer = FObjectInitializer::Get());
 
 	UPROPERTY(config, EditAnywhere, Category=GameplayTags)
 	TArray<FString>		GameplayTags;
diff --git a/Engine/Source/Runtime/IOS/IOSRuntimeSettings/Classes/IOSRuntimeSettings.h b/Engine/Source/Runtime/IOS/IOSRuntimeSettings/Classes/IOSRuntimeSettings.h
--- a/Engine/Source/Runtime/IOS/IOSRuntimeSettings/Classes/IOSRuntimeSettings.h
+++ b/Engine/Source/Runtime/IOS/IOSRuntimeSettings/Classes/IOSRuntimeSettings.h
@@ -41,7 +41,7 @@ UENUM()
 USTRUCT()
 struct FIOSBuildResourceFilePath
 {
-	GENERATED_USTRUCT_BODY()
+	GENERATED_BODY()
 	
 	/**
 	 * Custom export item used to serialize FIOSBuildResourceFilePath types as only a filename, no garland.
@@ -89,8 +89,10 @@ struct TStructOpsTypeTraits<FIOSBuildResourceFilePath> : public TStructOpsTypeTr
 UCLASS(config=Engine, defaultconfig)
 class IOSRUNTIMESETTINGS_API UIOSRuntimeSettings : public UObject
 {
+	GENERATED_BODY()
+public:
+	UIOSRuntimeSettings(const FObjectInitializer& ObjectInitializer = FObjectInitializer::Get());
 public:
-	GENERATED_UCLASS_BODY()
 
 	// Should Game Center support (iOS Online Subsystem) be enabled?
 	UPROPERTY(GlobalConfig, EditAnywhere, Category = Online)
diff --git a/Engine/Source/Runtime/Landscape/Classes/ControlPointMeshComponent.h b/Engine/Source/Runtime/Landscape/Classes/ControlPointMeshComponent.h
--- a/Engine/Source/Runtime/Landscape/Classes/ControlPointMeshComponent.h
+++ b/Engine/Source/Runtime/Landscape/Classes/ControlPointMeshComponent.h
@@ -8,7 +8,9 @@
 UCLASS(MinimalAPI)
 class UControlPointMeshComponent : public UStaticMeshComponent
 {
-	GENERATED_UCLASS_BODY()
+	GENERATED_BODY()
+public:
+	LANDSCAPE_API UControlPointMeshComponent(const FObjectInitializer& ObjectInitializer = FObjectInitializer::Get());
 
 #if WITH_EDITORONLY_DATA
 	UPROPERTY(transient)
diff --git a/Engine/Source/Runtime/Landscape/Classes/Landscape.h b/Engine/Source/Runtime/Landscape/Classes/Landscape.h
--- a/Engine/Source/Runtime/Landscape/Classes/Landscape.h
+++ b/Engine/Source/Runtime/Landscape/Classes/Landscape.h
@@ -29,7 +29,9 @@ enum ELandscapeSetupErrors
 UCLASS(hidecategories=LandscapeProxy, showcategories=(Display, Movement, Collision, Lighting, LOD, Input), MinimalAPI)
 class ALandscape : public ALandscapeProxy
 {
-	GENERATED_UCLASS_BODY()
+	GENERATED_BODY()
+public:
+	LANDSCAPE_API ALandscape(const FObjectInitializer& ObjectInitializer = FObjectInitializer::Get());
 
 	// Make a key for XYtoComponentMap
 	static FIntPoint MakeKey( int32 X, int32 Y ) { return FIntPoint(X, Y); }
@@ -97,4 +99,4 @@ DECLARE_DWORD_COUNTER_STAT_EXTERN(TEXT("Component Mem"), STAT_LandscapeComponent
  */
 LANDSCAPE_API ULandscapeInfoMap& GetLandscapeInfoMap(UWorld* World);
 
-#endif // WITH_EDITORONLY_DATA
\ No newline at end of file
+#endif // WITH_EDITORONLY_DATA
diff --git a/Engine/Source/Runtime/Landscape/Classes/LandscapeComponent.h b/Engine/Source/Runtime/Landscape/Classes/LandscapeComponent.h
--- a/Engine/Source/Runtime/Landscape/Classes/LandscapeComponent.h
+++ b/Engine/Source/Runtime/Landscape/Classes/LandscapeComponent.h
@@ -74,7 +74,7 @@ struct FLandscapeVertexRef
 USTRUCT()
 struct FWeightmapLayerAllocationInfo
 {
-	GENERATED_USTRUCT_BODY()
+	GENERATED_BODY()
 
 	UPROPERTY()
 	ULandscapeLayerInfoObject* LayerInfo;
@@ -141,7 +141,9 @@ struct FLandscapeComponentGrassData
 UCLASS(hidecategories=(Display, Attachment, Physics, Debug, Collision, Movement, Rendering, PrimitiveComponent, Object, Transform), showcategories=("Rendering|Material"), MinimalAPI)
 class ULandscapeComponent : public UPrimitiveComponent
 {
-	GENERATED_UCLASS_BODY()
+	GENERATED_BODY()
+public:
+	LANDSCAPE_API ULandscapeComponent(const FObjectInitializer& ObjectInitializer = FObjectInitializer::Get());
 	
 	/** X offset from global components grid origin (in quads) */
 	UPROPERTY(VisibleAnywhere, BlueprintReadOnly, Category=LandscapeComponent)
diff --git a/Engine/Source/Runtime/Landscape/Classes/LandscapeGizmoActiveActor.h b/Engine/Source/Runtime/Landscape/Classes/LandscapeGizmoActiveActor.h
--- a/Engine/Source/Runtime/Landscape/Classes/LandscapeGizmoActiveActor.h
+++ b/Engine/Source/Runtime/Landscape/Classes/LandscapeGizmoActiveActor.h
@@ -20,7 +20,7 @@ enum ELandscapeGizmoType
 USTRUCT()
 struct FGizmoSelectData
 {
-	GENERATED_USTRUCT_BODY()
+	GENERATED_BODY()
 
 #if WITH_EDITORONLY_DATA
 	UPROPERTY()
@@ -46,7 +46,9 @@ struct FGizmoSelectData
 UCLASS(notplaceable, MinimalAPI)
 class ALandscapeGizmoActiveActor : public ALandscapeGizmoActor
 {
-	GENERATED_UCLASS_BODY()
+	GENERATED_BODY()
+public:
+	LANDSCAPE_API ALandscapeGizmoActiveActor(const FObjectInitializer& ObjectInitializer = FObjectInitializer::Get());
 
 #if WITH_EDITORONLY_DATA
 	UPROPERTY(Transient)
diff --git a/Engine/Source/Runtime/Landscape/Classes/LandscapeGizmoActor.h b/Engine/Source/Runtime/Landscape/Classes/LandscapeGizmoActor.h
--- a/Engine/Source/Runtime/Landscape/Classes/LandscapeGizmoActor.h
+++ b/Engine/Source/Runtime/Landscape/Classes/LandscapeGizmoActor.h
@@ -9,7 +9,9 @@
 UCLASS(notplaceable, MinimalAPI, NotBlueprintable)
 class ALandscapeGizmoActor : public AActor
 {
-	GENERATED_UCLASS_BODY()
+	GENERATED_BODY()
+public:
+	LANDSCAPE_API ALandscapeGizmoActor(const FObjectInitializer& ObjectInitializer = FObjectInitializer::Get());
 
 #if WITH_EDITORONLY_DATA
 	UPROPERTY(EditAnywhere, Category=Gizmo)
diff --git a/Engine/Source/Runtime/Landscape/Classes/LandscapeGizmoRenderComponent.h b/Engine/Source/Runtime/Landscape/Classes/LandscapeGizmoRenderComponent.h
--- a/Engine/Source/Runtime/Landscape/Classes/LandscapeGizmoRenderComponent.h
+++ b/Engine/Source/Runtime/Landscape/Classes/LandscapeGizmoRenderComponent.h
@@ -7,7 +7,9 @@
 UCLASS(hidecategories=Object)
 class ULandscapeGizmoRenderComponent : public UPrimitiveComponent
 {
-	GENERATED_UCLASS_BODY()
+	GENERATED_BODY()
+public:
+	ULandscapeGizmoRenderComponent(const FObjectInitializer& ObjectInitializer = FObjectInitializer::Get());
 
 
 	// Begin UPrimitiveComponent Interface
diff --git a/Engine/Source/Runtime/Landscape/Classes/LandscapeGrassType.h b/Engine/Source/Runtime/Landscape/Classes/LandscapeGrassType.h
--- a/Engine/Source/Runtime/Landscape/Classes/LandscapeGrassType.h
+++ b/Engine/Source/Runtime/Landscape/Classes/LandscapeGrassType.h
@@ -7,7 +7,7 @@
 USTRUCT()
 struct FGrassVariety
 {
-	GENERATED_USTRUCT_BODY()
+	GENERATED_BODY()
 
 	UPROPERTY(EditAnywhere, Category=Grass)
 	class UStaticMesh* GrassMesh;
@@ -64,7 +64,9 @@ struct FGrassVariety
 UCLASS(MinimalAPI)
 class ULandscapeGrassType : public UObject
 {
-	GENERATED_UCLASS_BODY()
+	GENERATED_BODY()
+public:
+	LANDSCAPE_API ULandscapeGrassType(const FObjectInitializer& ObjectInitializer = FObjectInitializer::Get());
 
 	UPROPERTY(EditAnywhere, Category = Grass)
 	TArray<FGrassVariety> GrassVarieties;
diff --git a/Engine/Source/Runtime/Landscape/Classes/LandscapeHeightfieldCollisionComponent.h b/Engine/Source/Runtime/Landscape/Classes/LandscapeHeightfieldCollisionComponent.h
--- a/Engine/Source/Runtime/Landscape/Classes/LandscapeHeightfieldCollisionComponent.h
+++ b/Engine/Source/Runtime/Landscape/Classes/LandscapeHeightfieldCollisionComponent.h
@@ -20,7 +20,9 @@ namespace physx
 UCLASS(MinimalAPI)
 class ULandscapeHeightfieldCollisionComponent : public UPrimitiveComponent
 {
-	GENERATED_UCLASS_BODY()
+	GENERATED_BODY()
+public:
+	LANDSCAPE_API ULandscapeHeightfieldCollisionComponent(const FObjectInitializer& ObjectInitializer = FObjectInitializer::Get());
 
 	/** List of layers painted on this component. Matches the WeightmapLayerAllocations array in the LandscapeComponent. */
 	UPROPERTY()
diff --git a/Engine/Source/Runtime/Landscape/Classes/LandscapeInfo.h b/Engine/Source/Runtime/Landscape/Classes/LandscapeInfo.h
--- a/Engine/Source/Runtime/Landscape/Classes/LandscapeInfo.h
+++ b/Engine/Source/Runtime/Landscape/Classes/LandscapeInfo.h
@@ -16,7 +16,7 @@ struct FLandscapeDataInterface;
 USTRUCT()
 struct FLandscapeAddCollision
 {
-	GENERATED_USTRUCT_BODY()
+	GENERATED_BODY()
 
 #if WITH_EDITORONLY_DATA
 	UPROPERTY()
@@ -37,7 +37,7 @@ struct FLandscapeAddCollision
 USTRUCT()
 struct FLandscapeInfoLayerSettings
 {
-	GENERATED_USTRUCT_BODY()
+	GENERATED_BODY()
 
 	UPROPERTY()
 	ULandscapeLayerInfoObject* LayerInfoObj;
@@ -95,7 +95,9 @@ struct FLandscapeInfoLayerSettings
 UCLASS()
 class ULandscapeInfo : public UObject
 {
-	GENERATED_UCLASS_BODY()
+	GENERATED_BODY()
+public:
+	ULandscapeInfo(const FObjectInitializer& ObjectInitializer = FObjectInitializer::Get());
 
 	UPROPERTY()
 	TLazyObjectPtr<ALandscape> LandscapeActor;
diff --git a/Engine/Source/Runtime/Landscape/Classes/LandscapeInfoMap.h b/Engine/Source/Runtime/Landscape/Classes/LandscapeInfoMap.h
--- a/Engine/Source/Runtime/Landscape/Classes/LandscapeInfoMap.h
+++ b/Engine/Source/Runtime/Landscape/Classes/LandscapeInfoMap.h
@@ -9,7 +9,9 @@
 UCLASS()
 class ULandscapeInfoMap : public UObject
 {
-	GENERATED_UCLASS_BODY()
+	GENERATED_BODY()
+public:
+	ULandscapeInfoMap(const FObjectInitializer& ObjectInitializer = FObjectInitializer::Get());
 
 public:
 	void PostDuplicate(bool bDuplicateForPIE) override;
diff --git a/Engine/Source/Runtime/Landscape/Classes/LandscapeLayerInfoObject.h b/Engine/Source/Runtime/Landscape/Classes/LandscapeLayerInfoObject.h
--- a/Engine/Source/Runtime/Landscape/Classes/LandscapeLayerInfoObject.h
+++ b/Engine/Source/Runtime/Landscape/Classes/LandscapeLayerInfoObject.h
@@ -7,7 +7,9 @@
 UCLASS(MinimalAPI)
 class ULandscapeLayerInfoObject : public UObject
 {
-	GENERATED_UCLASS_BODY()
+	GENERATED_BODY()
+public:
+	LANDSCAPE_API ULandscapeLayerInfoObject(const FObjectInitializer& ObjectInitializer = FObjectInitializer::Get());
 
 	UPROPERTY(VisibleAnywhere, Category=LandscapeLayerInfoObject, AssetRegistrySearchable)
 	FName LayerName;
diff --git a/Engine/Source/Runtime/Landscape/Classes/LandscapeMaterialInstanceConstant.h b/Engine/Source/Runtime/Landscape/Classes/LandscapeMaterialInstanceConstant.h
--- a/Engine/Source/Runtime/Landscape/Classes/LandscapeMaterialInstanceConstant.h
+++ b/Engine/Source/Runtime/Landscape/Classes/LandscapeMaterialInstanceConstant.h
@@ -6,7 +6,9 @@
 UCLASS(MinimalAPI)
 class ULandscapeMaterialInstanceConstant : public UMaterialInstanceConstant
 {
-	GENERATED_UCLASS_BODY()
+	GENERATED_BODY()
+public:
+	LANDSCAPE_API ULandscapeMaterialInstanceConstant(const FObjectInitializer& ObjectInitializer = FObjectInitializer::Get());
 
 	UPROPERTY()
 	uint32 bIsLayerThumbnail:1;
diff --git a/Engine/Source/Runtime/Landscape/Classes/LandscapeMeshCollisionComponent.h b/Engine/Source/Runtime/Landscape/Classes/LandscapeMeshCollisionComponent.h
--- a/Engine/Source/Runtime/Landscape/Classes/LandscapeMeshCollisionComponent.h
+++ b/Engine/Source/Runtime/Landscape/Classes/LandscapeMeshCollisionComponent.h
@@ -18,7 +18,9 @@ namespace physx
 UCLASS()
 class ULandscapeMeshCollisionComponent : public ULandscapeHeightfieldCollisionComponent
 {
-	GENERATED_UCLASS_BODY()
+	GENERATED_BODY()
+public:
+	ULandscapeMeshCollisionComponent(const FObjectInitializer& ObjectInitializer = FObjectInitializer::Get());
 
 	// Keep the possibility to share projected height field PhysX object with editor mesh collision objects...
 
diff --git a/Engine/Source/Runtime/Landscape/Classes/LandscapeMeshProxyActor.h b/Engine/Source/Runtime/Landscape/Classes/LandscapeMeshProxyActor.h
--- a/Engine/Source/Runtime/Landscape/Classes/LandscapeMeshProxyActor.h
+++ b/Engine/Source/Runtime/Landscape/Classes/LandscapeMeshProxyActor.h
@@ -11,7 +11,9 @@
 UCLASS(MinimalAPI)
 class ALandscapeMeshProxyActor : public AActor
 {
-	GENERATED_UCLASS_BODY()
+	GENERATED_BODY()
+public:
+	LANDSCAPE_API ALandscapeMeshProxyActor(const FObjectInitializer& ObjectInitializer = FObjectInitializer::Get());
 
 private:
 	UPROPERTY(Category = LandscapeMeshProxyActor, VisibleAnywhere, BlueprintReadOnly, meta = (ExposeFunctionCategories = "Mesh,Rendering,Physics,Components|StaticMesh", AllowPrivateAccess = "true"))
diff --git a/Engine/Source/Runtime/Landscape/Classes/LandscapeMeshProxyComponent.h b/Engine/Source/Runtime/Landscape/Classes/LandscapeMeshProxyComponent.h
--- a/Engine/Source/Runtime/Landscape/Classes/LandscapeMeshProxyComponent.h
+++ b/Engine/Source/Runtime/Landscape/Classes/LandscapeMeshProxyComponent.h
@@ -11,7 +11,9 @@
 UCLASS(MinimalAPI)
 class ULandscapeMeshProxyComponent : public UStaticMeshComponent
 {
-	GENERATED_UCLASS_BODY()
+	GENERATED_BODY()
+public:
+	LANDSCAPE_API ULandscapeMeshProxyComponent(const FObjectInitializer& ObjectInitializer = FObjectInitializer::Get());
 
 private:
 	/* The landscape this proxy was generated for */
diff --git a/Engine/Source/Runtime/Landscape/Classes/LandscapeProxy.h b/Engine/Source/Runtime/Landscape/Classes/LandscapeProxy.h
--- a/Engine/Source/Runtime/Landscape/Classes/LandscapeProxy.h
+++ b/Engine/Source/Runtime/Landscape/Classes/LandscapeProxy.h
@@ -35,7 +35,7 @@ class USplineComponent;
 USTRUCT()
 struct FLandscapeWeightmapUsage
 {
-	GENERATED_USTRUCT_BODY()
+	GENERATED_BODY()
 
 	UPROPERTY()
 	ULandscapeComponent* ChannelUsage[4];
@@ -60,7 +60,7 @@ struct FLandscapeWeightmapUsage
 USTRUCT()
 struct FLandscapeEditorLayerSettings
 {
-	GENERATED_USTRUCT_BODY()
+	GENERATED_BODY()
 
 #if WITH_EDITORONLY_DATA
 	UPROPERTY()
@@ -91,7 +91,7 @@ struct FLandscapeEditorLayerSettings
 USTRUCT()
 struct FLandscapeLayerStruct
 {
-	GENERATED_USTRUCT_BODY()
+	GENERATED_BODY()
 
 	UPROPERTY()
 	ULandscapeLayerInfoObject* LayerInfoObj;
@@ -130,7 +130,7 @@ struct FLandscapeLayerStruct
 USTRUCT()
 struct FLandscapeImportLayerInfo
 {
-	GENERATED_USTRUCT_BODY()
+	GENERATED_BODY()
 
 #if WITH_EDITORONLY_DATA
 	UPROPERTY(Category="Import", VisibleAnywhere)
@@ -308,7 +308,9 @@ public:
 UCLASS(NotPlaceable, NotBlueprintable, hidecategories=(Display, Attachment, Physics, Debug, Lighting, LOD), showcategories=(Lighting, Rendering, "Utilities|Transformation"), MinimalAPI)
 class ALandscapeProxy : public AActor, public FTickableGameObject
 {
-	GENERATED_UCLASS_BODY()
+	GENERATED_BODY()
+public:
+	LANDSCAPE_API ALandscapeProxy(const FObjectInitializer& ObjectInitializer = FObjectInitializer::Get());
 
 	virtual ~ALandscapeProxy();
 
diff --git a/Engine/Source/Runtime/Landscape/Classes/LandscapeSplineControlPoint.h b/Engine/Source/Runtime/Landscape/Classes/LandscapeSplineControlPoint.h
--- a/Engine/Source/Runtime/Landscape/Classes/LandscapeSplineControlPoint.h
+++ b/Engine/Source/Runtime/Landscape/Classes/LandscapeSplineControlPoint.h
@@ -11,7 +11,7 @@ struct FLandscapeSplineInterpPoint;
 USTRUCT()
 struct FLandscapeSplineConnection
 {
-	GENERATED_USTRUCT_BODY()
+	GENERATED_BODY()
 
 	FLandscapeSplineConnection() {}
 
@@ -41,7 +41,9 @@ struct FLandscapeSplineConnection
 UCLASS(Within=LandscapeSplinesComponent,autoExpandCategories=LandscapeSplineControlPoint,MinimalAPI)
 class ULandscapeSplineControlPoint : public UObject
 {
-	GENERATED_UCLASS_BODY()
+	GENERATED_BODY()
+public:
+	LANDSCAPE_API ULandscapeSplineControlPoint(const FObjectInitializer& ObjectInitializer = FObjectInitializer::Get());
 
 // Directly editable data:
 	/** Location in Landscape-space */
diff --git a/Engine/Source/Runtime/Landscape/Classes/LandscapeSplineSegment.h b/Engine/Source/Runtime/Landscape/Classes/LandscapeSplineSegment.h
--- a/Engine/Source/Runtime/Landscape/Classes/LandscapeSplineSegment.h
+++ b/Engine/Source/Runtime/Landscape/Classes/LandscapeSplineSegment.h
@@ -7,7 +7,7 @@
 USTRUCT()
 struct FLandscapeSplineInterpPoint
 {
-	GENERATED_USTRUCT_BODY()
+	GENERATED_BODY()
 
 	/** Center Point */
 	UPROPERTY()
@@ -51,7 +51,7 @@ struct FLandscapeSplineInterpPoint
 USTRUCT()
 struct FLandscapeSplineSegmentConnection
 {
-	GENERATED_USTRUCT_BODY()
+	GENERATED_BODY()
 
 	// Control point connected to this end of the segment
 	UPROPERTY()
@@ -85,7 +85,7 @@ enum LandscapeSplineMeshOrientation
 USTRUCT()
 struct FLandscapeSplineMeshEntry
 {
-	GENERATED_USTRUCT_BODY()
+	GENERATED_BODY()
 
 	/** Mesh to use on the spline */
 	UPROPERTY(EditAnywhere, Category=LandscapeSplineMeshEntry)
@@ -143,7 +143,9 @@ struct FLandscapeSplineMeshEntry
 UCLASS(Within=LandscapeSplinesComponent,autoExpandCategories=(LandscapeSplineSegment,LandscapeSplineMeshes),MinimalAPI)
 class ULandscapeSplineSegment : public UObject
 {
-	GENERATED_UCLASS_BODY()
+	GENERATED_BODY()
+public:
+	LANDSCAPE_API ULandscapeSplineSegment(const FObjectInitializer& ObjectInitializer = FObjectInitializer::Get());
 
 // Directly editable data:
 	UPROPERTY(EditAnywhere, EditFixedSize, Category=LandscapeSplineSegment)
diff --git a/Engine/Source/Runtime/Landscape/Classes/LandscapeSplinesComponent.h b/Engine/Source/Runtime/Landscape/Classes/LandscapeSplinesComponent.h
--- a/Engine/Source/Runtime/Landscape/Classes/LandscapeSplinesComponent.h
+++ b/Engine/Source/Runtime/Landscape/Classes/LandscapeSplinesComponent.h
@@ -7,7 +7,9 @@
 UCLASS(MinimalAPI)
 class ULandscapeSplinesComponent : public UPrimitiveComponent
 {
-	GENERATED_UCLASS_BODY()
+	GENERATED_BODY()
+public:
+	LANDSCAPE_API ULandscapeSplinesComponent(const FObjectInitializer& ObjectInitializer = FObjectInitializer::Get());
 
 #if WITH_EDITORONLY_DATA
 	/** Resolution of the spline, in distance per point */
diff --git a/Engine/Source/Runtime/Landscape/Classes/Materials/MaterialExpressionLandscapeGrassOutput.h b/Engine/Source/Runtime/Landscape/Classes/Materials/MaterialExpressionLandscapeGrassOutput.h
--- a/Engine/Source/Runtime/Landscape/Classes/Materials/MaterialExpressionLandscapeGrassOutput.h
+++ b/Engine/Source/Runtime/Landscape/Classes/Materials/MaterialExpressionLandscapeGrassOutput.h
@@ -10,7 +10,7 @@
 USTRUCT()
 struct FGrassInput
 {
-	GENERATED_USTRUCT_BODY()
+	GENERATED_BODY()
 
 	UPROPERTY(EditAnywhere, Category = Grass)
 	FName Name;
@@ -34,7 +34,9 @@ struct FGrassInput
 UCLASS(collapsecategories, hidecategories=Object, MinimalAPI)
 class UMaterialExpressionLandscapeGrassOutput : public UMaterialExpressionCustomOutput
 {
-	GENERATED_UCLASS_BODY()
+	GENERATED_BODY()
+public:
+	LANDSCAPE_API UMaterialExpressionLandscapeGrassOutput(const FObjectInitializer& ObjectInitializer = FObjectInitializer::Get());
 	virtual int32 Compile(class FMaterialCompiler* Compiler, int32 OutputIndex, int32 MultiplexIndex) override;
 	virtual void GetCaption(TArray<FString>& OutCaptions) const override;
 	virtual const TArray<FExpressionInput*> GetInputs() override;
diff --git a/Engine/Source/Runtime/Landscape/Classes/Materials/MaterialExpressionLandscapeLayerBlend.h b/Engine/Source/Runtime/Landscape/Classes/Materials/MaterialExpressionLandscapeLayerBlend.h
--- a/Engine/Source/Runtime/Landscape/Classes/Materials/MaterialExpressionLandscapeLayerBlend.h
+++ b/Engine/Source/Runtime/Landscape/Classes/Materials/MaterialExpressionLandscapeLayerBlend.h
@@ -17,7 +17,7 @@ enum ELandscapeLayerBlendType
 USTRUCT()
 struct FLayerBlendInput
 {
-	GENERATED_USTRUCT_BODY()
+	GENERATED_BODY()
 
 	UPROPERTY(EditAnywhere, Category=LayerBlendInput)
 	FName LayerName;
@@ -54,7 +54,9 @@ struct FLayerBlendInput
 UCLASS(collapsecategories, hidecategories=Object)
 class UMaterialExpressionLandscapeLayerBlend : public UMaterialExpression
 {
-	GENERATED_UCLASS_BODY()
+	GENERATED_BODY()
+public:
+	UMaterialExpressionLandscapeLayerBlend(const FObjectInitializer& ObjectInitializer = FObjectInitializer::Get());
 
 	UPROPERTY(EditAnywhere, Category=MaterialExpressionLandscapeLayerBlend)
 	TArray<FLayerBlendInput> Layers;
diff --git a/Engine/Source/Runtime/Landscape/Classes/Materials/MaterialExpressionLandscapeLayerCoords.h b/Engine/Source/Runtime/Landscape/Classes/Materials/MaterialExpressionLandscapeLayerCoords.h
--- a/Engine/Source/Runtime/Landscape/Classes/Materials/MaterialExpressionLandscapeLayerCoords.h
+++ b/Engine/Source/Runtime/Landscape/Classes/Materials/MaterialExpressionLandscapeLayerCoords.h
@@ -29,7 +29,9 @@ enum ELandscapeCustomizedCoordType
 UCLASS(collapsecategories, hidecategories=Object)
 class UMaterialExpressionLandscapeLayerCoords : public UMaterialExpression
 {
-	GENERATED_UCLASS_BODY()
+	GENERATED_BODY()
+public:
+	UMaterialExpressionLandscapeLayerCoords(const FObjectInitializer& ObjectInitializer = FObjectInitializer::Get());
 
 	/** Determines the mapping place to use on the terrain. */
 	UPROPERTY(EditAnywhere, Category=MaterialExpressionLandscapeLayerCoords)
diff --git a/Engine/Source/Runtime/Landscape/Classes/Materials/MaterialExpressionLandscapeLayerSample.h b/Engine/Source/Runtime/Landscape/Classes/Materials/MaterialExpressionLandscapeLayerSample.h
--- a/Engine/Source/Runtime/Landscape/Classes/Materials/MaterialExpressionLandscapeLayerSample.h
+++ b/Engine/Source/Runtime/Landscape/Classes/Materials/MaterialExpressionLandscapeLayerSample.h
@@ -9,7 +9,9 @@
 UCLASS(collapsecategories, hidecategories=Object)
 class UMaterialExpressionLandscapeLayerSample : public UMaterialExpression
 {
-	GENERATED_UCLASS_BODY()
+	GENERATED_BODY()
+public:
+	UMaterialExpressionLandscapeLayerSample(const FObjectInitializer& ObjectInitializer = FObjectInitializer::Get());
 
 	UPROPERTY(EditAnywhere, Category=MaterialExpressionLandscapeLayerWeight)
 	FName ParameterName;
diff --git a/Engine/Source/Runtime/Landscape/Classes/Materials/MaterialExpressionLandscapeLayerSwitch.h b/Engine/Source/Runtime/Landscape/Classes/Materials/MaterialExpressionLandscapeLayerSwitch.h
--- a/Engine/Source/Runtime/Landscape/Classes/Materials/MaterialExpressionLandscapeLayerSwitch.h
+++ b/Engine/Source/Runtime/Landscape/Classes/Materials/MaterialExpressionLandscapeLayerSwitch.h
@@ -8,7 +8,9 @@
 UCLASS(collapsecategories, hidecategories=Object)
 class UMaterialExpressionLandscapeLayerSwitch : public UMaterialExpression
 {
-	GENERATED_UCLASS_BODY()
+	GENERATED_BODY()
+public:
+	UMaterialExpressionLandscapeLayerSwitch(const FObjectInitializer& ObjectInitializer = FObjectInitializer::Get());
 
 	UPROPERTY()
 	FExpressionInput LayerUsed;
diff --git a/Engine/Source/Runtime/Landscape/Classes/Materials/MaterialExpressionLandscapeLayerWeight.h b/Engine/Source/Runtime/Landscape/Classes/Materials/MaterialExpressionLandscapeLayerWeight.h
--- a/Engine/Source/Runtime/Landscape/Classes/Materials/MaterialExpressionLandscapeLayerWeight.h
+++ b/Engine/Source/Runtime/Landscape/Classes/Materials/MaterialExpressionLandscapeLayerWeight.h
@@ -8,7 +8,9 @@
 UCLASS(collapsecategories, hidecategories=Object)
 class UMaterialExpressionLandscapeLayerWeight : public UMaterialExpression
 {
-	GENERATED_UCLASS_BODY()
+	GENERATED_BODY()
+public:
+	UMaterialExpressionLandscapeLayerWeight(const FObjectInitializer& ObjectInitializer = FObjectInitializer::Get());
 
 	UPROPERTY(meta = (RequiredInput = "false", ToolTip = "Defaults to 'ConstBase' if not specified"))
 	FExpressionInput Base;
diff --git a/Engine/Source/Runtime/Landscape/Classes/Materials/MaterialExpressionLandscapeVisibilityMask.h b/Engine/Source/Runtime/Landscape/Classes/Materials/MaterialExpressionLandscapeVisibilityMask.h
--- a/Engine/Source/Runtime/Landscape/Classes/Materials/MaterialExpressionLandscapeVisibilityMask.h
+++ b/Engine/Source/Runtime/Landscape/Classes/Materials/MaterialExpressionLandscapeVisibilityMask.h
@@ -8,7 +8,9 @@
 UCLASS(collapsecategories, hidecategories=Object)
 class UMaterialExpressionLandscapeVisibilityMask : public UMaterialExpression
 {
-	GENERATED_UCLASS_BODY()
+	GENERATED_BODY()
+public:
+	UMaterialExpressionLandscapeVisibilityMask(const FObjectInitializer& ObjectInitializer = FObjectInitializer::Get());
 
 	/** GUID that should be unique within the material, this is used for parameter renaming. */
 	UPROPERTY()
diff --git a/Engine/Source/Runtime/MediaAssets/Public/MediaPlayer.h b/Engine/Source/Runtime/MediaAssets/Public/MediaPlayer.h
--- a/Engine/Source/Runtime/MediaAssets/Public/MediaPlayer.h
+++ b/Engine/Source/Runtime/MediaAssets/Public/MediaPlayer.h
@@ -41,7 +41,9 @@ UCLASS(BlueprintType, hidecategories=(Object))
 class MEDIAASSETS_API UMediaPlayer
 	: public UObject
 {
-	GENERATED_UCLASS_BODY()
+	GENERATED_BODY()
+public:
+	UMediaPlayer(const FObjectInitializer& ObjectInitializer = FObjectInitializer::Get());
 
 public:
 
diff --git a/Engine/Source/Runtime/MediaAssets/Public/MediaSoundWave.h b/Engine/Source/Runtime/MediaAssets/Public/MediaSoundWave.h
--- a/Engine/Source/Runtime/MediaAssets/Public/MediaSoundWave.h
+++ b/Engine/Source/Runtime/MediaAssets/Public/MediaSoundWave.h
@@ -20,7 +20,9 @@ UCLASS(hidecategories=(Compression, Sound, SoundWave, Subtitles))
 class MEDIAASSETS_API UMediaSoundWave
 	: public USoundWave
 {
-	GENERATED_UCLASS_BODY()
+	GENERATED_BODY()
+public:
+	UMediaSoundWave(const FObjectInitializer& ObjectInitializer = FObjectInitializer::Get());
 
 	/** The index of the MediaPlayer's audio track to get the wave data from. */
 	UPROPERTY(EditAnywhere, BlueprintReadWrite, Category=MediaPlayer)
diff --git a/Engine/Source/Runtime/MediaAssets/Public/MediaTexture.h b/Engine/Source/Runtime/MediaAssets/Public/MediaTexture.h
--- a/Engine/Source/Runtime/MediaAssets/Public/MediaTexture.h
+++ b/Engine/Source/Runtime/MediaAssets/Public/MediaTexture.h
@@ -21,7 +21,9 @@ UCLASS(hidecategories=(Compression, LevelOfDetail, Object, Texture))
 class MEDIAASSETS_API UMediaTexture
 	: public UTexture
 {
-	GENERATED_UCLASS_BODY()
+	GENERATED_BODY()
+public:
+	UMediaTexture(const FObjectInitializer& ObjectInitializer = FObjectInitializer::Get());
 
 	/** The addressing mode to use for the X axis. */
 	UPROPERTY(EditAnywhere, BlueprintReadWrite, Category=MediaTexture, AssetRegistrySearchable)
diff --git a/Engine/Source/Runtime/MoviePlayer/Public/MoviePlayerSettings.h b/Engine/Source/Runtime/MoviePlayer/Public/MoviePlayerSettings.h
--- a/Engine/Source/Runtime/MoviePlayer/Public/MoviePlayerSettings.h
+++ b/Engine/Source/Runtime/MoviePlayer/Public/MoviePlayerSettings.h
@@ -12,9 +12,11 @@ UCLASS(config=Game, defaultconfig)
 class MOVIEPLAYER_API UMoviePlayerSettings
 	: public UObject
 {
+	GENERATED_BODY()
+public:
+	UMoviePlayerSettings(const FObjectInitializer& ObjectInitializer = FObjectInitializer::Get());
 public:
 
-	GENERATED_UCLASS_BODY()
 
 	/** If enabled the game waits for startup movies to complete even if loading has finished */
 	UPROPERTY(globalconfig, EditAnywhere, Category="Movies")
diff --git a/Engine/Source/Runtime/MovieSceneCore/Classes/MovieScene.h b/Engine/Source/Runtime/MovieSceneCore/Classes/MovieScene.h
--- a/Engine/Source/Runtime/MovieSceneCore/Classes/MovieScene.h
+++ b/Engine/Source/Runtime/MovieSceneCore/Classes/MovieScene.h
@@ -12,7 +12,7 @@ MOVIESCENECORE_API DECLARE_LOG_CATEGORY_EXTERN(LogSequencerRuntime, Log, All);
 USTRUCT()
 struct FMovieSceneSpawnable 
 {
-	GENERATED_USTRUCT_BODY()
+	GENERATED_BODY()
 
 public:
 
@@ -75,7 +75,7 @@ private:
 USTRUCT()
 struct FMovieScenePossessable
 {
-	GENERATED_USTRUCT_BODY( FMovieScenePossessable )
+	GENERATED_BODY()
 
 public:
 
@@ -134,7 +134,7 @@ private:
 USTRUCT()
 struct FMovieSceneEditorData
 {
-	GENERATED_USTRUCT_BODY()
+	GENERATED_BODY()
 
 	/** List of collapsed sequencer nodes.  We store collapsed instead of expanded so that new nodes with no saved state are expanded by default */
 	UPROPERTY()
@@ -147,7 +147,7 @@ struct FMovieSceneEditorData
 USTRUCT()
 struct FMovieSceneObjectBinding
 {
-	GENERATED_USTRUCT_BODY()
+	GENERATED_BODY()
 
 	FMovieSceneObjectBinding()
 	{}
@@ -218,7 +218,9 @@ private:
 UCLASS( MinimalAPI )
 class UMovieScene : public UObject
 {
-	GENERATED_UCLASS_BODY()
+	GENERATED_BODY()
+public:
+	MOVIESCENECORE_API UMovieScene(const FObjectInitializer& ObjectInitializer = FObjectInitializer::Get());
 public:
 
 #if WITH_EDITOR
diff --git a/Engine/Source/Runtime/MovieSceneCore/Classes/MovieSceneBindings.h b/Engine/Source/Runtime/MovieSceneCore/Classes/MovieSceneBindings.h
--- a/Engine/Source/Runtime/MovieSceneCore/Classes/MovieSceneBindings.h
+++ b/Engine/Source/Runtime/MovieSceneCore/Classes/MovieSceneBindings.h
@@ -11,7 +11,7 @@
 USTRUCT()
 struct FMovieSceneBoundObject
 {
-	GENERATED_USTRUCT_BODY( FMovieSceneBoundObject )
+	GENERATED_BODY()
 
 public:
 
@@ -87,7 +87,9 @@ private:
 UCLASS( MinimalAPI )
 class UMovieSceneBindings : public UObject
 {
-	GENERATED_UCLASS_BODY()
+	GENERATED_BODY()
+public:
+	MOVIESCENECORE_API UMovieSceneBindings(const FObjectInitializer& ObjectInitializer = FObjectInitializer::Get());
 
 public:
 
diff --git a/Engine/Source/Runtime/MovieSceneCore/Classes/MovieScenePropertyTrack.h b/Engine/Source/Runtime/MovieSceneCore/Classes/MovieScenePropertyTrack.h
--- a/Engine/Source/Runtime/MovieSceneCore/Classes/MovieScenePropertyTrack.h
+++ b/Engine/Source/Runtime/MovieSceneCore/Classes/MovieScenePropertyTrack.h
@@ -11,7 +11,9 @@
 UCLASS( abstract )
 class MOVIESCENECORE_API UMovieScenePropertyTrack : public UMovieSceneTrack
 {
-	GENERATED_UCLASS_BODY()
+	GENERATED_BODY()
+public:
+	UMovieScenePropertyTrack(const FObjectInitializer& ObjectInitializer = FObjectInitializer::Get());
 
 public:
 	/** UMovieSceneTrack interface */
diff --git a/Engine/Source/Runtime/MovieSceneCore/Classes/MovieSceneSection.h b/Engine/Source/Runtime/MovieSceneCore/Classes/MovieSceneSection.h
--- a/Engine/Source/Runtime/MovieSceneCore/Classes/MovieSceneSection.h
+++ b/Engine/Source/Runtime/MovieSceneCore/Classes/MovieSceneSection.h
@@ -10,7 +10,9 @@
 UCLASS( abstract, MinimalAPI )
 class UMovieSceneSection : public UObject
 {
-	GENERATED_UCLASS_BODY()
+	GENERATED_BODY()
+public:
+	MOVIESCENECORE_API UMovieSceneSection(const FObjectInitializer& ObjectInitializer = FObjectInitializer::Get());
 public:
 
 	/**
diff --git a/Engine/Source/Runtime/MovieSceneCore/Classes/MovieSceneTrack.h b/Engine/Source/Runtime/MovieSceneCore/Classes/MovieSceneTrack.h
--- a/Engine/Source/Runtime/MovieSceneCore/Classes/MovieSceneTrack.h
+++ b/Engine/Source/Runtime/MovieSceneCore/Classes/MovieSceneTrack.h
@@ -12,7 +12,9 @@ class UMovieSceneSection;
 UCLASS( abstract, MinimalAPI )
 class UMovieSceneTrack : public UObject
 {
-	GENERATED_UCLASS_BODY()
+	GENERATED_BODY()
+public:
+	MOVIESCENECORE_API UMovieSceneTrack(const FObjectInitializer& ObjectInitializer = FObjectInitializer::Get());
 
 public:
 	/**
diff --git a/Engine/Source/Runtime/MovieSceneCore/Classes/RuntimeMovieScenePlayer.h b/Engine/Source/Runtime/MovieSceneCore/Classes/RuntimeMovieScenePlayer.h
--- a/Engine/Source/Runtime/MovieSceneCore/Classes/RuntimeMovieScenePlayer.h
+++ b/Engine/Source/Runtime/MovieSceneCore/Classes/RuntimeMovieScenePlayer.h
@@ -17,7 +17,9 @@
 UCLASS( MinimalAPI )
 class URuntimeMovieScenePlayer : public UObject, public IMovieScenePlayer, public IRuntimeMovieScenePlayerInterface
 {
-	GENERATED_UCLASS_BODY()
+	GENERATED_BODY()
+public:
+	MOVIESCENECORE_API URuntimeMovieScenePlayer(const FObjectInitializer& ObjectInitializer = FObjectInitializer::Get());
 
 public:
 
diff --git a/Engine/Source/Runtime/MovieSceneCoreTypes/Classes/MovieScene3DTransformSection.h b/Engine/Source/Runtime/MovieSceneCoreTypes/Classes/MovieScene3DTransformSection.h
--- a/Engine/Source/Runtime/MovieSceneCoreTypes/Classes/MovieScene3DTransformSection.h
+++ b/Engine/Source/Runtime/MovieSceneCoreTypes/Classes/MovieScene3DTransformSection.h
@@ -12,7 +12,9 @@
 UCLASS(MinimalAPI)
 class UMovieScene3DTransformSection : public UMovieSceneSection
 {
-	GENERATED_UCLASS_BODY()
+	GENERATED_BODY()
+public:
+	MOVIESCENECORETYPES_API UMovieScene3DTransformSection(const FObjectInitializer& ObjectInitializer = FObjectInitializer::Get());
 public:
 
 	/** MovieSceneSection interface */
@@ -111,4 +113,4 @@ private:
 	/** Scale curves */
 	UPROPERTY()
 	FRichCurve Scale[3];
-};
\ No newline at end of file
+};
diff --git a/Engine/Source/Runtime/MovieSceneCoreTypes/Classes/MovieScene3DTransformTrack.h b/Engine/Source/Runtime/MovieSceneCoreTypes/Classes/MovieScene3DTransformTrack.h
--- a/Engine/Source/Runtime/MovieSceneCoreTypes/Classes/MovieScene3DTransformTrack.h
+++ b/Engine/Source/Runtime/MovieSceneCoreTypes/Classes/MovieScene3DTransformTrack.h
@@ -129,7 +129,9 @@ private:
 UCLASS( MinimalAPI )
 class UMovieScene3DTransformTrack : public UMovieScenePropertyTrack
 {
-	GENERATED_UCLASS_BODY()
+	GENERATED_BODY()
+public:
+	MOVIESCENECORETYPES_API UMovieScene3DTransformTrack(const FObjectInitializer& ObjectInitializer = FObjectInitializer::Get());
 public:
 	/** UMovieSceneTrack interface */
 	virtual UMovieSceneSection* CreateNewSection() override;
diff --git a/Engine/Source/Runtime/MovieSceneCoreTypes/Classes/MovieSceneAnimationSection.h b/Engine/Source/Runtime/MovieSceneCoreTypes/Classes/MovieSceneAnimationSection.h
--- a/Engine/Source/Runtime/MovieSceneCoreTypes/Classes/MovieSceneAnimationSection.h
+++ b/Engine/Source/Runtime/MovieSceneCoreTypes/Classes/MovieSceneAnimationSection.h
@@ -11,7 +11,9 @@
 UCLASS( MinimalAPI )
 class UMovieSceneAnimationSection : public UMovieSceneSection
 {
-	GENERATED_UCLASS_BODY()
+	GENERATED_BODY()
+public:
+	MOVIESCENECORETYPES_API UMovieSceneAnimationSection(const FObjectInitializer& ObjectInitializer = FObjectInitializer::Get());
 public:
 	/** Sets the animation sequence for this section */
 	void SetAnimSequence(class UAnimSequence* InAnimSequence) {AnimSequence = InAnimSequence;}
diff --git a/Engine/Source/Runtime/MovieSceneCoreTypes/Classes/MovieSceneAnimationTrack.h b/Engine/Source/Runtime/MovieSceneCoreTypes/Classes/MovieSceneAnimationTrack.h
--- a/Engine/Source/Runtime/MovieSceneCoreTypes/Classes/MovieSceneAnimationTrack.h
+++ b/Engine/Source/Runtime/MovieSceneCoreTypes/Classes/MovieSceneAnimationTrack.h
@@ -13,7 +13,9 @@
 UCLASS( MinimalAPI )
 class UMovieSceneAnimationTrack : public UMovieSceneTrack
 {
-	GENERATED_UCLASS_BODY()
+	GENERATED_BODY()
+public:
+	MOVIESCENECORETYPES_API UMovieSceneAnimationTrack(const FObjectInitializer& ObjectInitializer = FObjectInitializer::Get());
 
 public:
 	/** UMovieSceneTrack interface */
diff --git a/Engine/Source/Runtime/MovieSceneCoreTypes/Classes/MovieSceneAudioSection.h b/Engine/Source/Runtime/MovieSceneCoreTypes/Classes/MovieSceneAudioSection.h
--- a/Engine/Source/Runtime/MovieSceneCoreTypes/Classes/MovieSceneAudioSection.h
+++ b/Engine/Source/Runtime/MovieSceneCoreTypes/Classes/MovieSceneAudioSection.h
@@ -11,7 +11,9 @@
 UCLASS( MinimalAPI )
 class UMovieSceneAudioSection : public UMovieSceneSection
 {
-	GENERATED_UCLASS_BODY()
+	GENERATED_BODY()
+public:
+	MOVIESCENECORETYPES_API UMovieSceneAudioSection(const FObjectInitializer& ObjectInitializer = FObjectInitializer::Get());
 public:
 	/** Sets this section's sound */
 	void SetSound(class USoundBase* InSound) {Sound = InSound;}
diff --git a/Engine/Source/Runtime/MovieSceneCoreTypes/Classes/MovieSceneAudioTrack.h b/Engine/Source/Runtime/MovieSceneCoreTypes/Classes/MovieSceneAudioTrack.h
--- a/Engine/Source/Runtime/MovieSceneCoreTypes/Classes/MovieSceneAudioTrack.h
+++ b/Engine/Source/Runtime/MovieSceneCoreTypes/Classes/MovieSceneAudioTrack.h
@@ -18,7 +18,9 @@ namespace AudioTrackConstants
 UCLASS( MinimalAPI )
 class UMovieSceneAudioTrack : public UMovieSceneTrack
 {
-	GENERATED_UCLASS_BODY()
+	GENERATED_BODY()
+public:
+	MOVIESCENECORETYPES_API UMovieSceneAudioTrack(const FObjectInitializer& ObjectInitializer = FObjectInitializer::Get());
 
 public:
 	/** UMovieSceneTrack interface */
diff --git a/Engine/Source/Runtime/MovieSceneCoreTypes/Classes/MovieSceneBoolSection.h b/Engine/Source/Runtime/MovieSceneCoreTypes/Classes/MovieSceneBoolSection.h
--- a/Engine/Source/Runtime/MovieSceneCoreTypes/Classes/MovieSceneBoolSection.h
+++ b/Engine/Source/Runtime/MovieSceneCoreTypes/Classes/MovieSceneBoolSection.h
@@ -10,7 +10,9 @@
 UCLASS( MinimalAPI )
 class UMovieSceneBoolSection : public UMovieSceneSection
 {
-	GENERATED_UCLASS_BODY()
+	GENERATED_BODY()
+public:
+	MOVIESCENECORETYPES_API UMovieSceneBoolSection(const FObjectInitializer& ObjectInitializer = FObjectInitializer::Get());
 public:
 	/**
 	 * Updates this section
diff --git a/Engine/Source/Runtime/MovieSceneCoreTypes/Classes/MovieSceneBoolTrack.h b/Engine/Source/Runtime/MovieSceneCoreTypes/Classes/MovieSceneBoolTrack.h
--- a/Engine/Source/Runtime/MovieSceneCoreTypes/Classes/MovieSceneBoolTrack.h
+++ b/Engine/Source/Runtime/MovieSceneCoreTypes/Classes/MovieSceneBoolTrack.h
@@ -13,7 +13,9 @@
 UCLASS( MinimalAPI )
 class UMovieSceneBoolTrack : public UMovieScenePropertyTrack
 {
-	GENERATED_UCLASS_BODY()
+	GENERATED_BODY()
+public:
+	MOVIESCENECORETYPES_API UMovieSceneBoolTrack(const FObjectInitializer& ObjectInitializer = FObjectInitializer::Get());
 
 public:
 	/** UMovieSceneTrack interface */
diff --git a/Engine/Source/Runtime/MovieSceneCoreTypes/Classes/MovieSceneColorSection.h b/Engine/Source/Runtime/MovieSceneCoreTypes/Classes/MovieSceneColorSection.h
--- a/Engine/Source/Runtime/MovieSceneCoreTypes/Classes/MovieSceneColorSection.h
+++ b/Engine/Source/Runtime/MovieSceneCoreTypes/Classes/MovieSceneColorSection.h
@@ -12,7 +12,9 @@ struct FColorKey;
 UCLASS( MinimalAPI )
 class UMovieSceneColorSection : public UMovieSceneSection
 {
-	GENERATED_UCLASS_BODY()
+	GENERATED_BODY()
+public:
+	MOVIESCENECORETYPES_API UMovieSceneColorSection(const FObjectInitializer& ObjectInitializer = FObjectInitializer::Get());
 public:
 	/** MovieSceneSection interface */
 	virtual void MoveSection( float DeltaPosition ) override;
diff --git a/Engine/Source/Runtime/MovieSceneCoreTypes/Classes/MovieSceneColorTrack.h b/Engine/Source/Runtime/MovieSceneCoreTypes/Classes/MovieSceneColorTrack.h
--- a/Engine/Source/Runtime/MovieSceneCoreTypes/Classes/MovieSceneColorTrack.h
+++ b/Engine/Source/Runtime/MovieSceneCoreTypes/Classes/MovieSceneColorTrack.h
@@ -21,7 +21,9 @@ struct FColorKey
 UCLASS( MinimalAPI )
 class UMovieSceneColorTrack : public UMovieScenePropertyTrack
 {
-	GENERATED_UCLASS_BODY()
+	GENERATED_BODY()
+public:
+	MOVIESCENECORETYPES_API UMovieSceneColorTrack(const FObjectInitializer& ObjectInitializer = FObjectInitializer::Get());
 
 public:
 	/** UMovieSceneTrack interface */
diff --git a/Engine/Source/Runtime/MovieSceneCoreTypes/Classes/MovieSceneDirectorTrack.h b/Engine/Source/Runtime/MovieSceneCoreTypes/Classes/MovieSceneDirectorTrack.h
--- a/Engine/Source/Runtime/MovieSceneCoreTypes/Classes/MovieSceneDirectorTrack.h
+++ b/Engine/Source/Runtime/MovieSceneCoreTypes/Classes/MovieSceneDirectorTrack.h
@@ -12,7 +12,9 @@
 UCLASS( MinimalAPI )
 class UMovieSceneDirectorTrack : public UMovieSceneTrack
 {
-	GENERATED_UCLASS_BODY()
+	GENERATED_BODY()
+public:
+	MOVIESCENECORETYPES_API UMovieSceneDirectorTrack(const FObjectInitializer& ObjectInitializer = FObjectInitializer::Get());
 
 public:
 	/** UMovieSceneTrack interface */
diff --git a/Engine/Source/Runtime/MovieSceneCoreTypes/Classes/MovieSceneFloatSection.h b/Engine/Source/Runtime/MovieSceneCoreTypes/Classes/MovieSceneFloatSection.h
--- a/Engine/Source/Runtime/MovieSceneCoreTypes/Classes/MovieSceneFloatSection.h
+++ b/Engine/Source/Runtime/MovieSceneCoreTypes/Classes/MovieSceneFloatSection.h
@@ -10,7 +10,9 @@
 UCLASS( MinimalAPI )
 class UMovieSceneFloatSection : public UMovieSceneSection
 {
-	GENERATED_UCLASS_BODY()
+	GENERATED_BODY()
+public:
+	MOVIESCENECORETYPES_API UMovieSceneFloatSection(const FObjectInitializer& ObjectInitializer = FObjectInitializer::Get());
 public:
 	/**
 	 * Updates this section
@@ -50,4 +52,4 @@ private:
 	/** Curve data */
 	UPROPERTY()
 	FRichCurve FloatCurve;
-};
\ No newline at end of file
+};
diff --git a/Engine/Source/Runtime/MovieSceneCoreTypes/Classes/MovieSceneFloatTrack.h b/Engine/Source/Runtime/MovieSceneCoreTypes/Classes/MovieSceneFloatTrack.h
--- a/Engine/Source/Runtime/MovieSceneCoreTypes/Classes/MovieSceneFloatTrack.h
+++ b/Engine/Source/Runtime/MovieSceneCoreTypes/Classes/MovieSceneFloatTrack.h
@@ -13,7 +13,9 @@
 UCLASS( MinimalAPI )
 class UMovieSceneFloatTrack : public UMovieScenePropertyTrack
 {
-	GENERATED_UCLASS_BODY()
+	GENERATED_BODY()
+public:
+	MOVIESCENECORETYPES_API UMovieSceneFloatTrack(const FObjectInitializer& ObjectInitializer = FObjectInitializer::Get());
 
 public:
 	/** UMovieSceneTrack interface */
diff --git a/Engine/Source/Runtime/MovieSceneCoreTypes/Classes/MovieSceneParticleSection.h b/Engine/Source/Runtime/MovieSceneCoreTypes/Classes/MovieSceneParticleSection.h
--- a/Engine/Source/Runtime/MovieSceneCoreTypes/Classes/MovieSceneParticleSection.h
+++ b/Engine/Source/Runtime/MovieSceneCoreTypes/Classes/MovieSceneParticleSection.h
@@ -24,7 +24,9 @@ namespace EParticleKey
 UCLASS( MinimalAPI )
 class UMovieSceneParticleSection : public UMovieSceneSection
 {
-	GENERATED_UCLASS_BODY()
+	GENERATED_BODY()
+public:
+	MOVIESCENECORETYPES_API UMovieSceneParticleSection(const FObjectInitializer& ObjectInitializer = FObjectInitializer::Get());
 public:
 	/** Sets this section's key type*/
 	void SetKeyType(EParticleKey::Type InKeyType) {KeyType = InKeyType;}
diff --git a/Engine/Source/Runtime/MovieSceneCoreTypes/Classes/MovieSceneParticleTrack.h b/Engine/Source/Runtime/MovieSceneCoreTypes/Classes/MovieSceneParticleTrack.h
--- a/Engine/Source/Runtime/MovieSceneCoreTypes/Classes/MovieSceneParticleTrack.h
+++ b/Engine/Source/Runtime/MovieSceneCoreTypes/Classes/MovieSceneParticleTrack.h
@@ -13,7 +13,9 @@
 UCLASS( MinimalAPI )
 class UMovieSceneParticleTrack : public UMovieSceneTrack
 {
-	GENERATED_UCLASS_BODY()
+	GENERATED_BODY()
+public:
+	MOVIESCENECORETYPES_API UMovieSceneParticleTrack(const FObjectInitializer& ObjectInitializer = FObjectInitializer::Get());
 
 public:
 	/** UMovieSceneTrack interface */
diff --git a/Engine/Source/Runtime/MovieSceneCoreTypes/Classes/MovieSceneShotSection.h b/Engine/Source/Runtime/MovieSceneCoreTypes/Classes/MovieSceneShotSection.h
--- a/Engine/Source/Runtime/MovieSceneCoreTypes/Classes/MovieSceneShotSection.h
+++ b/Engine/Source/Runtime/MovieSceneCoreTypes/Classes/MovieSceneShotSection.h
@@ -11,7 +11,9 @@
 UCLASS( MinimalAPI )
 class UMovieSceneShotSection : public UMovieSceneSection
 {
-	GENERATED_UCLASS_BODY()
+	GENERATED_BODY()
+public:
+	MOVIESCENECORETYPES_API UMovieSceneShotSection(const FObjectInitializer& ObjectInitializer = FObjectInitializer::Get());
 public:
 	/** Sets the camera guid for this shot section */
 	void SetCameraGuid(const FGuid& InGuid);
diff --git a/Engine/Source/Runtime/MovieSceneCoreTypes/Classes/MovieSceneVectorSection.h b/Engine/Source/Runtime/MovieSceneCoreTypes/Classes/MovieSceneVectorSection.h
--- a/Engine/Source/Runtime/MovieSceneCoreTypes/Classes/MovieSceneVectorSection.h
+++ b/Engine/Source/Runtime/MovieSceneCoreTypes/Classes/MovieSceneVectorSection.h
@@ -12,7 +12,9 @@
 UCLASS(MinimalAPI )
 class UMovieSceneVectorSection : public UMovieSceneSection
 {
-	GENERATED_UCLASS_BODY()
+	GENERATED_BODY()
+public:
+	MOVIESCENECORETYPES_API UMovieSceneVectorSection(const FObjectInitializer& ObjectInitializer = FObjectInitializer::Get());
 public:
 	
 	/**
diff --git a/Engine/Source/Runtime/MovieSceneCoreTypes/Classes/MovieSceneVectorTrack.h b/Engine/Source/Runtime/MovieSceneCoreTypes/Classes/MovieSceneVectorTrack.h
--- a/Engine/Source/Runtime/MovieSceneCoreTypes/Classes/MovieSceneVectorTrack.h
+++ b/Engine/Source/Runtime/MovieSceneCoreTypes/Classes/MovieSceneVectorTrack.h
@@ -23,7 +23,9 @@ struct FVectorKey
 UCLASS(MinimalAPI)
 class UMovieSceneVectorTrack : public UMovieScenePropertyTrack
 {
-	GENERATED_UCLASS_BODY()
+	GENERATED_BODY()
+public:
+	MOVIESCENECORETYPES_API UMovieSceneVectorTrack(const FObjectInitializer& ObjectInitializer = FObjectInitializer::Get());
 public:
 	/** UMovieSceneTrack interface */
 	virtual UMovieSceneSection* CreateNewSection() override;
diff --git a/Engine/Source/Runtime/MovieSceneCoreTypes/Classes/SubMovieSceneSection.h b/Engine/Source/Runtime/MovieSceneCoreTypes/Classes/SubMovieSceneSection.h
--- a/Engine/Source/Runtime/MovieSceneCoreTypes/Classes/SubMovieSceneSection.h
+++ b/Engine/Source/Runtime/MovieSceneCoreTypes/Classes/SubMovieSceneSection.h
@@ -10,7 +10,9 @@
 UCLASS( MinimalAPI )
 class USubMovieSceneSection : public UMovieSceneSection
 {
-	GENERATED_UCLASS_BODY()
+	GENERATED_BODY()
+public:
+	MOVIESCENECORETYPES_API USubMovieSceneSection(const FObjectInitializer& ObjectInitializer = FObjectInitializer::Get());
 public:
 	/**
 	 * Sets the movie scene played by this section
diff --git a/Engine/Source/Runtime/MovieSceneCoreTypes/Classes/SubMovieSceneTrack.h b/Engine/Source/Runtime/MovieSceneCoreTypes/Classes/SubMovieSceneTrack.h
--- a/Engine/Source/Runtime/MovieSceneCoreTypes/Classes/SubMovieSceneTrack.h
+++ b/Engine/Source/Runtime/MovieSceneCoreTypes/Classes/SubMovieSceneTrack.h
@@ -11,7 +11,9 @@
 UCLASS( MinimalAPI )
 class USubMovieSceneTrack : public UMovieSceneTrack
 {
-	GENERATED_UCLASS_BODY()
+	GENERATED_BODY()
+public:
+	MOVIESCENECORETYPES_API USubMovieSceneTrack(const FObjectInitializer& ObjectInitializer = FObjectInitializer::Get());
 public:
 	/** UMovieSceneTrack interface */
 	virtual UMovieSceneSection* CreateNewSection() override;
diff --git a/Engine/Source/Runtime/Online/BuildPatchServices/Private/BuildPatchManifest.h b/Engine/Source/Runtime/Online/BuildPatchServices/Private/BuildPatchManifest.h
--- a/Engine/Source/Runtime/Online/BuildPatchServices/Private/BuildPatchManifest.h
+++ b/Engine/Source/Runtime/Online/BuildPatchServices/Private/BuildPatchManifest.h
@@ -102,7 +102,7 @@ namespace EManifestFileHeader
 USTRUCT()
 struct FCustomFieldData
 {
-	GENERATED_USTRUCT_BODY()
+	GENERATED_BODY()
 
 	UPROPERTY()
 	FString Key;
@@ -120,7 +120,7 @@ struct FCustomFieldData
 USTRUCT()
 struct FChunkInfoData
 {
-	GENERATED_USTRUCT_BODY()
+	GENERATED_BODY()
 
 	UPROPERTY()
 	FGuid Guid;
@@ -143,7 +143,7 @@ struct FChunkInfoData
 USTRUCT()
 struct FChunkPartData
 {
-	GENERATED_USTRUCT_BODY()
+	GENERATED_BODY()
 
 	// The GUID of the chunk containing this part
 	UPROPERTY()
@@ -166,7 +166,7 @@ struct FChunkPartData
 USTRUCT()
 struct FFileManifestData
 {
-	GENERATED_USTRUCT_BODY()
+	GENERATED_BODY()
 
 	UPROPERTY()
 	FString Filename;
@@ -207,7 +207,9 @@ private:
 UCLASS()
 class UBuildPatchManifest : public UObject
 {
-	GENERATED_UCLASS_BODY()
+	GENERATED_BODY()
+public:
+	UBuildPatchManifest(const FObjectInitializer& ObjectInitializer = FObjectInitializer::Get());
 
 	UPROPERTY()
 	uint8 ManifestFileVersion;
diff --git a/Engine/Source/Runtime/Online/OnlineSubsystem/Public/Interfaces/TurnBasedMatchInterface.h b/Engine/Source/Runtime/Online/OnlineSubsystem/Public/Interfaces/TurnBasedMatchInterface.h
--- a/Engine/Source/Runtime/Online/OnlineSubsystem/Public/Interfaces/TurnBasedMatchInterface.h
+++ b/Engine/Source/Runtime/Online/OnlineSubsystem/Public/Interfaces/TurnBasedMatchInterface.h
@@ -9,12 +9,15 @@ class IOnlineSubsystem;
 UINTERFACE(Blueprintable)
 class UTurnBasedMatchInterface : public UInterface
 {
-	GENERATED_UINTERFACE_BODY()
+	GENERATED_BODY()
+public:
+	UTurnBasedMatchInterface(const FObjectInitializer& ObjectInitializer = FObjectInitializer::Get());
 };
 
 class ITurnBasedMatchInterface
 {
-	GENERATED_IINTERFACE_BODY()
+	GENERATED_BODY()
+public:
 
 public:
 	UFUNCTION(BlueprintImplementableEvent, Category = "Online|TurnBased")
diff --git a/Engine/Source/Runtime/Online/OnlineSubsystem/Public/NamedInterfaces.h b/Engine/Source/Runtime/Online/OnlineSubsystem/Public/NamedInterfaces.h
--- a/Engine/Source/Runtime/Online/OnlineSubsystem/Public/NamedInterfaces.h
+++ b/Engine/Source/Runtime/Online/OnlineSubsystem/Public/NamedInterfaces.h
@@ -11,7 +11,7 @@
 USTRUCT()
 struct FNamedInterface
 {
-	GENERATED_USTRUCT_BODY()
+	GENERATED_BODY()
 
 	FNamedInterface() :
 		InterfaceName(NAME_None),
@@ -31,7 +31,7 @@ struct FNamedInterface
 USTRUCT()
 struct FNamedInterfaceDef
 {
-	GENERATED_USTRUCT_BODY()
+	GENERATED_BODY()
 
 	FNamedInterfaceDef() : 
 		InterfaceName(NAME_None)
@@ -52,7 +52,9 @@ struct FNamedInterfaceDef
 UCLASS(transient, config=Engine)
 class UNamedInterfaces: public UObject
 {
-	GENERATED_UCLASS_BODY()
+	GENERATED_BODY()
+public:
+	UNamedInterfaces(const FObjectInitializer& ObjectInitializer = FObjectInitializer::Get());
 
 	void Initialize();
 	class UObject* GetNamedInterface(FName InterfaceName) const;
diff --git a/Engine/Source/Runtime/Online/OnlineSubsystemSteam/Classes/SteamNetConnection.h b/Engine/Source/Runtime/Online/OnlineSubsystemSteam/Classes/SteamNetConnection.h
--- a/Engine/Source/Runtime/Online/OnlineSubsystemSteam/Classes/SteamNetConnection.h
+++ b/Engine/Source/Runtime/Online/OnlineSubsystemSteam/Classes/SteamNetConnection.h
@@ -10,7 +10,9 @@
 UCLASS(transient, config=Engine)
 class USteamNetConnection : public UIpConnection
 {
-    GENERATED_UCLASS_BODY()
+    GENERATED_BODY()
+public:
+    USteamNetConnection(const FObjectInitializer& ObjectInitializer = FObjectInitializer::Get());
 
 	// Begin UIpConnection Interface
 	virtual void InitRemoteConnection(class UNetDriver* InDriver, class FSocket* InSocket, const FURL& InURL, const class FInternetAddr& InRemoteAddr, EConnectionState InState, int32 InMaxPacket = 0, int32 InPacketOverhead = 0) override;
diff --git a/Engine/Source/Runtime/Online/OnlineSubsystemSteam/Classes/SteamNetDriver.h b/Engine/Source/Runtime/Online/OnlineSubsystemSteam/Classes/SteamNetDriver.h
--- a/Engine/Source/Runtime/Online/OnlineSubsystemSteam/Classes/SteamNetDriver.h
+++ b/Engine/Source/Runtime/Online/OnlineSubsystemSteam/Classes/SteamNetDriver.h
@@ -10,7 +10,9 @@
 UCLASS(transient, config=Engine)
 class USteamNetDriver : public UIpNetDriver
 {
-	GENERATED_UCLASS_BODY()
+	GENERATED_BODY()
+public:
+	USteamNetDriver(const FObjectInitializer& ObjectInitializer = FObjectInitializer::Get());
 
 	/** Time between connection detail output */
 	double ConnectionDumpInterval;
diff --git a/Engine/Source/Runtime/Online/OnlineSubsystemUtils/Classes/AchievementBlueprintLibrary.h b/Engine/Source/Runtime/Online/OnlineSubsystemUtils/Classes/AchievementBlueprintLibrary.h
--- a/Engine/Source/Runtime/Online/OnlineSubsystemUtils/Classes/AchievementBlueprintLibrary.h
+++ b/Engine/Source/Runtime/Online/OnlineSubsystemUtils/Classes/AchievementBlueprintLibrary.h
@@ -8,7 +8,9 @@
 UCLASS()
 class ONLINESUBSYSTEMUTILS_API UAchievementBlueprintLibrary : public UBlueprintFunctionLibrary
 {
-	GENERATED_UCLASS_BODY()
+	GENERATED_BODY()
+public:
+	UAchievementBlueprintLibrary(const FObjectInitializer& ObjectInitializer = FObjectInitializer::Get());
 
 	/**
 	 * Gets the status of an achievement ID (you must call CacheAchievements first to cache them)
diff --git a/Engine/Source/Runtime/Online/OnlineSubsystemUtils/Classes/AchievementQueryCallbackProxy.h b/Engine/Source/Runtime/Online/OnlineSubsystemUtils/Classes/AchievementQueryCallbackProxy.h
--- a/Engine/Source/Runtime/Online/OnlineSubsystemUtils/Classes/AchievementQueryCallbackProxy.h
+++ b/Engine/Source/Runtime/Online/OnlineSubsystemUtils/Classes/AchievementQueryCallbackProxy.h
@@ -7,7 +7,9 @@
 UCLASS(MinimalAPI)
 class UAchievementQueryCallbackProxy : public UOnlineBlueprintCallProxyBase
 {
-	GENERATED_UCLASS_BODY()
+	GENERATED_BODY()
+public:
+	ONLINESUBSYSTEMUTILS_API UAchievementQueryCallbackProxy(const FObjectInitializer& ObjectInitializer = FObjectInitializer::Get());
 
 	// Called when there is a successful query
 	UPROPERTY(BlueprintAssignable)
diff --git a/Engine/Source/Runtime/Online/OnlineSubsystemUtils/Classes/AchievementWriteCallbackProxy.h b/Engine/Source/Runtime/Online/OnlineSubsystemUtils/Classes/AchievementWriteCallbackProxy.h
--- a/Engine/Source/Runtime/Online/OnlineSubsystemUtils/Classes/AchievementWriteCallbackProxy.h
+++ b/Engine/Source/Runtime/Online/OnlineSubsystemUtils/Classes/AchievementWriteCallbackProxy.h
@@ -9,7 +9,9 @@ DECLARE_DYNAMIC_MULTICAST_DELEGATE_ThreeParams(FAchievementWriteDelegate, FName,
 UCLASS(MinimalAPI)
 class UAchievementWriteCallbackProxy : public UOnlineBlueprintCallProxyBase
 {
-	GENERATED_UCLASS_BODY()
+	GENERATED_BODY()
+public:
+	ONLINESUBSYSTEMUTILS_API UAchievementWriteCallbackProxy(const FObjectInitializer& ObjectInitializer = FObjectInitializer::Get());
 
 	// Called when there is a successful achievement write
 	UPROPERTY(BlueprintAssignable)
diff --git a/Engine/Source/Runtime/Online/OnlineSubsystemUtils/Classes/ConnectionCallbackProxy.h b/Engine/Source/Runtime/Online/OnlineSubsystemUtils/Classes/ConnectionCallbackProxy.h
--- a/Engine/Source/Runtime/Online/OnlineSubsystemUtils/Classes/ConnectionCallbackProxy.h
+++ b/Engine/Source/Runtime/Online/OnlineSubsystemUtils/Classes/ConnectionCallbackProxy.h
@@ -11,7 +11,9 @@ DECLARE_DYNAMIC_MULTICAST_DELEGATE_OneParam(FOnlineConnectionResult, int32, Erro
 UCLASS(MinimalAPI)
 class UConnectionCallbackProxy : public UOnlineBlueprintCallProxyBase
 {
-	GENERATED_UCLASS_BODY()
+	GENERATED_BODY()
+public:
+	ONLINESUBSYSTEMUTILS_API UConnectionCallbackProxy(const FObjectInitializer& ObjectInitializer = FObjectInitializer::Get());
 
 	// Called when there is a successful query
 	UPROPERTY(BlueprintAssignable)
@@ -45,4 +47,4 @@ private:
 
 	// The world context object in which this call is taking place
 	UObject* WorldContextObject;
-};
\ No newline at end of file
+};
diff --git a/Engine/Source/Runtime/Online/OnlineSubsystemUtils/Classes/CreateSessionCallbackProxy.h b/Engine/Source/Runtime/Online/OnlineSubsystemUtils/Classes/CreateSessionCallbackProxy.h
--- a/Engine/Source/Runtime/Online/OnlineSubsystemUtils/Classes/CreateSessionCallbackProxy.h
+++ b/Engine/Source/Runtime/Online/OnlineSubsystemUtils/Classes/CreateSessionCallbackProxy.h
@@ -7,7 +7,9 @@
 UCLASS(MinimalAPI)
 class UCreateSessionCallbackProxy : public UOnlineBlueprintCallProxyBase
 {
-	GENERATED_UCLASS_BODY()
+	GENERATED_BODY()
+public:
+	ONLINESUBSYSTEMUTILS_API UCreateSessionCallbackProxy(const FObjectInitializer& ObjectInitializer = FObjectInitializer::Get());
 
 	// Called when the session was created successfully
 	UPROPERTY(BlueprintAssignable)
diff --git a/Engine/Source/Runtime/Online/OnlineSubsystemUtils/Classes/DestroySessionCallbackProxy.h b/Engine/Source/Runtime/Online/OnlineSubsystemUtils/Classes/DestroySessionCallbackProxy.h
--- a/Engine/Source/Runtime/Online/OnlineSubsystemUtils/Classes/DestroySessionCallbackProxy.h
+++ b/Engine/Source/Runtime/Online/OnlineSubsystemUtils/Classes/DestroySessionCallbackProxy.h
@@ -7,7 +7,9 @@
 UCLASS(MinimalAPI)
 class UDestroySessionCallbackProxy : public UOnlineBlueprintCallProxyBase
 {
-	GENERATED_UCLASS_BODY()
+	GENERATED_BODY()
+public:
+	ONLINESUBSYSTEMUTILS_API UDestroySessionCallbackProxy(const FObjectInitializer& ObjectInitializer = FObjectInitializer::Get());
 
 	// Called when there is a successful destroy
 	UPROPERTY(BlueprintAssignable)
diff --git a/Engine/Source/Runtime/Online/OnlineSubsystemUtils/Classes/EndTurnCallbackProxy.h b/Engine/Source/Runtime/Online/OnlineSubsystemUtils/Classes/EndTurnCallbackProxy.h
--- a/Engine/Source/Runtime/Online/OnlineSubsystemUtils/Classes/EndTurnCallbackProxy.h
+++ b/Engine/Source/Runtime/Online/OnlineSubsystemUtils/Classes/EndTurnCallbackProxy.h
@@ -11,7 +11,9 @@ class AActor;
 UCLASS(MinimalAPI)
 class UEndTurnCallbackProxy : public UOnlineBlueprintCallProxyBase
 {
-	GENERATED_UCLASS_BODY()
+	GENERATED_BODY()
+public:
+	ONLINESUBSYSTEMUTILS_API UEndTurnCallbackProxy(const FObjectInitializer& ObjectInitializer = FObjectInitializer::Get());
 
 	virtual ~UEndTurnCallbackProxy();
 
@@ -45,4 +47,4 @@ private:
 
 	// Object that is inheriting from UTurnBasedMatchInterface. Replicated properties will be serialize into the platform specific match data object.
 	UTurnBasedMatchInterface* TurnBasedMatchInterface;
-};
\ No newline at end of file
+};
