Back out changelist 2481333

diff --git a/Engine/Source/Runtime/EngineMessages/Classes/EngineServiceMessages.h b/Engine/Source/Runtime/EngineMessages/Classes/EngineServiceMessages.h
--- a/Engine/Source/Runtime/EngineMessages/Classes/EngineServiceMessages.h
+++ b/Engine/Source/Runtime/EngineMessages/Classes/EngineServiceMessages.h
@@ -14,7 +14,7 @@
 USTRUCT()
 struct FEngineServicePing
 {
-	GENERATED_BODY()
+	GENERATED_USTRUCT_BODY()
 };
 
 template<>
@@ -33,7 +33,7 @@ struct TStructOpsTypeTraits<FEngineServicePing> : public TStructOpsTypeTraitsBas
 USTRUCT()
 struct FEngineServicePong
 {
-	GENERATED_BODY()
+	GENERATED_USTRUCT_BODY()
 
 	/** Holds the name of the currently loaded level, if any. */
 	UPROPERTY()
@@ -83,7 +83,7 @@ struct TStructOpsTypeTraits<FEngineServicePong> : public TStructOpsTypeTraitsBas
 USTRUCT()
 struct FEngineServiceAuthDeny
 {
-	GENERATED_BODY()
+	GENERATED_USTRUCT_BODY()
 
 	/** Holds the name of the user that denied access. */
 	UPROPERTY()
@@ -110,7 +110,7 @@ struct TStructOpsTypeTraits<FEngineServiceAuthDeny> : public TStructOpsTypeTrait
 USTRUCT()
 struct FEngineServiceAuthGrant
 {
-	GENERATED_BODY()
+	GENERATED_USTRUCT_BODY()
 
 	/** Holds the name of the user that granted access. */
 	UPROPERTY()
@@ -141,7 +141,7 @@ struct TStructOpsTypeTraits<FEngineServiceAuthGrant> : public TStructOpsTypeTrai
 USTRUCT()
 struct FEngineServiceExecuteCommand
 {
-	GENERATED_BODY()
+	GENERATED_USTRUCT_BODY()
 
 	/** Holds the command to execute. */
 	UPROPERTY()
@@ -181,7 +181,7 @@ struct TStructOpsTypeTraits<FEngineServiceExecuteCommand> : public TStructOpsTyp
 USTRUCT()
 struct FEngineServiceTerminate
 {
-	GENERATED_BODY()
+	GENERATED_USTRUCT_BODY()
 
 	/** Holds the name of the user that wants to terminate the engine. */
 	UPROPERTY()
@@ -219,7 +219,7 @@ struct TStructOpsTypeTraits<FEngineServiceTerminate> : public TStructOpsTypeTrai
 USTRUCT()
 struct FEngineServiceNotification
 {
-	GENERATED_BODY()
+	GENERATED_USTRUCT_BODY()
 
 	/** Holds the notification text. */
 	UPROPERTY()
diff --git a/Engine/Source/Runtime/EngineSettings/Classes/ConsoleSettings.h b/Engine/Source/Runtime/EngineSettings/Classes/ConsoleSettings.h
--- a/Engine/Source/Runtime/EngineSettings/Classes/ConsoleSettings.h
+++ b/Engine/Source/Runtime/EngineSettings/Classes/ConsoleSettings.h
@@ -11,7 +11,7 @@
 USTRUCT()
 struct FAutoCompleteCommand
 {
-	GENERATED_BODY()
+	GENERATED_USTRUCT_BODY()
 
 	UPROPERTY(config, EditAnywhere, Category=Command)
 	FString Command;
@@ -28,9 +28,7 @@ UCLASS(config=Input, defaultconfig)
 class ENGINESETTINGS_API UConsoleSettings
 	: public UObject
 {
-	GENERATED_BODY()
-public:
-	UConsoleSettings(const FObjectInitializer& ObjectInitializer = FObjectInitializer::Get());
+	GENERATED_UCLASS_BODY()
 
 	/**  Visible Console stuff */
 	UPROPERTY(globalconfig, EditAnywhere, Category=General)
diff --git a/Engine/Source/Runtime/EngineSettings/Classes/GameMapsSettings.h b/Engine/Source/Runtime/EngineSettings/Classes/GameMapsSettings.h
--- a/Engine/Source/Runtime/EngineSettings/Classes/GameMapsSettings.h
+++ b/Engine/Source/Runtime/EngineSettings/Classes/GameMapsSettings.h
@@ -33,9 +33,7 @@ UCLASS(config=Engine, defaultconfig)
 class ENGINESETTINGS_API UGameMapsSettings
 	: public UObject
 {
-	GENERATED_BODY()
-public:
-	UGameMapsSettings(const FObjectInitializer& ObjectInitializer = FObjectInitializer::Get());
+	GENERATED_UCLASS_BODY()
 
 	/**
 	 * Get the default map specified in the settings.
diff --git a/Engine/Source/Runtime/EngineSettings/Classes/GameNetworkManagerSettings.h b/Engine/Source/Runtime/EngineSettings/Classes/GameNetworkManagerSettings.h
--- a/Engine/Source/Runtime/EngineSettings/Classes/GameNetworkManagerSettings.h
+++ b/Engine/Source/Runtime/EngineSettings/Classes/GameNetworkManagerSettings.h
@@ -12,9 +12,7 @@ UCLASS(config=Game)
 class ENGINESETTINGS_API UGameNetworkManagerSettings
 	: public UObject
 {
-	GENERATED_BODY()
-public:
-	UGameNetworkManagerSettings(const FObjectInitializer& ObjectInitializer = FObjectInitializer::Get());
+	GENERATED_UCLASS_BODY()
 
 	/** Minimum bandwidth dynamically set per connection. */
 	UPROPERTY(globalconfig, EditAnywhere, Category=Bandwidth)
diff --git a/Engine/Source/Runtime/EngineSettings/Classes/GameSessionSettings.h b/Engine/Source/Runtime/EngineSettings/Classes/GameSessionSettings.h
--- a/Engine/Source/Runtime/EngineSettings/Classes/GameSessionSettings.h
+++ b/Engine/Source/Runtime/EngineSettings/Classes/GameSessionSettings.h
@@ -9,9 +9,7 @@ UCLASS(config=Game)
 class ENGINESETTINGS_API UGameSessionSettings
 	: public UObject
 {
-	GENERATED_BODY()
-public:
-	UGameSessionSettings(const FObjectInitializer& ObjectInitializer = FObjectInitializer::Get());
+	GENERATED_UCLASS_BODY()
 
 	/** Maximum number of spectators allowed by this server. */
 	UPROPERTY(globalconfig, EditAnywhere, Category=GameSessionSettings)
diff --git a/Engine/Source/Runtime/EngineSettings/Classes/GeneralEngineSettings.h b/Engine/Source/Runtime/EngineSettings/Classes/GeneralEngineSettings.h
--- a/Engine/Source/Runtime/EngineSettings/Classes/GeneralEngineSettings.h
+++ b/Engine/Source/Runtime/EngineSettings/Classes/GeneralEngineSettings.h
@@ -9,8 +9,6 @@ UCLASS(config=Engine)
 class ENGINESETTINGS_API UGeneralEngineSettings
 	: public UObject
 {
-	GENERATED_BODY()
-public:
-	UGeneralEngineSettings(const FObjectInitializer& ObjectInitializer = FObjectInitializer::Get());
+	GENERATED_UCLASS_BODY()
 
 };
diff --git a/Engine/Source/Runtime/EngineSettings/Classes/GeneralProjectSettings.h b/Engine/Source/Runtime/EngineSettings/Classes/GeneralProjectSettings.h
--- a/Engine/Source/Runtime/EngineSettings/Classes/GeneralProjectSettings.h
+++ b/Engine/Source/Runtime/EngineSettings/Classes/GeneralProjectSettings.h
@@ -9,9 +9,7 @@ UCLASS(config=Game, defaultconfig)
 class ENGINESETTINGS_API UGeneralProjectSettings
 	: public UObject
 {
-	GENERATED_BODY()
-public:
-	UGeneralProjectSettings(const FObjectInitializer& ObjectInitializer = FObjectInitializer::Get());
+	GENERATED_UCLASS_BODY()
 
 	/** The name of the company (author, provider) that created the project. */
 	UPROPERTY(config, EditAnywhere, Category=Publisher)
diff --git a/Engine/Source/Runtime/EngineSettings/Classes/HudSettings.h b/Engine/Source/Runtime/EngineSettings/Classes/HudSettings.h
--- a/Engine/Source/Runtime/EngineSettings/Classes/HudSettings.h
+++ b/Engine/Source/Runtime/EngineSettings/Classes/HudSettings.h
@@ -9,9 +9,7 @@ UCLASS(config=Game)
 class ENGINESETTINGS_API UHudSettings
 	: public UObject
 {
-	GENERATED_BODY()
-public:
-	UHudSettings(const FObjectInitializer& ObjectInitializer = FObjectInitializer::Get());
+	GENERATED_UCLASS_BODY()
 
 	/** Whether the HUD is visible at all.	 */
 	UPROPERTY(config, EditAnywhere, Category=HudSettings)
diff --git a/Engine/Source/Runtime/Foliage/Private/FoliageStatistics.h b/Engine/Source/Runtime/Foliage/Private/FoliageStatistics.h
--- a/Engine/Source/Runtime/Foliage/Private/FoliageStatistics.h
+++ b/Engine/Source/Runtime/Foliage/Private/FoliageStatistics.h
@@ -7,9 +7,7 @@
 UCLASS()
 class FOLIAGE_API UFoliageStatistics : public UBlueprintFunctionLibrary
 {
-	GENERATED_BODY()
-public:
-	UFoliageStatistics(const FObjectInitializer& ObjectInitializer = FObjectInitializer::Get());
+	GENERATED_UCLASS_BODY()
 
 	/**
 	* Counts how many foliage instances overlap a given sphere
diff --git a/Engine/Source/Runtime/Foliage/Private/InteractiveFoliageComponent.h b/Engine/Source/Runtime/Foliage/Private/InteractiveFoliageComponent.h
--- a/Engine/Source/Runtime/Foliage/Private/InteractiveFoliageComponent.h
+++ b/Engine/Source/Runtime/Foliage/Private/InteractiveFoliageComponent.h
@@ -7,9 +7,7 @@
 UCLASS(hidecategories=Object)
 class UInteractiveFoliageComponent : public UStaticMeshComponent
 {
-	GENERATED_BODY()
-public:
-	UInteractiveFoliageComponent(const FObjectInitializer& ObjectInitializer = FObjectInitializer::Get());
+	GENERATED_UCLASS_BODY()
 
 
 public:
diff --git a/Engine/Source/Runtime/Foliage/Private/ProceduralFoliageBlockingVolume.h b/Engine/Source/Runtime/Foliage/Private/ProceduralFoliageBlockingVolume.h
--- a/Engine/Source/Runtime/Foliage/Private/ProceduralFoliageBlockingVolume.h
+++ b/Engine/Source/Runtime/Foliage/Private/ProceduralFoliageBlockingVolume.h
@@ -11,9 +11,7 @@ class AProceduralFoliageActor;
 UCLASS(MinimalAPI)
 class AProceduralFoliageBlockingVolume : public AVolume
 {
-	GENERATED_BODY()
-public:
-	FOLIAGE_API AProceduralFoliageBlockingVolume(const FObjectInitializer& ObjectInitializer = FObjectInitializer::Get());
+	GENERATED_UCLASS_BODY()
 
 	UPROPERTY(Category = ProceduralFoliage, EditAnywhere)
 	AProceduralFoliageActor* ProceduralFoliageActor;
diff --git a/Engine/Source/Runtime/Foliage/Public/FoliageType.h b/Engine/Source/Runtime/Foliage/Public/FoliageType.h
--- a/Engine/Source/Runtime/Foliage/Public/FoliageType.h
+++ b/Engine/Source/Runtime/Foliage/Public/FoliageType.h
@@ -37,9 +37,7 @@ enum class EFoliageScaling : uint8
 UCLASS(hidecategories = Object, editinlinenew, MinimalAPI, BlueprintType, Blueprintable)
 class UFoliageType : public UObject
 {
-	GENERATED_BODY()
-public:
-	FOLIAGE_API UFoliageType(const FObjectInitializer& ObjectInitializer = FObjectInitializer::Get());
+	GENERATED_UCLASS_BODY()
 
 	/* A GUID that is updated every time the foliage type is modified, 
 	   so foliage placed in the level can detect the FoliageType has changed. */
diff --git a/Engine/Source/Runtime/Foliage/Public/FoliageType_InstancedStaticMesh.h b/Engine/Source/Runtime/Foliage/Public/FoliageType_InstancedStaticMesh.h
--- a/Engine/Source/Runtime/Foliage/Public/FoliageType_InstancedStaticMesh.h
+++ b/Engine/Source/Runtime/Foliage/Public/FoliageType_InstancedStaticMesh.h
@@ -8,9 +8,7 @@
 UCLASS(hidecategories=Object, editinlinenew, MinimalAPI)
 class UFoliageType_InstancedStaticMesh : public UFoliageType
 {
-	GENERATED_BODY()
-public:
-	FOLIAGE_API UFoliageType_InstancedStaticMesh(const FObjectInitializer& ObjectInitializer = FObjectInitializer::Get());
+	GENERATED_UCLASS_BODY()
 
 	UPROPERTY(EditAnywhere, Category=FoliageType)
 	UStaticMesh* Mesh;
diff --git a/Engine/Source/Runtime/Foliage/Public/InstancedFoliageActor.h b/Engine/Source/Runtime/Foliage/Public/InstancedFoliageActor.h
--- a/Engine/Source/Runtime/Foliage/Public/InstancedFoliageActor.h
+++ b/Engine/Source/Runtime/Foliage/Public/InstancedFoliageActor.h
@@ -23,9 +23,7 @@ struct FDesiredFoliageInstance;
 UCLASS(notplaceable, hidecategories = Object, MinimalAPI, NotBlueprintable)
 class AInstancedFoliageActor : public AActor
 {
-	GENERATED_BODY()
-public:
-	FOLIAGE_API AInstancedFoliageActor(const FObjectInitializer& ObjectInitializer = FObjectInitializer::Get());
+	GENERATED_UCLASS_BODY()
 
 public:
 #if WITH_EDITORONLY_DATA
diff --git a/Engine/Source/Runtime/Foliage/Public/InteractiveFoliageActor.h b/Engine/Source/Runtime/Foliage/Public/InteractiveFoliageActor.h
--- a/Engine/Source/Runtime/Foliage/Public/InteractiveFoliageActor.h
+++ b/Engine/Source/Runtime/Foliage/Public/InteractiveFoliageActor.h
@@ -10,9 +10,7 @@
 UCLASS(MinimalAPI)
 class AInteractiveFoliageActor : public AStaticMeshActor
 {
-	GENERATED_BODY()
-public:
-	FOLIAGE_API AInteractiveFoliageActor(const FObjectInitializer& ObjectInitializer = FObjectInitializer::Get());
+	GENERATED_UCLASS_BODY()
 
 private:
 	/** Collision cylinder */
diff --git a/Engine/Source/Runtime/Foliage/Public/ProceduralFoliage.h b/Engine/Source/Runtime/Foliage/Public/ProceduralFoliage.h
--- a/Engine/Source/Runtime/Foliage/Public/ProceduralFoliage.h
+++ b/Engine/Source/Runtime/Foliage/Public/ProceduralFoliage.h
@@ -12,7 +12,7 @@ class UProceduralFoliageTile;
 USTRUCT()
 struct FProceduralFoliageTypeData
 {
-	GENERATED_BODY()
+	GENERATED_USTRUCT_BODY()
 
 	UPROPERTY(Category = ProceduralFoliageSimulation, EditAnywhere, BlueprintReadOnly)
 	TSubclassOf<UFoliageType_InstancedStaticMesh> Type;
@@ -27,9 +27,7 @@ struct FProceduralFoliageTypeData
 UCLASS(BlueprintType, Blueprintable)
 class FOLIAGE_API UProceduralFoliage : public UObject
 {
-	GENERATED_BODY()
-public:
-	UProceduralFoliage(const FObjectInitializer& ObjectInitializer = FObjectInitializer::Get());
+	GENERATED_UCLASS_BODY()
 
 	/**The seed used for generating the randomness of the simulation. */
 	UPROPERTY(Category = ProceduralFoliageSimulation, EditAnywhere, BlueprintReadOnly)
diff --git a/Engine/Source/Runtime/Foliage/Public/ProceduralFoliageActor.h b/Engine/Source/Runtime/Foliage/Public/ProceduralFoliageActor.h
--- a/Engine/Source/Runtime/Foliage/Public/ProceduralFoliageActor.h
+++ b/Engine/Source/Runtime/Foliage/Public/ProceduralFoliageActor.h
@@ -9,9 +9,7 @@ class UProceduralFoliageComponent;
 UCLASS()
 class FOLIAGE_API AProceduralFoliageActor: public AVolume
 {
-	GENERATED_BODY()
-public:
-	AProceduralFoliageActor(const FObjectInitializer& ObjectInitializer = FObjectInitializer::Get());
+	GENERATED_UCLASS_BODY()
 
 	UPROPERTY(Category = ProceduralFoliageTile, VisibleAnywhere, BlueprintReadOnly)
 	UProceduralFoliageComponent* ProceduralComponent;
@@ -19,4 +17,4 @@ public:
 #if WITH_EDITOR
 	virtual bool GetReferencedContentObjects(TArray<UObject*>& Objects) const override;
 #endif
-};
+};
\ No newline at end of file
diff --git a/Engine/Source/Runtime/Foliage/Public/ProceduralFoliageComponent.h b/Engine/Source/Runtime/Foliage/Public/ProceduralFoliageComponent.h
--- a/Engine/Source/Runtime/Foliage/Public/ProceduralFoliageComponent.h
+++ b/Engine/Source/Runtime/Foliage/Public/ProceduralFoliageComponent.h
@@ -12,9 +12,7 @@ struct FDesiredFoliageInstance;
 UCLASS(BlueprintType)
 class FOLIAGE_API UProceduralFoliageComponent : public UActorComponent
 {
-	GENERATED_BODY()
-public:
-	UProceduralFoliageComponent(const FObjectInitializer& ObjectInitializer = FObjectInitializer::Get());
+	GENERATED_UCLASS_BODY()
 
 	/** The overlap in Cms between two tiles. */
 	UPROPERTY(Category = "ProceduralFoliage", BlueprintReadWrite, EditAnywhere)
@@ -45,4 +43,4 @@ private:
 private:
 	UPROPERTY()
 	FGuid ProceduralGuid;
-};
+};
\ No newline at end of file
diff --git a/Engine/Source/Runtime/Foliage/Public/ProceduralFoliageTile.h b/Engine/Source/Runtime/Foliage/Public/ProceduralFoliageTile.h
--- a/Engine/Source/Runtime/Foliage/Public/ProceduralFoliageTile.h
+++ b/Engine/Source/Runtime/Foliage/Public/ProceduralFoliageTile.h
@@ -16,9 +16,7 @@ class UFoliageType_InstancedStaticMesh;
 UCLASS()
 class FOLIAGE_API UProceduralFoliageTile : public UObject
 {
-	GENERATED_BODY()
-public:
-	UProceduralFoliageTile(const FObjectInitializer& ObjectInitializer = FObjectInitializer::Get());
+	GENERATED_UCLASS_BODY()
 
 	UPROPERTY()
 	const UProceduralFoliage* ProceduralFoliage;
@@ -81,4 +79,4 @@ private:
 private:
 	float GetRandomGaussian();
 	FVector GetSeedOffset(const UFoliageType_InstancedStaticMesh* Type, float MinDistance);
-};
+};
\ No newline at end of file
diff --git a/Engine/Source/Runtime/GameLiveStreaming/GameLiveStreamingFunctionLibrary.h b/Engine/Source/Runtime/GameLiveStreaming/GameLiveStreamingFunctionLibrary.h
--- a/Engine/Source/Runtime/GameLiveStreaming/GameLiveStreamingFunctionLibrary.h
+++ b/Engine/Source/Runtime/GameLiveStreaming/GameLiveStreamingFunctionLibrary.h
@@ -5,9 +5,7 @@
 UCLASS()
 class UGameLiveStreamingFunctionLibrary : public UBlueprintFunctionLibrary
 {
-	GENERATED_BODY()
-public:
-	UGameLiveStreamingFunctionLibrary(const FObjectInitializer& ObjectInitializer = FObjectInitializer::Get());
+	GENERATED_UCLASS_BODY()
 
 
 	/**
diff --git a/Engine/Source/Runtime/GameLiveStreaming/QueryLiveStreamsCallbackProxy.h b/Engine/Source/Runtime/GameLiveStreaming/QueryLiveStreamsCallbackProxy.h
--- a/Engine/Source/Runtime/GameLiveStreaming/QueryLiveStreamsCallbackProxy.h
+++ b/Engine/Source/Runtime/GameLiveStreaming/QueryLiveStreamsCallbackProxy.h
@@ -6,7 +6,7 @@
 USTRUCT( BlueprintType )
 struct FBlueprintLiveStreamInfo 
 {
-	GENERATED_BODY()
+	GENERATED_USTRUCT_BODY()
 
 	/** Name of the game that is streaming */
 	UPROPERTY( EditAnywhere, Category="LiveStreaming" )
@@ -27,9 +27,7 @@ DECLARE_DYNAMIC_MULTICAST_DELEGATE_TwoParams( FOnQueryLiveStreamsCompleted, cons
 UCLASS()
 class UQueryLiveStreamsCallbackProxy : public UOnlineBlueprintCallProxyBase
 {
-	GENERATED_BODY()
-public:
-	UQueryLiveStreamsCallbackProxy(const FObjectInitializer& ObjectInitializer = FObjectInitializer::Get());
+	GENERATED_UCLASS_BODY()
 
 	// Called when the asynchronous request for live streams completes, passing along the list of live streams currently active, along with a boolean value that indicates whether the request was successful at all
 	UPROPERTY( BlueprintAssignable )
@@ -47,4 +45,4 @@ protected:
 
 	/** Name of the game that we're querying about */
 	FString GameName;
-};
+};
\ No newline at end of file
diff --git a/Engine/Source/Runtime/GameMenuBuilder/Public/GameMenuWidgetStyle.h b/Engine/Source/Runtime/GameMenuBuilder/Public/GameMenuWidgetStyle.h
--- a/Engine/Source/Runtime/GameMenuBuilder/Public/GameMenuWidgetStyle.h
+++ b/Engine/Source/Runtime/GameMenuBuilder/Public/GameMenuWidgetStyle.h
@@ -20,7 +20,7 @@ namespace GameMenuLayoutType
 USTRUCT()
 struct FGameMenuStyle : public FSlateWidgetStyle
 {
-	GENERATED_BODY()
+	GENERATED_USTRUCT_BODY()
 
 	FGameMenuStyle();
 	virtual ~FGameMenuStyle();
@@ -178,9 +178,7 @@ struct FGameMenuStyle : public FSlateWidgetStyle
 UCLASS(hidecategories = Object, MinimalAPI)
 class UGameMenuWidgetStyle : public USlateWidgetStyleContainerBase
 {
-	GENERATED_BODY()
-public:
-	GAMEMENUBUILDER_API UGameMenuWidgetStyle(const FObjectInitializer& ObjectInitializer = FObjectInitializer::Get());
+	GENERATED_UCLASS_BODY()
 
 public:
 	/** The actual data describing the menu's appearance. */
@@ -191,4 +189,4 @@ public:
 	{
 		return static_cast< const struct FSlateWidgetStyle* >(&MenuStyle);
 	}
-};
+};
\ No newline at end of file
diff --git a/Engine/Source/Runtime/GameplayAbilities/Classes/GameplayAbilityBlueprint.h b/Engine/Source/Runtime/GameplayAbilities/Classes/GameplayAbilityBlueprint.h
--- a/Engine/Source/Runtime/GameplayAbilities/Classes/GameplayAbilityBlueprint.h
+++ b/Engine/Source/Runtime/GameplayAbilities/Classes/GameplayAbilityBlueprint.h
@@ -13,9 +13,7 @@
 UCLASS(BlueprintType)
 class GAMEPLAYABILITIES_API UGameplayAbilityBlueprint : public UBlueprint
 {
-	GENERATED_BODY()
-public:
-	UGameplayAbilityBlueprint(const FObjectInitializer& ObjectInitializer = FObjectInitializer::Get());
+	GENERATED_UCLASS_BODY()
 
 #if WITH_EDITOR
 
diff --git a/Engine/Source/Runtime/GameplayAbilities/Public/Abilities/GameplayAbility.h b/Engine/Source/Runtime/GameplayAbilities/Public/Abilities/GameplayAbility.h
--- a/Engine/Source/Runtime/GameplayAbilities/Public/Abilities/GameplayAbility.h
+++ b/Engine/Source/Runtime/GameplayAbilities/Public/Abilities/GameplayAbility.h
@@ -77,7 +77,7 @@ DECLARE_MULTICAST_DELEGATE_OneParam(FOnGameplayAbilityStateEnded, FName);
 USTRUCT()
 struct FAbilityTriggerData
 {
-	GENERATED_BODY()
+	GENERATED_USTRUCT_BODY()
 
 	FAbilityTriggerData() 
 	: TriggerSource(EGameplayAbilityTriggerSource::GameplayEvent)
@@ -98,9 +98,7 @@ struct FAbilityTriggerData
 UCLASS(Blueprintable)
 class GAMEPLAYABILITIES_API UGameplayAbility : public UObject
 {
-	GENERATED_BODY()
-public:
-	UGameplayAbility(const FObjectInitializer& ObjectInitializer = FObjectInitializer::Get());
+	GENERATED_UCLASS_BODY()
 
 	friend class UAbilitySystemComponent;
 	friend class UGameplayAbilitySet;
diff --git a/Engine/Source/Runtime/GameplayAbilities/Public/Abilities/GameplayAbilityTargetActor.h b/Engine/Source/Runtime/GameplayAbilities/Public/Abilities/GameplayAbilityTargetActor.h
--- a/Engine/Source/Runtime/GameplayAbilities/Public/Abilities/GameplayAbilityTargetActor.h
+++ b/Engine/Source/Runtime/GameplayAbilities/Public/Abilities/GameplayAbilityTargetActor.h
@@ -20,9 +20,7 @@ namespace ETargetAbilitySelfSelection
 UCLASS(Blueprintable, abstract, notplaceable)
 class GAMEPLAYABILITIES_API AGameplayAbilityTargetActor : public AActor
 {
-	GENERATED_BODY()
-public:
-	AGameplayAbilityTargetActor(const FObjectInitializer& ObjectInitializer = FObjectInitializer::Get());
+	GENERATED_UCLASS_BODY()
 
 public:
 
@@ -99,4 +97,4 @@ public:
 
 	FDelegateHandle GenericConfirmHandle;
 	FDelegateHandle GenericCancelHandle;
-};
+};
\ No newline at end of file
diff --git a/Engine/Source/Runtime/GameplayAbilities/Public/Abilities/GameplayAbilityTargetActor_ActorPlacement.h b/Engine/Source/Runtime/GameplayAbilities/Public/Abilities/GameplayAbilityTargetActor_ActorPlacement.h
--- a/Engine/Source/Runtime/GameplayAbilities/Public/Abilities/GameplayAbilityTargetActor_ActorPlacement.h
+++ b/Engine/Source/Runtime/GameplayAbilities/Public/Abilities/GameplayAbilityTargetActor_ActorPlacement.h
@@ -9,9 +9,7 @@
 UCLASS(Blueprintable)
 class AGameplayAbilityTargetActor_ActorPlacement : public AGameplayAbilityTargetActor_GroundTrace
 {
-	GENERATED_BODY()
-public:
-	AGameplayAbilityTargetActor_ActorPlacement(const FObjectInitializer& ObjectInitializer = FObjectInitializer::Get());
+	GENERATED_UCLASS_BODY()
 
 	virtual void EndPlay(const EEndPlayReason::Type EndPlayReason) override;
 
@@ -27,4 +25,4 @@ public:
 
 	/** Visualization for the intended location of the placed actor. */
 	TWeakObjectPtr<AGameplayAbilityWorldReticle_ActorVisualization> ActorVisualizationReticle;
-};
+};
\ No newline at end of file
diff --git a/Engine/Source/Runtime/GameplayAbilities/Public/Abilities/GameplayAbilityTargetActor_GroundTrace.h b/Engine/Source/Runtime/GameplayAbilities/Public/Abilities/GameplayAbilityTargetActor_GroundTrace.h
--- a/Engine/Source/Runtime/GameplayAbilities/Public/Abilities/GameplayAbilityTargetActor_GroundTrace.h
+++ b/Engine/Source/Runtime/GameplayAbilities/Public/Abilities/GameplayAbilityTargetActor_GroundTrace.h
@@ -9,9 +9,7 @@
 UCLASS(Blueprintable)
 class GAMEPLAYABILITIES_API AGameplayAbilityTargetActor_GroundTrace : public AGameplayAbilityTargetActor_Trace
 {
-	GENERATED_BODY()
-public:
-	AGameplayAbilityTargetActor_GroundTrace(const FObjectInitializer& ObjectInitializer = FObjectInitializer::Get());
+	GENERATED_UCLASS_BODY()
 
 	virtual void StartTargeting(UGameplayAbility* InAbility) override;
 
@@ -35,4 +33,4 @@ protected:
 	FCollisionShape CollisionShape;
 	float CollisionHeightOffset;		//When tracing, give this much extra height to avoid start-in-ground problems. Dealing with thick placement actors while standing near walls may be trickier.
 	bool bLastTraceWasGood;
-};
+};
\ No newline at end of file
diff --git a/Engine/Source/Runtime/GameplayAbilities/Public/Abilities/GameplayAbilityTargetActor_Radius.h b/Engine/Source/Runtime/GameplayAbilities/Public/Abilities/GameplayAbilityTargetActor_Radius.h
--- a/Engine/Source/Runtime/GameplayAbilities/Public/Abilities/GameplayAbilityTargetActor_Radius.h
+++ b/Engine/Source/Runtime/GameplayAbilities/Public/Abilities/GameplayAbilityTargetActor_Radius.h
@@ -11,9 +11,7 @@
 UCLASS(Blueprintable, notplaceable)
 class GAMEPLAYABILITIES_API AGameplayAbilityTargetActor_Radius : public AGameplayAbilityTargetActor
 {
-	GENERATED_BODY()
-public:
-	AGameplayAbilityTargetActor_Radius(const FObjectInitializer& ObjectInitializer = FObjectInitializer::Get());
+	GENERATED_UCLASS_BODY()
 
 public:
 
@@ -31,4 +29,4 @@ protected:
 
 	FGameplayAbilityTargetDataHandle MakeTargetData(const TArray<TWeakObjectPtr<AActor> > Actors, const FVector& Origin) const;
 
-};
+};
\ No newline at end of file
diff --git a/Engine/Source/Runtime/GameplayAbilities/Public/Abilities/GameplayAbilityTargetActor_SingleLineTrace.h b/Engine/Source/Runtime/GameplayAbilities/Public/Abilities/GameplayAbilityTargetActor_SingleLineTrace.h
--- a/Engine/Source/Runtime/GameplayAbilities/Public/Abilities/GameplayAbilityTargetActor_SingleLineTrace.h
+++ b/Engine/Source/Runtime/GameplayAbilities/Public/Abilities/GameplayAbilityTargetActor_SingleLineTrace.h
@@ -9,10 +9,8 @@
 UCLASS(Blueprintable)
 class GAMEPLAYABILITIES_API AGameplayAbilityTargetActor_SingleLineTrace : public AGameplayAbilityTargetActor_Trace
 {
-	GENERATED_BODY()
-public:
-	AGameplayAbilityTargetActor_SingleLineTrace(const FObjectInitializer& ObjectInitializer = FObjectInitializer::Get());
+	GENERATED_UCLASS_BODY()
 
 protected:
 	virtual FHitResult PerformTrace(AActor* InSourceActor) override;
-};
+};
\ No newline at end of file
diff --git a/Engine/Source/Runtime/GameplayAbilities/Public/Abilities/GameplayAbilityTargetTypes.h b/Engine/Source/Runtime/GameplayAbilities/Public/Abilities/GameplayAbilityTargetTypes.h
--- a/Engine/Source/Runtime/GameplayAbilities/Public/Abilities/GameplayAbilityTargetTypes.h
+++ b/Engine/Source/Runtime/GameplayAbilities/Public/Abilities/GameplayAbilityTargetTypes.h
@@ -59,7 +59,7 @@ namespace EGameplayTargetingConfirmation
 USTRUCT()
 struct GAMEPLAYABILITIES_API FGameplayAbilityTargetData
 {
-	GENERATED_BODY()
+	GENERATED_USTRUCT_BODY()
 
 	virtual ~FGameplayAbilityTargetData() { }
 
@@ -160,7 +160,7 @@ namespace EGameplayAbilityTargetingLocationType
 USTRUCT(BlueprintType)
 struct GAMEPLAYABILITIES_API FGameplayAbilityTargetDataHandle
 {
-	GENERATED_BODY()
+	GENERATED_USTRUCT_BODY()
 
 	FGameplayAbilityTargetDataHandle() { }
 	FGameplayAbilityTargetDataHandle(struct FGameplayAbilityTargetData* DataPtr)
@@ -248,7 +248,7 @@ struct TStructOpsTypeTraits<FGameplayAbilityTargetDataHandle> : public TStructOp
 USTRUCT(BlueprintType)
 struct GAMEPLAYABILITIES_API FGameplayAbilityTargetingLocationInfo
 {
-	GENERATED_BODY()
+	GENERATED_USTRUCT_BODY()
 
 	FGameplayAbilityTargetingLocationInfo()
 	: LocationType(EGameplayAbilityTargetingLocationType::LiteralTransform)
@@ -345,7 +345,7 @@ struct TStructOpsTypeTraits<FGameplayAbilityTargetingLocationInfo> : public TStr
 USTRUCT(BlueprintType)
 struct GAMEPLAYABILITIES_API FGameplayAbilityTargetData_LocationInfo : public FGameplayAbilityTargetData
 {
-	GENERATED_BODY()
+	GENERATED_USTRUCT_BODY()
 
 	/** Generic location data for source */
 	UPROPERTY(BlueprintReadWrite, EditAnywhere, Category = Targeting)
@@ -406,7 +406,7 @@ struct TStructOpsTypeTraits<FGameplayAbilityTargetData_LocationInfo> : public TS
 USTRUCT(BlueprintType)
 struct GAMEPLAYABILITIES_API FGameplayAbilityTargetData_ActorArray : public FGameplayAbilityTargetData
 {
-	GENERATED_BODY()
+	GENERATED_USTRUCT_BODY()
 
 	/** We could be selecting this group of actors from any type of location, so use a generic location type */
 	UPROPERTY(BlueprintReadWrite, EditAnywhere, Category = Targeting)
@@ -508,7 +508,7 @@ struct TStructOpsTypeTraits<FGameplayAbilityTargetData_ActorArray> : public TStr
 USTRUCT(BlueprintType)
 struct GAMEPLAYABILITIES_API FGameplayAbilityTargetData_SingleTargetHit : public FGameplayAbilityTargetData
 {
-	GENERATED_BODY()
+	GENERATED_USTRUCT_BODY()
 
 	FGameplayAbilityTargetData_SingleTargetHit()
 	{ }
@@ -670,4 +670,4 @@ struct FAbilityReplicatedDataCache
 		}
 
 	}
-};
+};
\ No newline at end of file
diff --git a/Engine/Source/Runtime/GameplayAbilities/Public/Abilities/GameplayAbilityTypes.h b/Engine/Source/Runtime/GameplayAbilities/Public/Abilities/GameplayAbilityTypes.h
--- a/Engine/Source/Runtime/GameplayAbilities/Public/Abilities/GameplayAbilityTypes.h
+++ b/Engine/Source/Runtime/GameplayAbilities/Public/Abilities/GameplayAbilityTypes.h
@@ -115,7 +115,7 @@ namespace EGameplayAbilityTriggerSource
 USTRUCT(BlueprintType)
 struct GAMEPLAYABILITIES_API FGameplayAbilityActorInfo
 {
-	GENERATED_BODY()
+	GENERATED_USTRUCT_BODY()
 
 	virtual ~FGameplayAbilityActorInfo() {}
 
@@ -164,7 +164,7 @@ struct GAMEPLAYABILITIES_API FGameplayAbilityActorInfo
 USTRUCT()
 struct FGameplayAbilitySpecHandleAndPredictionKey
 {
-	GENERATED_BODY()
+	GENERATED_USTRUCT_BODY()
 
 	FGameplayAbilitySpecHandleAndPredictionKey()
 	: PredictionKeyAtCreation(0)
@@ -206,7 +206,7 @@ private:
 USTRUCT()
 struct GAMEPLAYABILITIES_API FGameplayAbilityRepAnimMontage
 {
-	GENERATED_BODY()
+	GENERATED_USTRUCT_BODY()
 
 	/** AnimMontage ref */
 	UPROPERTY()
@@ -255,7 +255,7 @@ struct GAMEPLAYABILITIES_API FGameplayAbilityRepAnimMontage
 USTRUCT()
 struct GAMEPLAYABILITIES_API FGameplayAbilityLocalAnimMontage
 {
-	GENERATED_BODY()
+	GENERATED_USTRUCT_BODY()
 
 	FGameplayAbilityLocalAnimMontage()
 	: AnimMontage(nullptr), PlayBit(false), AnimatingAbility(nullptr)
@@ -281,7 +281,7 @@ struct GAMEPLAYABILITIES_API FGameplayAbilityLocalAnimMontage
 USTRUCT(BlueprintType)
 struct GAMEPLAYABILITIES_API FGameplayEventData
 {
-	GENERATED_BODY()
+	GENERATED_USTRUCT_BODY()
 
 	FGameplayEventData()
 	: Instigator(nullptr)
@@ -364,7 +364,7 @@ struct FGameplayAbiliyInputBinds
 USTRUCT()
 struct GAMEPLAYABILITIES_API FAttributeDefaults
 {
-	GENERATED_BODY()
+	GENERATED_USTRUCT_BODY()
 
 	UPROPERTY(EditAnywhere, Category = "AttributeTest")
 	TSubclassOf<UAttributeSet> Attributes;
diff --git a/Engine/Source/Runtime/GameplayAbilities/Public/Abilities/GameplayAbilityWorldReticle.h b/Engine/Source/Runtime/GameplayAbilities/Public/Abilities/GameplayAbilityWorldReticle.h
--- a/Engine/Source/Runtime/GameplayAbilities/Public/Abilities/GameplayAbilityWorldReticle.h
+++ b/Engine/Source/Runtime/GameplayAbilities/Public/Abilities/GameplayAbilityWorldReticle.h
@@ -8,7 +8,7 @@ class AGameplayAbilityTargetActor;
 USTRUCT(BlueprintType)
 struct FWorldReticleParameters
 {
-	GENERATED_BODY()
+	GENERATED_USTRUCT_BODY()
 
 	//Use this so that we can't slip in new parameters without some actor knowing about it.
 	void Initialize(FVector InAOEScale)
@@ -24,9 +24,7 @@ struct FWorldReticleParameters
 UCLASS(Blueprintable, notplaceable)
 class GAMEPLAYABILITIES_API AGameplayAbilityWorldReticle : public AActor
 {
-	GENERATED_BODY()
-public:
-	AGameplayAbilityWorldReticle(const FObjectInitializer& ObjectInitializer = FObjectInitializer::Get());
+	GENERATED_UCLASS_BODY()
 
 public:
 
diff --git a/Engine/Source/Runtime/GameplayAbilities/Public/Abilities/GameplayAbilityWorldReticle_ActorVisualization.h b/Engine/Source/Runtime/GameplayAbilities/Public/Abilities/GameplayAbilityWorldReticle_ActorVisualization.h
--- a/Engine/Source/Runtime/GameplayAbilities/Public/Abilities/GameplayAbilityWorldReticle_ActorVisualization.h
+++ b/Engine/Source/Runtime/GameplayAbilities/Public/Abilities/GameplayAbilityWorldReticle_ActorVisualization.h
@@ -11,9 +11,7 @@ class AGameplayAbilityTargetActor;
 UCLASS(notplaceable)
 class AGameplayAbilityWorldReticle_ActorVisualization : public AGameplayAbilityWorldReticle
 {
-	GENERATED_BODY()
-public:
-	AGameplayAbilityWorldReticle_ActorVisualization(const FObjectInitializer& ObjectInitializer = FObjectInitializer::Get());
+	GENERATED_UCLASS_BODY()
 
 public:
 	void InitializeReticleVisualizationInformation(AGameplayAbilityTargetActor* InTargetingActor, AActor* VisualizationActor, UMaterialInterface *VisualizationMaterial);
diff --git a/Engine/Source/Runtime/GameplayAbilities/Public/Abilities/GameplayAbility_CharacterJump.h b/Engine/Source/Runtime/GameplayAbilities/Public/Abilities/GameplayAbility_CharacterJump.h
--- a/Engine/Source/Runtime/GameplayAbilities/Public/Abilities/GameplayAbility_CharacterJump.h
+++ b/Engine/Source/Runtime/GameplayAbilities/Public/Abilities/GameplayAbility_CharacterJump.h
@@ -11,9 +11,7 @@
 UCLASS()
 class GAMEPLAYABILITIES_API UGameplayAbility_CharacterJump : public UGameplayAbility
 {
-	GENERATED_BODY()
-public:
-	UGameplayAbility_CharacterJump(const FObjectInitializer& ObjectInitializer = FObjectInitializer::Get());
+	GENERATED_UCLASS_BODY()
 
 public:
 
@@ -24,4 +22,4 @@ public:
 	virtual void InputReleased(const FGameplayAbilitySpecHandle Handle, const FGameplayAbilityActorInfo* ActorInfo, const FGameplayAbilityActivationInfo ActivationInfo) override;
 
 	virtual void CancelAbility(const FGameplayAbilitySpecHandle Handle, const FGameplayAbilityActorInfo* ActorInfo, const FGameplayAbilityActivationInfo ActivationInfo, bool bReplicateCancelAbility) override;
-};
+};
\ No newline at end of file
diff --git a/Engine/Source/Runtime/GameplayAbilities/Public/Abilities/GameplayAbility_Montage.h b/Engine/Source/Runtime/GameplayAbilities/Public/Abilities/GameplayAbility_Montage.h
--- a/Engine/Source/Runtime/GameplayAbilities/Public/Abilities/GameplayAbility_Montage.h
+++ b/Engine/Source/Runtime/GameplayAbilities/Public/Abilities/GameplayAbility_Montage.h
@@ -11,9 +11,7 @@
 UCLASS()
 class GAMEPLAYABILITIES_API UGameplayAbility_Montage : public UGameplayAbility
 {
-	GENERATED_BODY()
-public:
-	UGameplayAbility_Montage(const FObjectInitializer& ObjectInitializer = FObjectInitializer::Get());
+	GENERATED_UCLASS_BODY()
 
 public:
 	
@@ -41,4 +39,4 @@ public:
 	void GetGameplayEffectsWhileAnimating(TArray<const UGameplayEffect*>& OutEffects) const;
 
 	virtual void ConvertDeprecatedGameplayEffectReferencesToBlueprintReferences(UGameplayEffect* OldGE, TSubclassOf<UGameplayEffect> NewGEClass) override;
-};
+};
\ No newline at end of file
diff --git a/Engine/Source/Runtime/GameplayAbilities/Public/Abilities/Tasks/AbilityTask.h b/Engine/Source/Runtime/GameplayAbilities/Public/Abilities/Tasks/AbilityTask.h
--- a/Engine/Source/Runtime/GameplayAbilities/Public/Abilities/Tasks/AbilityTask.h
+++ b/Engine/Source/Runtime/GameplayAbilities/Public/Abilities/Tasks/AbilityTask.h
@@ -60,9 +60,7 @@
 UCLASS()
 class GAMEPLAYABILITIES_API UAbilityTask : public UObject
 {
-	GENERATED_BODY()
-public:
-	UAbilityTask(const FObjectInitializer& ObjectInitializer = FObjectInitializer::Get());
+	GENERATED_UCLASS_BODY()
 
 	DECLARE_DYNAMIC_MULTICAST_DELEGATE(FGenericAbilityTaskDelegate);
 
diff --git a/Engine/Source/Runtime/GameplayAbilities/Public/Abilities/Tasks/AbilityTask_MoveToLocation.h b/Engine/Source/Runtime/GameplayAbilities/Public/Abilities/Tasks/AbilityTask_MoveToLocation.h
--- a/Engine/Source/Runtime/GameplayAbilities/Public/Abilities/Tasks/AbilityTask_MoveToLocation.h
+++ b/Engine/Source/Runtime/GameplayAbilities/Public/Abilities/Tasks/AbilityTask_MoveToLocation.h
@@ -23,9 +23,7 @@ class UPrimitiveComponent;
 UCLASS(MinimalAPI)
 class UAbilityTask_MoveToLocation : public UAbilityTask
 {
-	GENERATED_BODY()
-public:
-	GAMEPLAYABILITIES_API UAbilityTask_MoveToLocation(const FObjectInitializer& ObjectInitializer = FObjectInitializer::Get());
+	GENERATED_UCLASS_BODY()
 
 	UPROPERTY(BlueprintAssignable)
 	FMoveToLocationDelegate		OnTargetLocationReached;
@@ -67,4 +65,4 @@ protected:
 	UCurveFloat* LerpCurve;
 
 	
-};
+};
\ No newline at end of file
diff --git a/Engine/Source/Runtime/GameplayAbilities/Public/Abilities/Tasks/AbilityTask_NetworkSyncPoint.h b/Engine/Source/Runtime/GameplayAbilities/Public/Abilities/Tasks/AbilityTask_NetworkSyncPoint.h
--- a/Engine/Source/Runtime/GameplayAbilities/Public/Abilities/Tasks/AbilityTask_NetworkSyncPoint.h
+++ b/Engine/Source/Runtime/GameplayAbilities/Public/Abilities/Tasks/AbilityTask_NetworkSyncPoint.h
@@ -24,9 +24,7 @@ enum class EAbilityTaskNetSyncType : uint8
 UCLASS(MinimalAPI)
 class UAbilityTask_NetworkSyncPoint : public UAbilityTask
 {
-	GENERATED_BODY()
-public:
-	GAMEPLAYABILITIES_API UAbilityTask_NetworkSyncPoint(const FObjectInitializer& ObjectInitializer = FObjectInitializer::Get());
+	GENERATED_UCLASS_BODY()
 
 	UPROPERTY(BlueprintAssignable)
 	FNetworkSyncDelegate	OnSync;
@@ -63,4 +61,4 @@ protected:
 	EAbilityGenericReplicatedEvent::Type	ReplicatedEventToListenFor;
 
 	EAbilityTaskNetSyncType SyncType;
-};
+};
\ No newline at end of file
diff --git a/Engine/Source/Runtime/GameplayAbilities/Public/Abilities/Tasks/AbilityTask_PlayMontageAndWait.h b/Engine/Source/Runtime/GameplayAbilities/Public/Abilities/Tasks/AbilityTask_PlayMontageAndWait.h
--- a/Engine/Source/Runtime/GameplayAbilities/Public/Abilities/Tasks/AbilityTask_PlayMontageAndWait.h
+++ b/Engine/Source/Runtime/GameplayAbilities/Public/Abilities/Tasks/AbilityTask_PlayMontageAndWait.h
@@ -8,9 +8,7 @@ DECLARE_DYNAMIC_MULTICAST_DELEGATE(FMontageWaitSimpleDelegate);
 UCLASS(MinimalAPI)
 class UAbilityTask_PlayMontageAndWait : public UAbilityTask
 {
-	GENERATED_BODY()
-public:
-	GAMEPLAYABILITIES_API UAbilityTask_PlayMontageAndWait(const FObjectInitializer& ObjectInitializer = FObjectInitializer::Get());
+	GENERATED_UCLASS_BODY()
 
 	UPROPERTY(BlueprintAssignable)
 	FMontageWaitSimpleDelegate	OnComplete;
@@ -49,4 +47,4 @@ private:
 	UAnimMontage* MontageToPlay;
 	float Rate;	
 	FName StartSection;
-};
+};
\ No newline at end of file
diff --git a/Engine/Source/Runtime/GameplayAbilities/Public/Abilities/Tasks/AbilityTask_Repeat.h b/Engine/Source/Runtime/GameplayAbilities/Public/Abilities/Tasks/AbilityTask_Repeat.h
--- a/Engine/Source/Runtime/GameplayAbilities/Public/Abilities/Tasks/AbilityTask_Repeat.h
+++ b/Engine/Source/Runtime/GameplayAbilities/Public/Abilities/Tasks/AbilityTask_Repeat.h
@@ -13,9 +13,7 @@ DECLARE_DYNAMIC_MULTICAST_DELEGATE_OneParam(FRepeatedActionDelegate, int32, Acti
 UCLASS()
 class GAMEPLAYABILITIES_API UAbilityTask_Repeat : public UAbilityTask
 {
-	GENERATED_BODY()
-public:
-	UAbilityTask_Repeat(const FObjectInitializer& ObjectInitializer = FObjectInitializer::Get());
+	GENERATED_UCLASS_BODY()
 
 	UPROPERTY(BlueprintAssignable)
 	FRepeatedActionDelegate		OnPerformAction;
@@ -42,4 +40,4 @@ protected:
 	FTimerHandle TimerHandle_PerformAction;
 
 	virtual void OnDestroy(bool AbilityIsEnding) override;
-};
+};
\ No newline at end of file
diff --git a/Engine/Source/Runtime/GameplayAbilities/Public/Abilities/Tasks/AbilityTask_SpawnActor.h b/Engine/Source/Runtime/GameplayAbilities/Public/Abilities/Tasks/AbilityTask_SpawnActor.h
--- a/Engine/Source/Runtime/GameplayAbilities/Public/Abilities/Tasks/AbilityTask_SpawnActor.h
+++ b/Engine/Source/Runtime/GameplayAbilities/Public/Abilities/Tasks/AbilityTask_SpawnActor.h
@@ -23,9 +23,7 @@ DECLARE_DYNAMIC_MULTICAST_DELEGATE_OneParam(FSpawnActorDelegate, AActor*, Spawne
 UCLASS(MinimalAPI)
 class UAbilityTask_SpawnActor: public UAbilityTask
 {
-	GENERATED_BODY()
-public:
-	GAMEPLAYABILITIES_API UAbilityTask_SpawnActor(const FObjectInitializer& ObjectInitializer = FObjectInitializer::Get());
+	GENERATED_UCLASS_BODY()
 
 	UPROPERTY(BlueprintAssignable)
 	FSpawnActorDelegate	Success;
@@ -46,4 +44,4 @@ public:
 
 protected:
 	FGameplayAbilityTargetDataHandle CachedTargetDataHandle;
-};
+};
\ No newline at end of file
diff --git a/Engine/Source/Runtime/GameplayAbilities/Public/Abilities/Tasks/AbilityTask_StartAbilityState.h b/Engine/Source/Runtime/GameplayAbilities/Public/Abilities/Tasks/AbilityTask_StartAbilityState.h
--- a/Engine/Source/Runtime/GameplayAbilities/Public/Abilities/Tasks/AbilityTask_StartAbilityState.h
+++ b/Engine/Source/Runtime/GameplayAbilities/Public/Abilities/Tasks/AbilityTask_StartAbilityState.h
@@ -25,9 +25,7 @@ DECLARE_DYNAMIC_MULTICAST_DELEGATE(FAbilityStateDelegate);
 UCLASS(MinimalAPI)
 class UAbilityTask_StartAbilityState : public UAbilityTask
 {
-	GENERATED_BODY()
-public:
-	GAMEPLAYABILITIES_API UAbilityTask_StartAbilityState(const FObjectInitializer& ObjectInitializer = FObjectInitializer::Get());
+	GENERATED_UCLASS_BODY()
 
 	/** Invoked if 'EndAbilityState' is called or if 'EndAbility' is called and this state is active. */
 	UPROPERTY(BlueprintAssignable)
diff --git a/Engine/Source/Runtime/GameplayAbilities/Public/Abilities/Tasks/AbilityTask_VisualizeTargeting.h b/Engine/Source/Runtime/GameplayAbilities/Public/Abilities/Tasks/AbilityTask_VisualizeTargeting.h
--- a/Engine/Source/Runtime/GameplayAbilities/Public/Abilities/Tasks/AbilityTask_VisualizeTargeting.h
+++ b/Engine/Source/Runtime/GameplayAbilities/Public/Abilities/Tasks/AbilityTask_VisualizeTargeting.h
@@ -12,9 +12,7 @@ DECLARE_DYNAMIC_MULTICAST_DELEGATE(FVisualizeTargetingDelegate);
 UCLASS(MinimalAPI)
 class UAbilityTask_VisualizeTargeting: public UAbilityTask
 {
-	GENERATED_BODY()
-public:
-	GAMEPLAYABILITIES_API UAbilityTask_VisualizeTargeting(const FObjectInitializer& ObjectInitializer = FObjectInitializer::Get());
+	GENERATED_UCLASS_BODY()
 
 	UPROPERTY(BlueprintAssignable)
 	FVisualizeTargetingDelegate TimeElapsed;
@@ -66,4 +64,4 @@ protected:
 *			-This function should spawn the actor with SpawnActorDeferred and return true/false if it spawned something.
 *		-Have a function named FinishSpawningActor w/ an AActor* of the class you spawned
 *			-This function *must* call ExecuteConstruction + PostActorConstruction
-*/
+*/
\ No newline at end of file
diff --git a/Engine/Source/Runtime/GameplayAbilities/Public/Abilities/Tasks/AbilityTask_WaitAbilityActivate.h b/Engine/Source/Runtime/GameplayAbilities/Public/Abilities/Tasks/AbilityTask_WaitAbilityActivate.h
--- a/Engine/Source/Runtime/GameplayAbilities/Public/Abilities/Tasks/AbilityTask_WaitAbilityActivate.h
+++ b/Engine/Source/Runtime/GameplayAbilities/Public/Abilities/Tasks/AbilityTask_WaitAbilityActivate.h
@@ -14,9 +14,7 @@ class AActor;
 UCLASS(MinimalAPI)
 class UAbilityTask_WaitAbilityActivate : public UAbilityTask
 {
-	GENERATED_BODY()
-public:
-	GAMEPLAYABILITIES_API UAbilityTask_WaitAbilityActivate(const FObjectInitializer& ObjectInitializer = FObjectInitializer::Get());
+	GENERATED_UCLASS_BODY()
 
 	UPROPERTY(BlueprintAssignable)
 	FWaitAbilityActivateDelegate	OnActivate;
@@ -39,4 +37,4 @@ protected:
 	virtual void OnDestroy(bool AbilityEnded) override;
 
 	FDelegateHandle OnAbilityActivateDelegateHandle;
-};
+};
\ No newline at end of file
diff --git a/Engine/Source/Runtime/GameplayAbilities/Public/Abilities/Tasks/AbilityTask_WaitAbilityCommit.h b/Engine/Source/Runtime/GameplayAbilities/Public/Abilities/Tasks/AbilityTask_WaitAbilityCommit.h
--- a/Engine/Source/Runtime/GameplayAbilities/Public/Abilities/Tasks/AbilityTask_WaitAbilityCommit.h
+++ b/Engine/Source/Runtime/GameplayAbilities/Public/Abilities/Tasks/AbilityTask_WaitAbilityCommit.h
@@ -15,9 +15,7 @@ DECLARE_DYNAMIC_MULTICAST_DELEGATE_OneParam(FWaitAbilityCommitDelegate, UGamepla
 UCLASS(MinimalAPI)
 class UAbilityTask_WaitAbilityCommit : public UAbilityTask
 {
-	GENERATED_BODY()
-public:
-	GAMEPLAYABILITIES_API UAbilityTask_WaitAbilityCommit(const FObjectInitializer& ObjectInitializer = FObjectInitializer::Get());
+	GENERATED_UCLASS_BODY()
 
 	UPROPERTY(BlueprintAssignable)
 	FWaitAbilityCommitDelegate	OnCommit;
@@ -39,4 +37,4 @@ protected:
 	virtual void OnDestroy(bool AbilityEnded) override;
 
 	FDelegateHandle OnAbilityCommitDelegateHandle;
-};
+};
\ No newline at end of file
diff --git a/Engine/Source/Runtime/GameplayAbilities/Public/Abilities/Tasks/AbilityTask_WaitAttributeChange.h b/Engine/Source/Runtime/GameplayAbilities/Public/Abilities/Tasks/AbilityTask_WaitAttributeChange.h
--- a/Engine/Source/Runtime/GameplayAbilities/Public/Abilities/Tasks/AbilityTask_WaitAttributeChange.h
+++ b/Engine/Source/Runtime/GameplayAbilities/Public/Abilities/Tasks/AbilityTask_WaitAttributeChange.h
@@ -31,9 +31,7 @@ DECLARE_DYNAMIC_MULTICAST_DELEGATE(FWaitAttributeChangeDelegate);
 UCLASS(MinimalAPI)
 class UAbilityTask_WaitAttributeChange : public UAbilityTask
 {
-	GENERATED_BODY()
-public:
-	GAMEPLAYABILITIES_API UAbilityTask_WaitAttributeChange(const FObjectInitializer& ObjectInitializer = FObjectInitializer::Get());
+	GENERATED_UCLASS_BODY()
 
 	UPROPERTY(BlueprintAssignable)
 	FWaitAttributeChangeDelegate	OnChange;
diff --git a/Engine/Source/Runtime/GameplayAbilities/Public/Abilities/Tasks/AbilityTask_WaitConfirm.h b/Engine/Source/Runtime/GameplayAbilities/Public/Abilities/Tasks/AbilityTask_WaitConfirm.h
--- a/Engine/Source/Runtime/GameplayAbilities/Public/Abilities/Tasks/AbilityTask_WaitConfirm.h
+++ b/Engine/Source/Runtime/GameplayAbilities/Public/Abilities/Tasks/AbilityTask_WaitConfirm.h
@@ -7,9 +7,7 @@
 UCLASS(MinimalAPI)
 class UAbilityTask_WaitConfirm : public UAbilityTask
 {
-	GENERATED_BODY()
-public:
-	GAMEPLAYABILITIES_API UAbilityTask_WaitConfirm(const FObjectInitializer& ObjectInitializer = FObjectInitializer::Get());
+	GENERATED_UCLASS_BODY()
 
 	UPROPERTY(BlueprintAssignable)
 	FGenericAbilityTaskDelegate	OnConfirm;
@@ -30,4 +28,4 @@ protected:
 	bool RegisteredCallback;
 
 	FDelegateHandle OnConfirmCallbackDelegateHandle;
-};
+};
\ No newline at end of file
diff --git a/Engine/Source/Runtime/GameplayAbilities/Public/Abilities/Tasks/AbilityTask_WaitConfirmCancel.h b/Engine/Source/Runtime/GameplayAbilities/Public/Abilities/Tasks/AbilityTask_WaitConfirmCancel.h
--- a/Engine/Source/Runtime/GameplayAbilities/Public/Abilities/Tasks/AbilityTask_WaitConfirmCancel.h
+++ b/Engine/Source/Runtime/GameplayAbilities/Public/Abilities/Tasks/AbilityTask_WaitConfirmCancel.h
@@ -14,9 +14,7 @@ DECLARE_DYNAMIC_MULTICAST_DELEGATE(FWaitConfirmCancelDelegate);
 UCLASS(MinimalAPI)
 class UAbilityTask_WaitConfirmCancel : public UAbilityTask
 {
-	GENERATED_BODY()
-public:
-	GAMEPLAYABILITIES_API UAbilityTask_WaitConfirmCancel(const FObjectInitializer& ObjectInitializer = FObjectInitializer::Get());
+	GENERATED_UCLASS_BODY()
 
 	UPROPERTY(BlueprintAssignable)
 	FWaitConfirmCancelDelegate	OnConfirm;
@@ -47,4 +45,4 @@ protected:
 
 	bool RegisteredCallbacks;
 	
-};
+};
\ No newline at end of file
diff --git a/Engine/Source/Runtime/GameplayAbilities/Public/Abilities/Tasks/AbilityTask_WaitDelay.h b/Engine/Source/Runtime/GameplayAbilities/Public/Abilities/Tasks/AbilityTask_WaitDelay.h
--- a/Engine/Source/Runtime/GameplayAbilities/Public/Abilities/Tasks/AbilityTask_WaitDelay.h
+++ b/Engine/Source/Runtime/GameplayAbilities/Public/Abilities/Tasks/AbilityTask_WaitDelay.h
@@ -9,9 +9,7 @@ DECLARE_DYNAMIC_MULTICAST_DELEGATE(FWaitDelayDelegate);
 UCLASS(MinimalAPI)
 class UAbilityTask_WaitDelay : public UAbilityTask
 {
-	GENERATED_BODY()
-public:
-	GAMEPLAYABILITIES_API UAbilityTask_WaitDelay(const FObjectInitializer& ObjectInitializer = FObjectInitializer::Get());
+	GENERATED_UCLASS_BODY()
 
 	UPROPERTY(BlueprintAssignable)
 	FWaitDelayDelegate	OnFinish;
@@ -32,4 +30,4 @@ private:
 	virtual void OnDestroy(bool AbilityEnded) override;
 	float Time;
 	float TimeStarted;
-};
+};
\ No newline at end of file
diff --git a/Engine/Source/Runtime/GameplayAbilities/Public/Abilities/Tasks/AbilityTask_WaitGameplayEffectApplied.h b/Engine/Source/Runtime/GameplayAbilities/Public/Abilities/Tasks/AbilityTask_WaitGameplayEffectApplied.h
--- a/Engine/Source/Runtime/GameplayAbilities/Public/Abilities/Tasks/AbilityTask_WaitGameplayEffectApplied.h
+++ b/Engine/Source/Runtime/GameplayAbilities/Public/Abilities/Tasks/AbilityTask_WaitGameplayEffectApplied.h
@@ -9,9 +9,7 @@
 UCLASS(MinimalAPI)
 class UAbilityTask_WaitGameplayEffectApplied : public UAbilityTask
 {
-	GENERATED_BODY()
-public:
-	GAMEPLAYABILITIES_API UAbilityTask_WaitGameplayEffectApplied(const FObjectInitializer& ObjectInitializer = FObjectInitializer::Get());
+	GENERATED_UCLASS_BODY()
 
 	UFUNCTION()
 	void OnApplyGameplayEffectCallback(UAbilitySystemComponent* Target, const FGameplayEffectSpec& SpecApplied, FActiveGameplayEffectHandle ActiveHandle);
@@ -39,4 +37,4 @@ protected:
 	bool UseExternalOwner;
 
 	TWeakObjectPtr<UAbilitySystemComponent> ExternalOwner;
-};
+};
\ No newline at end of file
diff --git a/Engine/Source/Runtime/GameplayAbilities/Public/Abilities/Tasks/AbilityTask_WaitGameplayEffectApplied_Self.h b/Engine/Source/Runtime/GameplayAbilities/Public/Abilities/Tasks/AbilityTask_WaitGameplayEffectApplied_Self.h
--- a/Engine/Source/Runtime/GameplayAbilities/Public/Abilities/Tasks/AbilityTask_WaitGameplayEffectApplied_Self.h
+++ b/Engine/Source/Runtime/GameplayAbilities/Public/Abilities/Tasks/AbilityTask_WaitGameplayEffectApplied_Self.h
@@ -10,9 +10,7 @@ DECLARE_DYNAMIC_MULTICAST_DELEGATE_ThreeParams(FGameplayEffectAppliedSelfDelegat
 UCLASS(MinimalAPI)
 class UAbilityTask_WaitGameplayEffectApplied_Self : public UAbilityTask_WaitGameplayEffectApplied
 {
-	GENERATED_BODY()
-public:
-	GAMEPLAYABILITIES_API UAbilityTask_WaitGameplayEffectApplied_Self(const FObjectInitializer& ObjectInitializer = FObjectInitializer::Get());
+	GENERATED_UCLASS_BODY()
 
 	UPROPERTY(BlueprintAssignable)
 	FGameplayEffectAppliedSelfDelegate OnApplied;
@@ -32,4 +30,4 @@ protected:
 
 private:
 	FDelegateHandle OnApplyGameplayEffectCallbackDelegateHandle;
-};
+};
\ No newline at end of file
diff --git a/Engine/Source/Runtime/GameplayAbilities/Public/Abilities/Tasks/AbilityTask_WaitGameplayEffectApplied_Target.h b/Engine/Source/Runtime/GameplayAbilities/Public/Abilities/Tasks/AbilityTask_WaitGameplayEffectApplied_Target.h
--- a/Engine/Source/Runtime/GameplayAbilities/Public/Abilities/Tasks/AbilityTask_WaitGameplayEffectApplied_Target.h
+++ b/Engine/Source/Runtime/GameplayAbilities/Public/Abilities/Tasks/AbilityTask_WaitGameplayEffectApplied_Target.h
@@ -10,9 +10,7 @@ DECLARE_DYNAMIC_MULTICAST_DELEGATE_ThreeParams(FGameplayEffectAppliedTargetDeleg
 UCLASS(MinimalAPI)
 class UAbilityTask_WaitGameplayEffectApplied_Target : public UAbilityTask_WaitGameplayEffectApplied
 {
-	GENERATED_BODY()
-public:
-	GAMEPLAYABILITIES_API UAbilityTask_WaitGameplayEffectApplied_Target(const FObjectInitializer& ObjectInitializer = FObjectInitializer::Get());
+	GENERATED_UCLASS_BODY()
 
 	UPROPERTY(BlueprintAssignable)
 	FGameplayEffectAppliedTargetDelegate OnApplied;
@@ -32,4 +30,4 @@ protected:
 
 private:
 	FDelegateHandle OnApplyGameplayEffectCallbackDelegateHandle;
-};
+};
\ No newline at end of file
diff --git a/Engine/Source/Runtime/GameplayAbilities/Public/Abilities/Tasks/AbilityTask_WaitGameplayEffectRemoved.h b/Engine/Source/Runtime/GameplayAbilities/Public/Abilities/Tasks/AbilityTask_WaitGameplayEffectRemoved.h
--- a/Engine/Source/Runtime/GameplayAbilities/Public/Abilities/Tasks/AbilityTask_WaitGameplayEffectRemoved.h
+++ b/Engine/Source/Runtime/GameplayAbilities/Public/Abilities/Tasks/AbilityTask_WaitGameplayEffectRemoved.h
@@ -17,9 +17,7 @@ class UPrimitiveComponent;
 UCLASS(MinimalAPI)
 class UAbilityTask_WaitGameplayEffectRemoved : public UAbilityTask
 {
-	GENERATED_BODY()
-public:
-	GAMEPLAYABILITIES_API UAbilityTask_WaitGameplayEffectRemoved(const FObjectInitializer& ObjectInitializer = FObjectInitializer::Get());
+	GENERATED_UCLASS_BODY()
 
 	UPROPERTY(BlueprintAssignable)
 	FWaitGameplayEffectRemovedDelegate	OnRemoved;
@@ -44,4 +42,4 @@ protected:
 	bool Registered;
 
 	FDelegateHandle OnGameplayEffectRemovedDelegateHandle;
-};
+};
\ No newline at end of file
diff --git a/Engine/Source/Runtime/GameplayAbilities/Public/Abilities/Tasks/AbilityTask_WaitGameplayTag.h b/Engine/Source/Runtime/GameplayAbilities/Public/Abilities/Tasks/AbilityTask_WaitGameplayTag.h
--- a/Engine/Source/Runtime/GameplayAbilities/Public/Abilities/Tasks/AbilityTask_WaitGameplayTag.h
+++ b/Engine/Source/Runtime/GameplayAbilities/Public/Abilities/Tasks/AbilityTask_WaitGameplayTag.h
@@ -11,9 +11,7 @@ DECLARE_DYNAMIC_MULTICAST_DELEGATE(FWaitGameplayTagDelegate);
 UCLASS(MinimalAPI)
 class UAbilityTask_WaitGameplayTagAdded : public UAbilityTask_WaitGameplayTag
 {
-	GENERATED_BODY()
-public:
-	GAMEPLAYABILITIES_API UAbilityTask_WaitGameplayTagAdded(const FObjectInitializer& ObjectInitializer = FObjectInitializer::Get());
+	GENERATED_UCLASS_BODY()
 
 	UPROPERTY(BlueprintAssignable)
 	FWaitGameplayTagDelegate	Added;
@@ -33,9 +31,7 @@ public:
 UCLASS(MinimalAPI)
 class UAbilityTask_WaitGameplayTagRemoved : public UAbilityTask_WaitGameplayTag
 {
-	GENERATED_BODY()
-public:
-	GAMEPLAYABILITIES_API UAbilityTask_WaitGameplayTagRemoved(const FObjectInitializer& ObjectInitializer = FObjectInitializer::Get());
+	GENERATED_UCLASS_BODY()
 
 	UPROPERTY(BlueprintAssignable)
 	FWaitGameplayTagDelegate	Removed;
@@ -50,4 +46,4 @@ public:
 	virtual void Activate() override;
 
 	virtual void GameplayTagCallback(const FGameplayTag Tag, int32 NewCount) override;
-};
+};
\ No newline at end of file
diff --git a/Engine/Source/Runtime/GameplayAbilities/Public/Abilities/Tasks/AbilityTask_WaitGameplayTagBase.h b/Engine/Source/Runtime/GameplayAbilities/Public/Abilities/Tasks/AbilityTask_WaitGameplayTagBase.h
--- a/Engine/Source/Runtime/GameplayAbilities/Public/Abilities/Tasks/AbilityTask_WaitGameplayTagBase.h
+++ b/Engine/Source/Runtime/GameplayAbilities/Public/Abilities/Tasks/AbilityTask_WaitGameplayTagBase.h
@@ -8,9 +8,7 @@
 UCLASS(MinimalAPI)
 class UAbilityTask_WaitGameplayTag : public UAbilityTask
 {
-	GENERATED_BODY()
-public:
-	GAMEPLAYABILITIES_API UAbilityTask_WaitGameplayTag(const FObjectInitializer& ObjectInitializer = FObjectInitializer::Get());
+	GENERATED_UCLASS_BODY()
 
 	virtual void Activate() override;
 
@@ -34,4 +32,4 @@ protected:
 	bool OnlyTriggerOnce;
 
 	FDelegateHandle DelegateHandle;
-};
+};
\ No newline at end of file
diff --git a/Engine/Source/Runtime/GameplayAbilities/Public/Abilities/Tasks/AbilityTask_WaitInputPress.h b/Engine/Source/Runtime/GameplayAbilities/Public/Abilities/Tasks/AbilityTask_WaitInputPress.h
--- a/Engine/Source/Runtime/GameplayAbilities/Public/Abilities/Tasks/AbilityTask_WaitInputPress.h
+++ b/Engine/Source/Runtime/GameplayAbilities/Public/Abilities/Tasks/AbilityTask_WaitInputPress.h
@@ -14,9 +14,7 @@ DECLARE_DYNAMIC_MULTICAST_DELEGATE_OneParam(FInputPressDelegate, float, TimeWait
 UCLASS(MinimalAPI)
 class UAbilityTask_WaitInputPress : public UAbilityTask
 {
-	GENERATED_BODY()
-public:
-	GAMEPLAYABILITIES_API UAbilityTask_WaitInputPress(const FObjectInitializer& ObjectInitializer = FObjectInitializer::Get());
+	GENERATED_UCLASS_BODY()
 
 	UPROPERTY(BlueprintAssignable)
 	FInputPressDelegate		OnPress;
@@ -37,4 +35,4 @@ protected:
 	float StartTime;
 	bool bTestInitialState;
 	FDelegateHandle DelegateHandle;
-};
+};
\ No newline at end of file
diff --git a/Engine/Source/Runtime/GameplayAbilities/Public/Abilities/Tasks/AbilityTask_WaitInputRelease.h b/Engine/Source/Runtime/GameplayAbilities/Public/Abilities/Tasks/AbilityTask_WaitInputRelease.h
--- a/Engine/Source/Runtime/GameplayAbilities/Public/Abilities/Tasks/AbilityTask_WaitInputRelease.h
+++ b/Engine/Source/Runtime/GameplayAbilities/Public/Abilities/Tasks/AbilityTask_WaitInputRelease.h
@@ -14,9 +14,7 @@ DECLARE_DYNAMIC_MULTICAST_DELEGATE_OneParam(FInputReleaseDelegate, float, TimeHe
 UCLASS(MinimalAPI)
 class UAbilityTask_WaitInputRelease : public UAbilityTask
 {
-	GENERATED_BODY()
-public:
-	GAMEPLAYABILITIES_API UAbilityTask_WaitInputRelease(const FObjectInitializer& ObjectInitializer = FObjectInitializer::Get());
+	GENERATED_UCLASS_BODY()
 
 	UPROPERTY(BlueprintAssignable)
 	FInputReleaseDelegate	OnRelease;
@@ -37,4 +35,4 @@ protected:
 	float StartTime;
 	bool bTestInitialState;
 	FDelegateHandle DelegateHandle;
-};
+};
\ No newline at end of file
diff --git a/Engine/Source/Runtime/GameplayAbilities/Public/Abilities/Tasks/AbilityTask_WaitMovementModeChange.h b/Engine/Source/Runtime/GameplayAbilities/Public/Abilities/Tasks/AbilityTask_WaitMovementModeChange.h
--- a/Engine/Source/Runtime/GameplayAbilities/Public/Abilities/Tasks/AbilityTask_WaitMovementModeChange.h
+++ b/Engine/Source/Runtime/GameplayAbilities/Public/Abilities/Tasks/AbilityTask_WaitMovementModeChange.h
@@ -10,9 +10,7 @@ class ACharacter;
 UCLASS(MinimalAPI)
 class UAbilityTask_WaitMovementModeChange : public UAbilityTask
 {
-	GENERATED_BODY()
-public:
-	GAMEPLAYABILITIES_API UAbilityTask_WaitMovementModeChange(const FObjectInitializer& ObjectInitializer = FObjectInitializer::Get());
+	GENERATED_UCLASS_BODY()
 
 public:
 
@@ -34,4 +32,4 @@ private:
 	virtual void OnDestroy(bool AbilityEnded) override;
 
 	TWeakObjectPtr<ACharacter>	MyCharacter;
-};
+};
\ No newline at end of file
diff --git a/Engine/Source/Runtime/GameplayAbilities/Public/Abilities/Tasks/AbilityTask_WaitOverlap.h b/Engine/Source/Runtime/GameplayAbilities/Public/Abilities/Tasks/AbilityTask_WaitOverlap.h
--- a/Engine/Source/Runtime/GameplayAbilities/Public/Abilities/Tasks/AbilityTask_WaitOverlap.h
+++ b/Engine/Source/Runtime/GameplayAbilities/Public/Abilities/Tasks/AbilityTask_WaitOverlap.h
@@ -17,9 +17,7 @@ class UPrimitiveComponent;
 UCLASS(MinimalAPI)
 class UAbilityTask_WaitOverlap : public UAbilityTask
 {
-	GENERATED_BODY()
-public:
-	GAMEPLAYABILITIES_API UAbilityTask_WaitOverlap(const FObjectInitializer& ObjectInitializer = FObjectInitializer::Get());
+	GENERATED_UCLASS_BODY()
 
 	UPROPERTY(BlueprintAssignable)
 	FWaitOverlapDelegate	OnOverlap;
@@ -42,4 +40,4 @@ private:
 
 	UPrimitiveComponent* GetComponent();
 	
-};
+};
\ No newline at end of file
diff --git a/Engine/Source/Runtime/GameplayAbilities/Public/Abilities/Tasks/AbilityTask_WaitTargetData.h b/Engine/Source/Runtime/GameplayAbilities/Public/Abilities/Tasks/AbilityTask_WaitTargetData.h
--- a/Engine/Source/Runtime/GameplayAbilities/Public/Abilities/Tasks/AbilityTask_WaitTargetData.h
+++ b/Engine/Source/Runtime/GameplayAbilities/Public/Abilities/Tasks/AbilityTask_WaitTargetData.h
@@ -13,9 +13,7 @@ DECLARE_DYNAMIC_MULTICAST_DELEGATE_OneParam(FWaitTargetDataDelegate, FGameplayAb
 UCLASS(notplaceable)
 class GAMEPLAYABILITIES_API UAbilityTask_WaitTargetData: public UAbilityTask
 {
-	GENERATED_BODY()
-public:
-	UAbilityTask_WaitTargetData(const FObjectInitializer& ObjectInitializer = FObjectInitializer::Get());
+	GENERATED_UCLASS_BODY()
 
 	UPROPERTY(BlueprintAssignable)
 	FWaitTargetDataDelegate	ValidData;
diff --git a/Engine/Source/Runtime/GameplayAbilities/Public/Abilities/Tasks/AbilityTask_WaitVelocityChange.h b/Engine/Source/Runtime/GameplayAbilities/Public/Abilities/Tasks/AbilityTask_WaitVelocityChange.h
--- a/Engine/Source/Runtime/GameplayAbilities/Public/Abilities/Tasks/AbilityTask_WaitVelocityChange.h
+++ b/Engine/Source/Runtime/GameplayAbilities/Public/Abilities/Tasks/AbilityTask_WaitVelocityChange.h
@@ -9,9 +9,7 @@ DECLARE_DYNAMIC_MULTICAST_DELEGATE(FWaitVelocityChangeDelegate);
 UCLASS(MinimalAPI)
 class UAbilityTask_WaitVelocityChange: public UAbilityTask
 {
-	GENERATED_BODY()
-public:
-	GAMEPLAYABILITIES_API UAbilityTask_WaitVelocityChange(const FObjectInitializer& ObjectInitializer = FObjectInitializer::Get());
+	GENERATED_UCLASS_BODY()
 
 	/** Delegate called when velocity requirements are met */
 	UPROPERTY(BlueprintAssignable)
@@ -32,4 +30,4 @@ protected:
 
 	float	MinimumMagnitude;
 	FVector Direction;
-};
+};
\ No newline at end of file
diff --git a/Engine/Source/Runtime/GameplayAbilities/Public/AbilitySystemBlueprintLibrary.h b/Engine/Source/Runtime/GameplayAbilities/Public/AbilitySystemBlueprintLibrary.h
--- a/Engine/Source/Runtime/GameplayAbilities/Public/AbilitySystemBlueprintLibrary.h
+++ b/Engine/Source/Runtime/GameplayAbilities/Public/AbilitySystemBlueprintLibrary.h
@@ -17,9 +17,7 @@ class UAbilityTask_WaitConfirmCancel;
 UCLASS()
 class GAMEPLAYABILITIES_API UAbilitySystemBlueprintLibrary : public UBlueprintFunctionLibrary
 {
-	GENERATED_BODY()
-public:
-	UAbilitySystemBlueprintLibrary(const FObjectInitializer& ObjectInitializer = FObjectInitializer::Get());
+	GENERATED_UCLASS_BODY()
 
 	UFUNCTION(BlueprintPure, Category = Ability)
 	static UAbilitySystemComponent* GetAbilitySystemComponent(AActor *Actor);
diff --git a/Engine/Source/Runtime/GameplayAbilities/Public/AbilitySystemComponent.h b/Engine/Source/Runtime/GameplayAbilities/Public/AbilitySystemComponent.h
--- a/Engine/Source/Runtime/GameplayAbilities/Public/AbilitySystemComponent.h
+++ b/Engine/Source/Runtime/GameplayAbilities/Public/AbilitySystemComponent.h
@@ -51,9 +51,7 @@ DECLARE_MULTICAST_DELEGATE_TwoParams(FAbilityFailedDelegate, const UGameplayAbil
 UCLASS(ClassGroup=AbilitySystem, hidecategories=(Object,LOD,Lighting,Transform,Sockets,TextureStreaming), editinlinenew, meta=(BlueprintSpawnableComponent))
 class GAMEPLAYABILITIES_API UAbilitySystemComponent : public UActorComponent, public IGameplayTagAssetInterface
 {
-	GENERATED_BODY()
-public:
-	UAbilitySystemComponent(const FObjectInitializer& ObjectInitializer = FObjectInitializer::Get());
+	GENERATED_UCLASS_BODY()
 
 	/** Used to register callbacks to confirm/cancel input */
 	DECLARE_DYNAMIC_MULTICAST_DELEGATE(FAbilityConfirmOrCancel);
@@ -459,19 +457,15 @@ public:
 	// Do not call these functions directly, call the wrappers on GameplayCueManager instead
 	UFUNCTION(NetMulticast, unreliable)
 	void NetMulticast_InvokeGameplayCueExecuted_FromSpec(const FGameplayEffectSpecForRPC Spec, FPredictionKey PredictionKey);
-	virtual void NetMulticast_InvokeGameplayCueExecuted_FromSpec_Implementation(const FGameplayEffectSpecForRPC Spec, FPredictionKey PredictionKey);
 
 	UFUNCTION(NetMulticast, unreliable)
 	void NetMulticast_InvokeGameplayCueExecuted(const FGameplayTag GameplayCueTag, FPredictionKey PredictionKey, FGameplayEffectContextHandle EffectContext);
-	virtual void NetMulticast_InvokeGameplayCueExecuted_Implementation(const FGameplayTag GameplayCueTag, FPredictionKey PredictionKey, FGameplayEffectContextHandle EffectContext);
 
 	UFUNCTION(NetMulticast, unreliable)
 	void NetMulticast_InvokeGameplayCueExecuted_WithParams(const FGameplayTag GameplayCueTag, FPredictionKey PredictionKey, FGameplayCueParameters GameplayCueParameters);
-	virtual void NetMulticast_InvokeGameplayCueExecuted_WithParams_Implementation(const FGameplayTag GameplayCueTag, FPredictionKey PredictionKey, FGameplayCueParameters GameplayCueParameters);
 
 	UFUNCTION(NetMulticast, unreliable)
 	void NetMulticast_InvokeGameplayCueAdded(const FGameplayTag GameplayCueTag, FPredictionKey PredictionKey, FGameplayEffectContextHandle EffectContext);
-	virtual void NetMulticast_InvokeGameplayCueAdded_Implementation(const FGameplayTag GameplayCueTag, FPredictionKey PredictionKey, FGameplayEffectContextHandle EffectContext);
 
 	// GameplayCues can also come on their own. These take an optional effect context to pass through hit result, etc
 	void ExecuteGameplayCue(const FGameplayTag GameplayCueTag, FGameplayEffectContextHandle EffectContext = FGameplayEffectContextHandle());
@@ -656,15 +650,11 @@ public:
 	UFUNCTION()
 	void	OnRep_ActivateAbilities();
 
-	UFUNCTION(Server="ServerTryActivateAbility_Implementation", reliable, WithValidation="ServerTryActivateAbility_Validate")
+	UFUNCTION(Server, reliable, WithValidation)
 	void	ServerTryActivateAbility(FGameplayAbilitySpecHandle AbilityToActivate, FPredictionKey PredictionKey);
-	virtual void ServerTryActivateAbility_Implementation(FGameplayAbilitySpecHandle AbilityToActivate, FPredictionKey PredictionKey);
-	virtual bool ServerTryActivateAbility_Validate(FGameplayAbilitySpecHandle , FPredictionKey );
 
-	UFUNCTION(Server="ServerTryActivateAbilityWithEventData_Implementation", reliable, WithValidation="ServerTryActivateAbilityWithEventData_Validate")
+	UFUNCTION(Server, reliable, WithValidation)
 	void	ServerTryActivateAbilityWithEventData(FGameplayAbilitySpecHandle AbilityToActivate, FPredictionKey PredictionKey, FGameplayEventData TriggerEventData);
-	virtual void ServerTryActivateAbilityWithEventData_Implementation(FGameplayAbilitySpecHandle AbilityToActivate, FPredictionKey PredictionKey, FGameplayEventData TriggerEventData);
-	virtual bool ServerTryActivateAbilityWithEventData_Validate(FGameplayAbilitySpecHandle , FPredictionKey , FGameplayEventData );
 
 	/** Replicate that an ability has ended/canceled, to the client or server as appropriate */
 	void	ReplicateEndOrCancelAbility(FGameplayAbilitySpecHandle Handle, FGameplayAbilityActivationInfo ActivationInfo, UGameplayAbility* Ability, bool bWasCanceled);
@@ -672,37 +662,28 @@ public:
 	/** Called by ServerEndAbility/ServerCancelAbility and ClientEndAbility/ClientCancelAbility; avoids code duplication. */
 	void	RemoteEndOrCancelAbility(FGameplayAbilitySpecHandle AbilityToEnd, FGameplayAbilityActivationInfo ActivationInfo, bool bWasCanceled);
 
-	UFUNCTION(Server="ServerEndAbility_Implementation", reliable, WithValidation="ServerEndAbility_Validate")
+	UFUNCTION(Server, reliable, WithValidation)
 	void	ServerEndAbility(FGameplayAbilitySpecHandle AbilityToEnd, FGameplayAbilityActivationInfo ActivationInfo, FPredictionKey PredictionKey);
-	virtual void ServerEndAbility_Implementation(FGameplayAbilitySpecHandle AbilityToEnd, FGameplayAbilityActivationInfo ActivationInfo, FPredictionKey PredictionKey);
-	virtual bool ServerEndAbility_Validate(FGameplayAbilitySpecHandle , FGameplayAbilityActivationInfo , FPredictionKey );
 
-	UFUNCTION(Client="ClientEndAbility_Implementation", reliable)
+	UFUNCTION(Client, reliable)
 	void	ClientEndAbility(FGameplayAbilitySpecHandle AbilityToEnd, FGameplayAbilityActivationInfo ActivationInfo);
-	virtual void ClientEndAbility_Implementation(FGameplayAbilitySpecHandle AbilityToEnd, FGameplayAbilityActivationInfo ActivationInfo);
 
-	UFUNCTION(Server="ServerCancelAbility_Implementation", reliable, WithValidation="ServerCancelAbility_Validate")
+	UFUNCTION(Server, reliable, WithValidation)
 	void    ServerCancelAbility(FGameplayAbilitySpecHandle AbilityToCancel, FGameplayAbilityActivationInfo ActivationInfo);
-	virtual void ServerCancelAbility_Implementation(FGameplayAbilitySpecHandle AbilityToCancel, FGameplayAbilityActivationInfo ActivationInfo);
-	virtual bool ServerCancelAbility_Validate(FGameplayAbilitySpecHandle , FGameplayAbilityActivationInfo );
 
-	UFUNCTION(Client="ClientCancelAbility_Implementation", reliable)
+	UFUNCTION(Client, reliable)
 	void    ClientCancelAbility(FGameplayAbilitySpecHandle AbilityToCancel, FGameplayAbilityActivationInfo ActivationInfo);
-	virtual void ClientCancelAbility_Implementation(FGameplayAbilitySpecHandle AbilityToCancel, FGameplayAbilityActivationInfo ActivationInfo);
 
-	UFUNCTION(Client="ClientActivateAbilityFailed_Implementation", Reliable)
+	UFUNCTION(Client, Reliable)
 	void	ClientActivateAbilityFailed(FGameplayAbilitySpecHandle AbilityToActivate, int16 PredictionKey);
-	virtual void ClientActivateAbilityFailed_Implementation(FGameplayAbilitySpecHandle AbilityToActivate, int16 PredictionKey);
 
 	void	OnClientActivateAbilityFailed(FGameplayAbilitySpecHandle AbilityToActivate, FPredictionKey::KeyType PredictionKey);
 
-	UFUNCTION(Client="ClientActivateAbilitySucceed_Implementation", Reliable)
+	UFUNCTION(Client, Reliable)
 	void	ClientActivateAbilitySucceed(FGameplayAbilitySpecHandle AbilityToActivate, int16 PredictionKey);
-	virtual void ClientActivateAbilitySucceed_Implementation(FGameplayAbilitySpecHandle AbilityToActivate, int16 PredictionKey);
 
-	UFUNCTION(Client="ClientActivateAbilitySucceedWithEventData_Implementation", Reliable)
+	UFUNCTION(Client, Reliable)
 	void	ClientActivateAbilitySucceedWithEventData(FGameplayAbilitySpecHandle AbilityToActivate, int16 PredictionKey, FGameplayEventData TriggerEventData);
-	virtual void ClientActivateAbilitySucceedWithEventData_Implementation(FGameplayAbilitySpecHandle AbilityToActivate, int16 PredictionKey, FGameplayEventData TriggerEventData);
 
 	// ----------------------------------------------------------------------------------------------------------------
 
@@ -847,16 +828,12 @@ protected:
 	void OnRep_ReplicatedAnimMontage();
 
 	/** RPC function called from CurrentMontageSetNextSectopnName, replicates to other clients */
-	UFUNCTION(reliable, server="ServerCurrentMontageSetNextSectionName_Implementation", WithValidation="ServerCurrentMontageSetNextSectionName_Validate")
+	UFUNCTION(reliable, server, WithValidation)
 	void ServerCurrentMontageSetNextSectionName(UAnimMontage* ClientAnimMontage, float ClientPosition, FName SectionName, FName NextSectionName);
-	virtual void ServerCurrentMontageSetNextSectionName_Implementation(UAnimMontage* ClientAnimMontage, float ClientPosition, FName SectionName, FName NextSectionName);
-	virtual bool ServerCurrentMontageSetNextSectionName_Validate(UAnimMontage* , float , FName , FName );
 
 	/** RPC function called from CurrentMontageJumpToSection, replicates to other clients */
-	UFUNCTION(reliable, server="ServerCurrentMontageJumpToSectionName_Implementation", WithValidation="ServerCurrentMontageJumpToSectionName_Validate")
+	UFUNCTION(reliable, server, WithValidation)
 	void ServerCurrentMontageJumpToSectionName(UAnimMontage* ClientAnimMontage, FName SectionName);
-	virtual void ServerCurrentMontageJumpToSectionName_Implementation(UAnimMontage* ClientAnimMontage, FName SectionName);
-	virtual bool ServerCurrentMontageJumpToSectionName_Validate(UAnimMontage* , FName );
 
 	/** Abilities that are triggered from a gameplay event */
 	TMap<FGameplayTag, TArray<FGameplayAbilitySpecHandle > > GameplayEventTriggeredAbilities;
@@ -938,30 +915,23 @@ public:
 	 */
 
 	/** Replicates the Generic Replicated Event to the server. */
-	UFUNCTION(Server="ServerSetReplicatedEvent_Implementation", reliable, WithValidation="ServerSetReplicatedEvent_Validate")
+	UFUNCTION(Server, reliable, WithValidation)
 	void ServerSetReplicatedEvent(EAbilityGenericReplicatedEvent::Type EventType, FGameplayAbilitySpecHandle AbilityHandle, FPredictionKey AbilityOriginalPredictionKey, FPredictionKey CurrentPredictionKey);
-	virtual void ServerSetReplicatedEvent_Implementation(EAbilityGenericReplicatedEvent::Type EventType, FGameplayAbilitySpecHandle AbilityHandle, FPredictionKey AbilityOriginalPredictionKey, FPredictionKey CurrentPredictionKey);
-	virtual bool ServerSetReplicatedEvent_Validate(EAbilityGenericReplicatedEvent::Type , FGameplayAbilitySpecHandle , FPredictionKey , FPredictionKey );
 
 	/** Replicates the Generic Replicated Event to the server. */
-	UFUNCTION(Client="ClientSetReplicatedEvent_Implementation", reliable)
+	UFUNCTION(Client, reliable)
 	void ClientSetReplicatedEvent(EAbilityGenericReplicatedEvent::Type EventType, FGameplayAbilitySpecHandle AbilityHandle, FPredictionKey AbilityOriginalPredictionKey);
-	virtual void ClientSetReplicatedEvent_Implementation(EAbilityGenericReplicatedEvent::Type EventType, FGameplayAbilitySpecHandle AbilityHandle, FPredictionKey AbilityOriginalPredictionKey);
 
 	/** Calls local callbacks that are registered with the given Generic Replicated Event */
 	bool InvokeReplicatedEvent(EAbilityGenericReplicatedEvent::Type EventType, FGameplayAbilitySpecHandle AbilityHandle, FPredictionKey AbilityOriginalPredictionKey);
 
 	/**  */
-	UFUNCTION(Server="ServerSetReplicatedTargetData_Implementation", reliable, WithValidation="ServerSetReplicatedTargetData_Validate")
+	UFUNCTION(Server, reliable, WithValidation)
 	void ServerSetReplicatedTargetData(FGameplayAbilitySpecHandle AbilityHandle, FPredictionKey AbilityOriginalPredictionKey, FGameplayAbilityTargetDataHandle ReplicatedTargetData, FGameplayTag ApplicationTag, FPredictionKey CurrentPredictionKey);
-	virtual void ServerSetReplicatedTargetData_Implementation(FGameplayAbilitySpecHandle AbilityHandle, FPredictionKey AbilityOriginalPredictionKey, FGameplayAbilityTargetDataHandle ReplicatedTargetData, FGameplayTag ApplicationTag, FPredictionKey CurrentPredictionKey);
-	virtual bool ServerSetReplicatedTargetData_Validate(FGameplayAbilitySpecHandle , FPredictionKey , FGameplayAbilityTargetDataHandle , FGameplayTag , FPredictionKey );
 
 	/** Replicates to the server that targeting has been cancelled */
-	UFUNCTION(Server="ServerSetReplicatedTargetDataCancelled_Implementation", reliable, WithValidation="ServerSetReplicatedTargetDataCancelled_Validate")
+	UFUNCTION(Server, reliable, WithValidation)
 	void ServerSetReplicatedTargetDataCancelled(FGameplayAbilitySpecHandle AbilityHandle, FPredictionKey AbilityOriginalPredictionKey, FPredictionKey CurrentPredictionKey);
-	virtual void ServerSetReplicatedTargetDataCancelled_Implementation(FGameplayAbilitySpecHandle AbilityHandle, FPredictionKey AbilityOriginalPredictionKey, FPredictionKey CurrentPredictionKey);
-	virtual bool ServerSetReplicatedTargetDataCancelled_Validate(FGameplayAbilitySpecHandle , FPredictionKey , FPredictionKey );
 
 	/** Sets the current target data and calls applicable callbacks */
 	void ConfirmAbilityTargetData(FGameplayAbilitySpecHandle AbilityHandle, FPredictionKey AbilityOriginalPredictionKey, const FGameplayAbilityTargetDataHandle& TargetData, const FGameplayTag& ApplicationTag);
@@ -1000,15 +970,11 @@ public:
 	FSimpleMulticastDelegate& AbilityReplicatedEventDelegate(EAbilityGenericReplicatedEvent::Type EventType, FGameplayAbilitySpecHandle AbilityHandle, FPredictionKey AbilityOriginalPredictionKey);
 
 	// Direct Input state replication. These will be called if bReplicateInputDirectly is true on the ability and is generally not a good thing to use. (Instead, prefer to use Generic Replicated Events).
-	UFUNCTION(Server="ServerSetInputPressed_Implementation", reliable, WithValidation="ServerSetInputPressed_Validate")
+	UFUNCTION(Server, reliable, WithValidation)
 	void ServerSetInputPressed(FGameplayAbilitySpecHandle AbilityHandle);
-	virtual void ServerSetInputPressed_Implementation(FGameplayAbilitySpecHandle AbilityHandle);
-	virtual bool ServerSetInputPressed_Validate(FGameplayAbilitySpecHandle );
 
-	UFUNCTION(Server="ServerSetInputReleased_Implementation", reliable, WithValidation="ServerSetInputReleased_Validate")
+	UFUNCTION(Server, reliable, WithValidation)
 	void ServerSetInputReleased(FGameplayAbilitySpecHandle AbilityHandle);
-	virtual void ServerSetInputReleased_Implementation(FGameplayAbilitySpecHandle AbilityHandle);
-	virtual bool ServerSetInputReleased_Validate(FGameplayAbilitySpecHandle );
 
 	/** Called on local player always. Called on server only if bReplicateInputDirectly is set on the GameplayAbility. */
 	void AbilitySpecInputReleased(FGameplayAbilitySpec& Spec);
diff --git a/Engine/Source/Runtime/GameplayAbilities/Public/AbilitySystemDebugHUD.h b/Engine/Source/Runtime/GameplayAbilities/Public/AbilitySystemDebugHUD.h
--- a/Engine/Source/Runtime/GameplayAbilities/Public/AbilitySystemDebugHUD.h
+++ b/Engine/Source/Runtime/GameplayAbilities/Public/AbilitySystemDebugHUD.h
@@ -27,9 +27,7 @@ namespace EAlignVertical
 UCLASS()
 class AAbilitySystemDebugHUD : public AHUD
 {
-	GENERATED_BODY()
-public:
-	AAbilitySystemDebugHUD(const FObjectInitializer& ObjectInitializer = FObjectInitializer::Get());
+	GENERATED_UCLASS_BODY()
 
 	/** main HUD update loop */
 	void DrawWithBackground(UFont* InFont, const FString& Text, const FColor& TextColor, EAlignHorizontal::Type HAlign, float& OffsetX, EAlignVertical::Type VAlign, float& OffsetY, float Alpha=1.f);
@@ -40,4 +38,4 @@ private:
 
 	void DrawDebugAbilitySystemComponent(UAbilitySystemComponent *Component);
 
-};
+};
\ No newline at end of file
diff --git a/Engine/Source/Runtime/GameplayAbilities/Public/AbilitySystemInterface.h b/Engine/Source/Runtime/GameplayAbilities/Public/AbilitySystemInterface.h
--- a/Engine/Source/Runtime/GameplayAbilities/Public/AbilitySystemInterface.h
+++ b/Engine/Source/Runtime/GameplayAbilities/Public/AbilitySystemInterface.h
@@ -12,15 +12,12 @@ class UAbilitySystemComponent;
 UINTERFACE(MinimalAPI, meta = (CannotImplementInterfaceInBlueprint))
 class UAbilitySystemInterface : public UInterface
 {
-	GENERATED_BODY()
-public:
-	GAMEPLAYABILITIES_API UAbilitySystemInterface(const FObjectInitializer& ObjectInitializer = FObjectInitializer::Get());
+	GENERATED_UINTERFACE_BODY()
 };
 
 class GAMEPLAYABILITIES_API IAbilitySystemInterface
 {
-	GENERATED_BODY()
-public:
+	GENERATED_IINTERFACE_BODY()
 
 	virtual UAbilitySystemComponent* GetAbilitySystemComponent() const = 0;
-};
+};
\ No newline at end of file
diff --git a/Engine/Source/Runtime/GameplayAbilities/Public/AbilitySystemTestAttributeSet.h b/Engine/Source/Runtime/GameplayAbilities/Public/AbilitySystemTestAttributeSet.h
--- a/Engine/Source/Runtime/GameplayAbilities/Public/AbilitySystemTestAttributeSet.h
+++ b/Engine/Source/Runtime/GameplayAbilities/Public/AbilitySystemTestAttributeSet.h
@@ -7,9 +7,7 @@
 UCLASS(Blueprintable, BlueprintType, meta=(HideInDetailsView))
 class GAMEPLAYABILITIES_API UAbilitySystemTestAttributeSet : public UAttributeSet
 {
-	GENERATED_BODY()
-public:
-	UAbilitySystemTestAttributeSet(const FObjectInitializer& ObjectInitializer = FObjectInitializer::Get());
+	GENERATED_UCLASS_BODY()
 
 
 
diff --git a/Engine/Source/Runtime/GameplayAbilities/Public/AbilitySystemTestPawn.h b/Engine/Source/Runtime/GameplayAbilities/Public/AbilitySystemTestPawn.h
--- a/Engine/Source/Runtime/GameplayAbilities/Public/AbilitySystemTestPawn.h
+++ b/Engine/Source/Runtime/GameplayAbilities/Public/AbilitySystemTestPawn.h
@@ -8,9 +8,7 @@
 UCLASS(Blueprintable, BlueprintType, notplaceable)
 class GAMEPLAYABILITIES_API AAbilitySystemTestPawn : public ADefaultPawn, public IGameplayCueInterface, public IAbilitySystemInterface
 {
-	GENERATED_BODY()
-public:
-	AAbilitySystemTestPawn(const FObjectInitializer& ObjectInitializer = FObjectInitializer::Get());
+	GENERATED_UCLASS_BODY()
 
 	virtual void PostInitializeComponents() override;
 
diff --git a/Engine/Source/Runtime/GameplayAbilities/Public/AttributeSet.h b/Engine/Source/Runtime/GameplayAbilities/Public/AttributeSet.h
--- a/Engine/Source/Runtime/GameplayAbilities/Public/AttributeSet.h
+++ b/Engine/Source/Runtime/GameplayAbilities/Public/AttributeSet.h
@@ -12,7 +12,7 @@ struct FGameplayAbilityActorInfo;
 USTRUCT(BlueprintType)
 struct GAMEPLAYABILITIES_API FGameplayAttribute
 {
-	GENERATED_BODY()
+	GENERATED_USTRUCT_BODY()
 
 	FGameplayAttribute()
 		: Attribute(nullptr)
@@ -74,9 +74,7 @@ private:
 UCLASS(DefaultToInstanced, Blueprintable)
 class GAMEPLAYABILITIES_API UAttributeSet : public UObject
 {
-	GENERATED_BODY()
-public:
-	UAttributeSet(const FObjectInitializer& ObjectInitializer = FObjectInitializer::Get());
+	GENERATED_UCLASS_BODY()
 
 public:
 	bool IsNameStableForNetworking() const override;
@@ -146,7 +144,7 @@ protected:
 USTRUCT()
 struct GAMEPLAYABILITIES_API FGlobalCurveDataOverride
 {
-	GENERATED_BODY()
+	GENERATED_USTRUCT_BODY()
 
 	TArray<UCurveTable*>	Overrides;
 };
@@ -158,7 +156,7 @@ struct GAMEPLAYABILITIES_API FGlobalCurveDataOverride
 USTRUCT()
 struct GAMEPLAYABILITIES_API FScalableFloat
 {
-	GENERATED_BODY()
+	GENERATED_USTRUCT_BODY()
 
 	FScalableFloat()
 		: Value(0.f)
@@ -239,7 +237,7 @@ private:
 USTRUCT(BlueprintType)
 struct GAMEPLAYABILITIES_API FAttributeMetaData : public FTableRowBase
 {
-	GENERATED_BODY()
+	GENERATED_USTRUCT_BODY()
 
 public:
 
diff --git a/Engine/Source/Runtime/GameplayAbilities/Public/GameplayAbilityBlueprintGeneratedClass.h b/Engine/Source/Runtime/GameplayAbilities/Public/GameplayAbilityBlueprintGeneratedClass.h
--- a/Engine/Source/Runtime/GameplayAbilities/Public/GameplayAbilityBlueprintGeneratedClass.h
+++ b/Engine/Source/Runtime/GameplayAbilities/Public/GameplayAbilityBlueprintGeneratedClass.h
@@ -7,7 +7,5 @@
 UCLASS()
 class  GAMEPLAYABILITIES_API UGameplayAbilityBlueprintGeneratedClass : public UBlueprintGeneratedClass
 {
-	GENERATED_BODY()
-public:
-	UGameplayAbilityBlueprintGeneratedClass(const FObjectInitializer& ObjectInitializer = FObjectInitializer::Get());
-};
+	GENERATED_UCLASS_BODY()
+};
\ No newline at end of file
diff --git a/Engine/Source/Runtime/GameplayAbilities/Public/GameplayAbilitySet.h b/Engine/Source/Runtime/GameplayAbilities/Public/GameplayAbilitySet.h
--- a/Engine/Source/Runtime/GameplayAbilities/Public/GameplayAbilitySet.h
+++ b/Engine/Source/Runtime/GameplayAbilities/Public/GameplayAbilitySet.h
@@ -38,7 +38,7 @@ namespace EGameplayAbilityInputBinds
 USTRUCT()
 struct FGameplayAbilityBindInfo
 {
-	GENERATED_BODY()
+	GENERATED_USTRUCT_BODY()
 
 	UPROPERTY(EditAnywhere, Category = BindInfo)
 	TEnumAsByte<EGameplayAbilityInputBinds::Type>	Command;
@@ -56,13 +56,11 @@ struct FGameplayAbilityBindInfo
 UCLASS()
 class GAMEPLAYABILITIES_API UGameplayAbilitySet : public UDataAsset
 {
-	GENERATED_BODY()
-public:
-	UGameplayAbilitySet(const FObjectInitializer& ObjectInitializer = FObjectInitializer::Get());
+	GENERATED_UCLASS_BODY()
 
 	
 	UPROPERTY(EditAnywhere, Category = AbilitySet)
 	TArray<FGameplayAbilityBindInfo>	Abilities;
 
 	void GiveAbilities(UAbilitySystemComponent* AbilitySystemComponent) const;
-};
+};
\ No newline at end of file
diff --git a/Engine/Source/Runtime/GameplayAbilities/Public/GameplayCueInterface.h b/Engine/Source/Runtime/GameplayAbilities/Public/GameplayCueInterface.h
--- a/Engine/Source/Runtime/GameplayAbilities/Public/GameplayCueInterface.h
+++ b/Engine/Source/Runtime/GameplayAbilities/Public/GameplayCueInterface.h
@@ -11,15 +11,12 @@
 UINTERFACE(MinimalAPI, meta = (CannotImplementInterfaceInBlueprint))
 class UGameplayCueInterface: public UInterface
 {
-	GENERATED_BODY()
-public:
-	GAMEPLAYABILITIES_API UGameplayCueInterface(const FObjectInitializer& ObjectInitializer = FObjectInitializer::Get());
+	GENERATED_UINTERFACE_BODY()
 };
 
 class GAMEPLAYABILITIES_API IGameplayCueInterface
 {
-	GENERATED_BODY()
-public:
+	GENERATED_IINTERFACE_BODY()
 
 	virtual void HandleGameplayCue(AActor *Self, FGameplayTag GameplayCueTag, EGameplayCueEvent::Type EventType, FGameplayCueParameters Parameters);
 
@@ -61,7 +58,7 @@ struct FActiveGameplayCueContainer;
 USTRUCT(BlueprintType)
 struct FActiveGameplayCue : public FFastArraySerializerItem
 {
-	GENERATED_BODY()
+	GENERATED_USTRUCT_BODY()
 
 	UPROPERTY()
 	FGameplayTag GameplayCueTag;
@@ -77,7 +74,7 @@ struct FActiveGameplayCue : public FFastArraySerializerItem
 USTRUCT(BlueprintType)
 struct FActiveGameplayCueContainer : public FFastArraySerializer
 {
-	GENERATED_BODY()
+	GENERATED_USTRUCT_BODY()
 
 	UPROPERTY()
 	TArray< FActiveGameplayCue >	GameplayCues;
@@ -105,4 +102,4 @@ struct TStructOpsTypeTraits< FActiveGameplayCueContainer > : public TStructOpsTy
 	{
 		WithNetDeltaSerializer = true,
 	};
-};
+};
\ No newline at end of file
diff --git a/Engine/Source/Runtime/GameplayAbilities/Public/GameplayCueManager.h b/Engine/Source/Runtime/GameplayAbilities/Public/GameplayCueManager.h
--- a/Engine/Source/Runtime/GameplayAbilities/Public/GameplayCueManager.h
+++ b/Engine/Source/Runtime/GameplayAbilities/Public/GameplayCueManager.h
@@ -48,7 +48,7 @@
 USTRUCT()
 struct FGameplayCueNotifyData
 {
-	GENERATED_BODY()
+	GENERATED_USTRUCT_BODY()
 
 	FGameplayCueNotifyData()
 	: LoadedGameplayCueClass(nullptr)
@@ -80,7 +80,7 @@ enum class EGameplayCuePayloadType : uint8
 USTRUCT()
 struct FGameplayCuePendingExecute
 {
-	GENERATED_BODY()
+	GENERATED_USTRUCT_BODY()
 
 	FGameplayCuePendingExecute()
 	: PayloadType(EGameplayCuePayloadType::EffectContext)
@@ -129,9 +129,7 @@ struct GAMEPLAYABILITIES_API FScopedGameplayCueSendContext
 UCLASS()
 class GAMEPLAYABILITIES_API UGameplayCueManager : public UDataAsset
 {
-	GENERATED_BODY()
-public:
-	UGameplayCueManager(const FObjectInitializer& ObjectInitializer = FObjectInitializer::Get());
+	GENERATED_UCLASS_BODY()
 
 	// -------------------------------------------------------------
 	// Wrappers to handle replicating executed cues
diff --git a/Engine/Source/Runtime/GameplayAbilities/Public/GameplayCueNotify_HitImpact.h b/Engine/Source/Runtime/GameplayAbilities/Public/GameplayCueNotify_HitImpact.h
--- a/Engine/Source/Runtime/GameplayAbilities/Public/GameplayCueNotify_HitImpact.h
+++ b/Engine/Source/Runtime/GameplayAbilities/Public/GameplayCueNotify_HitImpact.h
@@ -15,9 +15,7 @@
 UCLASS()
 class GAMEPLAYABILITIES_API UGameplayCueNotify_HitImpact : public UGameplayCueNotify_Static
 {
-	GENERATED_BODY()
-public:
-	UGameplayCueNotify_HitImpact(const FObjectInitializer& ObjectInitializer = FObjectInitializer::Get());
+	GENERATED_UCLASS_BODY()
 
 	/** Does this GameplayCueNotify handle this type of GameplayCueEvent? */
 	virtual bool HandlesEvent(EGameplayCueEvent::Type EventType) const override;
diff --git a/Engine/Source/Runtime/GameplayAbilities/Public/GameplayCueNotify_Static.h b/Engine/Source/Runtime/GameplayAbilities/Public/GameplayCueNotify_Static.h
--- a/Engine/Source/Runtime/GameplayAbilities/Public/GameplayCueNotify_Static.h
+++ b/Engine/Source/Runtime/GameplayAbilities/Public/GameplayCueNotify_Static.h
@@ -22,9 +22,7 @@
 UCLASS(Blueprintable, meta = (ShowWorldContextPin), hidecategories = (Replication))
 class GAMEPLAYABILITIES_API UGameplayCueNotify_Static : public UObject
 {
-	GENERATED_BODY()
-public:
-	UGameplayCueNotify_Static(const FObjectInitializer& ObjectInitializer = FObjectInitializer::Get());
+	GENERATED_UCLASS_BODY()
 
 	/** Does this GameplayCueNotify handle this type of GameplayCueEvent? */
 	virtual bool HandlesEvent(EGameplayCueEvent::Type EventType) const;
@@ -47,19 +45,15 @@ public:
 
 	UFUNCTION(BlueprintNativeEvent, BlueprintPure, Category = "GameplayCueNotify")
 	bool OnExecute(AActor* MyTarget, FGameplayCueParameters Parameters) const;
-	virtual bool OnExecute_Implementation(AActor* MyTarget, FGameplayCueParameters Parameters) const;
 
 	UFUNCTION(BlueprintNativeEvent, BlueprintPure, Category = "GameplayCueNotify")
 	bool OnActive(AActor* MyTarget, FGameplayCueParameters Parameters) const;
-	virtual bool OnActive_Implementation(AActor* MyTarget, FGameplayCueParameters Parameters) const;
 
 	UFUNCTION(BlueprintNativeEvent, BlueprintPure, Category = "GameplayCueNotify")
 	bool WhileActive(AActor* MyTarget, FGameplayCueParameters Parameters) const;
-	virtual bool WhileActive_Implementation(AActor* MyTarget, FGameplayCueParameters Parameters) const;
 
 	UFUNCTION(BlueprintNativeEvent, BlueprintPure, Category = "GameplayCueNotify")
 	bool OnRemove(AActor* MyTarget, FGameplayCueParameters Parameters) const;
-	virtual bool OnRemove_Implementation(AActor* MyTarget, FGameplayCueParameters Parameters) const;
 
 	UPROPERTY(EditDefaultsOnly, Category = GameplayCue)
 	FGameplayTag	GameplayCueTag;
@@ -74,4 +68,4 @@ public:
 
 private:
 	virtual void DeriveGameplayCueTagFromAssetName();
-};
+};
\ No newline at end of file
diff --git a/Engine/Source/Runtime/GameplayAbilities/Public/GameplayEffectCalculation.h b/Engine/Source/Runtime/GameplayAbilities/Public/GameplayEffectCalculation.h
--- a/Engine/Source/Runtime/GameplayAbilities/Public/GameplayEffectCalculation.h
+++ b/Engine/Source/Runtime/GameplayAbilities/Public/GameplayEffectCalculation.h
@@ -12,9 +12,7 @@ class UAbilitySystemComponent;
 UCLASS(BlueprintType, Blueprintable, Abstract)
 class GAMEPLAYABILITIES_API UGameplayEffectCalculation : public UObject
 {
-	GENERATED_BODY()
-public:
-	UGameplayEffectCalculation(const FObjectInitializer& ObjectInitializer = FObjectInitializer::Get());
+	GENERATED_UCLASS_BODY()
 
 public:
 
diff --git a/Engine/Source/Runtime/GameplayAbilities/Public/GameplayEffectExtension.h b/Engine/Source/Runtime/GameplayAbilities/Public/GameplayEffectExtension.h
--- a/Engine/Source/Runtime/GameplayAbilities/Public/GameplayEffectExtension.h
+++ b/Engine/Source/Runtime/GameplayAbilities/Public/GameplayEffectExtension.h
@@ -27,9 +27,7 @@ struct FGameplayEffectModCallbackData
 UCLASS(BlueprintType)
 class GAMEPLAYABILITIES_API UGameplayEffectExtension: public UObject
 {
-	GENERATED_BODY()
-public:
-	UGameplayEffectExtension(const FObjectInitializer& ObjectInitializer = FObjectInitializer::Get());
+	GENERATED_UCLASS_BODY()
 
 	/** Attributes on the source instigator that are relevant to calculating modifications using this extension */
 	UPROPERTY(EditDefaultsOnly, Category=Calculation)
@@ -50,4 +48,4 @@ public:
 	{
 
 	}
-};
+};
\ No newline at end of file
diff --git a/Engine/Source/Runtime/GameplayAbilities/Public/GameplayEffectExtension_LifestealTest.h b/Engine/Source/Runtime/GameplayAbilities/Public/GameplayEffectExtension_LifestealTest.h
--- a/Engine/Source/Runtime/GameplayAbilities/Public/GameplayEffectExtension_LifestealTest.h
+++ b/Engine/Source/Runtime/GameplayAbilities/Public/GameplayEffectExtension_LifestealTest.h
@@ -11,9 +11,7 @@
 UCLASS(BlueprintType)
 class GAMEPLAYABILITIES_API UGameplayEffectExtension_LifestealTest : public UGameplayEffectExtension
 {
-	GENERATED_BODY()
-public:
-	UGameplayEffectExtension_LifestealTest(const FObjectInitializer& ObjectInitializer = FObjectInitializer::Get());
+	GENERATED_UCLASS_BODY()
 
 public:
 
@@ -23,4 +21,4 @@ public:
 	/** The GameplayEffect to apply when restoring health to the instigator */
 	UPROPERTY(EditDefaultsOnly, Category = Lifesteal)
 	UGameplayEffect *	HealthRestoreGameplayEffect;
-};
+};
\ No newline at end of file
diff --git a/Engine/Source/Runtime/GameplayAbilities/Public/GameplayEffectExtension_ShieldTest.h b/Engine/Source/Runtime/GameplayAbilities/Public/GameplayEffectExtension_ShieldTest.h
--- a/Engine/Source/Runtime/GameplayAbilities/Public/GameplayEffectExtension_ShieldTest.h
+++ b/Engine/Source/Runtime/GameplayAbilities/Public/GameplayEffectExtension_ShieldTest.h
@@ -11,9 +11,7 @@
 UCLASS(BlueprintType)
 class GAMEPLAYABILITIES_API UGameplayEffectExtension_ShieldTest : public UGameplayEffectExtension
 {
-	GENERATED_BODY()
-public:
-	UGameplayEffectExtension_ShieldTest(const FObjectInitializer& ObjectInitializer = FObjectInitializer::Get());
+	GENERATED_UCLASS_BODY()
 
 public:
 
@@ -22,4 +20,4 @@ public:
 
 	UPROPERTY(EditDefaultsOnly, Category = Lifesteal)
 	UGameplayEffect * ShieldRemoveGameplayEffect;
-};
+};
\ No newline at end of file
diff --git a/Engine/Source/Runtime/GameplayAbilities/Public/GameplayEffectTemplate.h b/Engine/Source/Runtime/GameplayAbilities/Public/GameplayEffectTemplate.h
--- a/Engine/Source/Runtime/GameplayAbilities/Public/GameplayEffectTemplate.h
+++ b/Engine/Source/Runtime/GameplayAbilities/Public/GameplayEffectTemplate.h
@@ -27,13 +27,11 @@
 UCLASS()
 class UGameplayEffectTemplate : public UGameplayEffect
 {
-	GENERATED_BODY()
-public:
-	UGameplayEffectTemplate(const FObjectInitializer& ObjectInitializer = FObjectInitializer::Get());
+	GENERATED_UCLASS_BODY()
 
 #if WITH_EDITORONLY_DATA
 	/** Names of the properties that should be exposed in the limited/basic view when editing a GameplayEffect based on this template */
 	UPROPERTY(EditDefaultsOnly, Category="Template")
 	TArray<FString>	EditableProperties;
 #endif
-};
+};
\ No newline at end of file
diff --git a/Engine/Source/Runtime/GameplayAbilities/Public/GameplayEffectUIData.h b/Engine/Source/Runtime/GameplayAbilities/Public/GameplayEffectUIData.h
--- a/Engine/Source/Runtime/GameplayAbilities/Public/GameplayEffectUIData.h
+++ b/Engine/Source/Runtime/GameplayAbilities/Public/GameplayEffectUIData.h
@@ -11,12 +11,10 @@
 UCLASS(Blueprintable, Abstract, EditInlineNew, CollapseCategories)
 class GAMEPLAYABILITIES_API UGameplayEffectUIData : public UObject
 {
-	GENERATED_BODY()
-public:
-	UGameplayEffectUIData(const FObjectInitializer& ObjectInitializer = FObjectInitializer::Get());
+	GENERATED_UCLASS_BODY()
 
 	virtual bool NeedsLoadForServer() const override
 	{
 		return false;
 	}
-};
+};
\ No newline at end of file
diff --git a/Engine/Source/Runtime/GameplayAbilities/Public/GameplayEffectUIData_TextOnly.h b/Engine/Source/Runtime/GameplayAbilities/Public/GameplayEffectUIData_TextOnly.h
--- a/Engine/Source/Runtime/GameplayAbilities/Public/GameplayEffectUIData_TextOnly.h
+++ b/Engine/Source/Runtime/GameplayAbilities/Public/GameplayEffectUIData_TextOnly.h
@@ -11,10 +11,8 @@
 UCLASS()
 class UGameplayEffectUIData_TextOnly : public UGameplayEffectUIData
 {
-	GENERATED_BODY()
-public:
-	UGameplayEffectUIData_TextOnly(const FObjectInitializer& ObjectInitializer = FObjectInitializer::Get());
+	GENERATED_UCLASS_BODY()
 
 	UPROPERTY(EditDefaultsOnly, BlueprintReadOnly, Category=Data)
 	FText Description;
-};
+};
\ No newline at end of file
diff --git a/Engine/Source/Runtime/GameplayAbilities/Public/GameplayModMagnitudeCalculation.h b/Engine/Source/Runtime/GameplayAbilities/Public/GameplayModMagnitudeCalculation.h
--- a/Engine/Source/Runtime/GameplayAbilities/Public/GameplayModMagnitudeCalculation.h
+++ b/Engine/Source/Runtime/GameplayAbilities/Public/GameplayModMagnitudeCalculation.h
@@ -14,11 +14,9 @@ class UAbilitySystemComponent;
 UCLASS(BlueprintType, Blueprintable, Abstract)
 class GAMEPLAYABILITIES_API UGameplayModMagnitudeCalculation : public UGameplayEffectCalculation
 {
-	GENERATED_BODY()
-public:
-	UGameplayModMagnitudeCalculation(const FObjectInitializer& ObjectInitializer = FObjectInitializer::Get());
 
 public:
+	GENERATED_UCLASS_BODY()
 
 	/**
 	 * Calculate the base magnitude of the gameplay effect modifier, given the specified spec. Note that the owngin spec def can still modify this base value
@@ -30,11 +28,10 @@ public:
 	 */
 	UFUNCTION(BlueprintNativeEvent, Category="Calculation")
 	float CalculateBaseMagnitude(const FGameplayEffectSpec& Spec) const;
-	virtual float CalculateBaseMagnitude_Implementation(const FGameplayEffectSpec& Spec) const;
 
 protected:
 	
 
 	/** Convenience method to get attribute magnitude during a CalculateMagnitude call */
 	bool GetCapturedAttributeMagnitude(const FGameplayEffectAttributeCaptureDefinition& Def, const FGameplayEffectSpec& Spec, const FAggregatorEvaluateParameters& EvaluationParameters, OUT float& Magnitude) const;
-};
+};
\ No newline at end of file
diff --git a/Engine/Source/Runtime/GameplayAbilities/Public/TickableAttributeSetInterface.h b/Engine/Source/Runtime/GameplayAbilities/Public/TickableAttributeSetInterface.h
--- a/Engine/Source/Runtime/GameplayAbilities/Public/TickableAttributeSetInterface.h
+++ b/Engine/Source/Runtime/GameplayAbilities/Public/TickableAttributeSetInterface.h
@@ -8,15 +8,12 @@
 UINTERFACE(MinimalAPI, meta = (CannotImplementInterfaceInBlueprint))
 class UTickableAttributeSetInterface : public UInterface
 {
-	GENERATED_BODY()
-public:
-	GAMEPLAYABILITIES_API UTickableAttributeSetInterface(const FObjectInitializer& ObjectInitializer = FObjectInitializer::Get());
+	GENERATED_UINTERFACE_BODY()
 };
 
 class ITickableAttributeSetInterface
 {
-	GENERATED_BODY()
-public:
+	GENERATED_IINTERFACE_BODY()
 
 public:
 	/**
diff --git a/Engine/Source/Runtime/GameplayTags/Classes/BlueprintGameplayTagLibrary.h b/Engine/Source/Runtime/GameplayTags/Classes/BlueprintGameplayTagLibrary.h
--- a/Engine/Source/Runtime/GameplayTags/Classes/BlueprintGameplayTagLibrary.h
+++ b/Engine/Source/Runtime/GameplayTags/Classes/BlueprintGameplayTagLibrary.h
@@ -13,9 +13,7 @@ struct FGameplayTagContainer;
 UCLASS(MinimalAPI)
 class UBlueprintGameplayTagLibrary : public UBlueprintFunctionLibrary
 {
-	GENERATED_BODY()
-public:
-	GAMEPLAYTAGS_API UBlueprintGameplayTagLibrary(const FObjectInitializer& ObjectInitializer = FObjectInitializer::Get());
+	GENERATED_UCLASS_BODY()
 
 	/**
 	 * Determine if the specified gameplay tags match, given the specified match types
diff --git a/Engine/Source/Runtime/GameplayTags/Classes/GameplayTagAssetInterface.h b/Engine/Source/Runtime/GameplayTags/Classes/GameplayTagAssetInterface.h
--- a/Engine/Source/Runtime/GameplayTags/Classes/GameplayTagAssetInterface.h
+++ b/Engine/Source/Runtime/GameplayTags/Classes/GameplayTagAssetInterface.h
@@ -9,15 +9,12 @@
 UINTERFACE(Blueprintable, MinimalAPI, meta=(CannotImplementInterfaceInBlueprint))
 class UGameplayTagAssetInterface : public UInterface
 {
-	GENERATED_BODY()
-public:
-	GAMEPLAYTAGS_API UGameplayTagAssetInterface(const FObjectInitializer& ObjectInitializer = FObjectInitializer::Get());
+	GENERATED_UINTERFACE_BODY()
 };
 
 class GAMEPLAYTAGS_API IGameplayTagAssetInterface
 {
-	GENERATED_BODY()
-public:
+	GENERATED_IINTERFACE_BODY()
 
 	/**
 	 * Get any owned gameplay tags on the asset
diff --git a/Engine/Source/Runtime/GameplayTags/Classes/GameplayTagsManager.h b/Engine/Source/Runtime/GameplayTags/Classes/GameplayTagsManager.h
--- a/Engine/Source/Runtime/GameplayTags/Classes/GameplayTagsManager.h
+++ b/Engine/Source/Runtime/GameplayTags/Classes/GameplayTagsManager.h
@@ -9,7 +9,7 @@
 USTRUCT()
 struct FGameplayTagTableRow : public FTableRowBase
 {
-	GENERATED_BODY()
+	GENERATED_USTRUCT_BODY()
 
 	/** Tag specified in the table */
 	UPROPERTY(EditAnywhere, BlueprintReadWrite, Category=GameplayTag)
@@ -33,7 +33,7 @@ struct FGameplayTagTableRow : public FTableRowBase
 USTRUCT()
 struct FGameplayTagNode
 {
-	GENERATED_BODY()
+	GENERATED_USTRUCT_BODY()
 	FGameplayTagNode(){};
 
 	/** Simple constructor */
@@ -117,9 +117,7 @@ private:
 UCLASS(config=Engine)
 class GAMEPLAYTAGS_API UGameplayTagsManager : public UObject
 {
-	GENERATED_BODY()
-public:
-	UGameplayTagsManager(const FObjectInitializer& ObjectInitializer = FObjectInitializer::Get());
+	GENERATED_UCLASS_BODY()
 
 	// Destructor
 	~UGameplayTagsManager();
diff --git a/Engine/Source/Runtime/IOS/IOSRuntimeSettings/Classes/IOSRuntimeSettings.h b/Engine/Source/Runtime/IOS/IOSRuntimeSettings/Classes/IOSRuntimeSettings.h
--- a/Engine/Source/Runtime/IOS/IOSRuntimeSettings/Classes/IOSRuntimeSettings.h
+++ b/Engine/Source/Runtime/IOS/IOSRuntimeSettings/Classes/IOSRuntimeSettings.h
@@ -41,7 +41,7 @@ UENUM()
 USTRUCT()
 struct FIOSBuildResourceFilePath
 {
-	GENERATED_BODY()
+	GENERATED_USTRUCT_BODY()
 	
 	/**
 	 * Custom export item used to serialize FIOSBuildResourceFilePath types as only a filename, no garland.
@@ -89,10 +89,8 @@ struct TStructOpsTypeTraits<FIOSBuildResourceFilePath> : public TStructOpsTypeTr
 UCLASS(config=Engine, defaultconfig)
 class IOSRUNTIMESETTINGS_API UIOSRuntimeSettings : public UObject
 {
-	GENERATED_BODY()
-public:
-	UIOSRuntimeSettings(const FObjectInitializer& ObjectInitializer = FObjectInitializer::Get());
 public:
+	GENERATED_UCLASS_BODY()
 
 	// Should Game Center support (iOS Online Subsystem) be enabled?
 	UPROPERTY(GlobalConfig, EditAnywhere, Category = Online)
diff --git a/Engine/Source/Runtime/Landscape/Classes/ControlPointMeshComponent.h b/Engine/Source/Runtime/Landscape/Classes/ControlPointMeshComponent.h
--- a/Engine/Source/Runtime/Landscape/Classes/ControlPointMeshComponent.h
+++ b/Engine/Source/Runtime/Landscape/Classes/ControlPointMeshComponent.h
@@ -8,9 +8,7 @@
 UCLASS(MinimalAPI)
 class UControlPointMeshComponent : public UStaticMeshComponent
 {
-	GENERATED_BODY()
-public:
-	LANDSCAPE_API UControlPointMeshComponent(const FObjectInitializer& ObjectInitializer = FObjectInitializer::Get());
+	GENERATED_UCLASS_BODY()
 
 #if WITH_EDITORONLY_DATA
 	UPROPERTY(transient)
diff --git a/Engine/Source/Runtime/Landscape/Classes/Landscape.h b/Engine/Source/Runtime/Landscape/Classes/Landscape.h
--- a/Engine/Source/Runtime/Landscape/Classes/Landscape.h
+++ b/Engine/Source/Runtime/Landscape/Classes/Landscape.h
@@ -29,9 +29,7 @@ enum ELandscapeSetupErrors
 UCLASS(hidecategories=LandscapeProxy, showcategories=(Display, Movement, Collision, Lighting, LOD, Input), MinimalAPI)
 class ALandscape : public ALandscapeProxy
 {
-	GENERATED_BODY()
-public:
-	LANDSCAPE_API ALandscape(const FObjectInitializer& ObjectInitializer = FObjectInitializer::Get());
+	GENERATED_UCLASS_BODY()
 
 	// Make a key for XYtoComponentMap
 	static FIntPoint MakeKey( int32 X, int32 Y ) { return FIntPoint(X, Y); }
@@ -99,4 +97,4 @@ DECLARE_DWORD_COUNTER_STAT_EXTERN(TEXT("Component Mem"), STAT_LandscapeComponent
  */
 LANDSCAPE_API ULandscapeInfoMap& GetLandscapeInfoMap(UWorld* World);
 
-#endif // WITH_EDITORONLY_DATA
+#endif // WITH_EDITORONLY_DATA
\ No newline at end of file
diff --git a/Engine/Source/Runtime/Landscape/Classes/LandscapeComponent.h b/Engine/Source/Runtime/Landscape/Classes/LandscapeComponent.h
--- a/Engine/Source/Runtime/Landscape/Classes/LandscapeComponent.h
+++ b/Engine/Source/Runtime/Landscape/Classes/LandscapeComponent.h
@@ -74,7 +74,7 @@ struct FLandscapeVertexRef
 USTRUCT()
 struct FWeightmapLayerAllocationInfo
 {
-	GENERATED_BODY()
+	GENERATED_USTRUCT_BODY()
 
 	UPROPERTY()
 	ULandscapeLayerInfoObject* LayerInfo;
@@ -141,9 +141,7 @@ struct FLandscapeComponentGrassData
 UCLASS(hidecategories=(Display, Attachment, Physics, Debug, Collision, Movement, Rendering, PrimitiveComponent, Object, Transform), showcategories=("Rendering|Material"), MinimalAPI)
 class ULandscapeComponent : public UPrimitiveComponent
 {
-	GENERATED_BODY()
-public:
-	LANDSCAPE_API ULandscapeComponent(const FObjectInitializer& ObjectInitializer = FObjectInitializer::Get());
+	GENERATED_UCLASS_BODY()
 	
 	/** X offset from global components grid origin (in quads) */
 	UPROPERTY(VisibleAnywhere, BlueprintReadOnly, Category=LandscapeComponent)
diff --git a/Engine/Source/Runtime/Landscape/Classes/LandscapeGrassType.h b/Engine/Source/Runtime/Landscape/Classes/LandscapeGrassType.h
--- a/Engine/Source/Runtime/Landscape/Classes/LandscapeGrassType.h
+++ b/Engine/Source/Runtime/Landscape/Classes/LandscapeGrassType.h
@@ -7,7 +7,7 @@
 USTRUCT()
 struct FGrassVariety
 {
-	GENERATED_BODY()
+	GENERATED_USTRUCT_BODY()
 
 	UPROPERTY(EditAnywhere, Category=Grass)
 	class UStaticMesh* GrassMesh;
@@ -64,9 +64,7 @@ struct FGrassVariety
 UCLASS(MinimalAPI)
 class ULandscapeGrassType : public UObject
 {
-	GENERATED_BODY()
-public:
-	LANDSCAPE_API ULandscapeGrassType(const FObjectInitializer& ObjectInitializer = FObjectInitializer::Get());
+	GENERATED_UCLASS_BODY()
 
 	UPROPERTY(EditAnywhere, Category = Grass)
 	TArray<FGrassVariety> GrassVarieties;
