Merging

//depot/UE4-Orion/Engine/Source/Runtime/GameplayAbilities/...

to //depot/UE4/Engine/Source/Runtime/GameplayAbilities/...

diff --git a/Engine/Source/Runtime/GameplayAbilities/Private/GameplayCueManager.cpp b/Engine/Source/Runtime/GameplayAbilities/Private/GameplayCueManager.cpp
--- a/Engine/Source/Runtime/GameplayAbilities/Private/GameplayCueManager.cpp
+++ b/Engine/Source/Runtime/GameplayAbilities/Private/GameplayCueManager.cpp
@@ -121,7 +121,9 @@ void UGameplayCueManager::HandleGameplayCueNotify_Internal(AActor* TargetActor,
 				{
 					// We don't have an instance for this, and we need one, so make one
 					//SpawnedInstancedCue = static_cast<AGameplayCueNotify_Actor*>(StaticDuplicateObject(InstancedCue, this, TEXT("None"), ~RF_RootSet));
-					SpawnedInstancedCue = TargetActor->GetWorld()->SpawnActor<AGameplayCueNotify_Actor>(InstancedCue->GetClass(), TargetActor->GetActorLocation(), TargetActor->GetActorRotation());
+					FActorSpawnParameters SpawnParams;
+					SpawnParams.Owner = TargetActor;
+					SpawnedInstancedCue = TargetActor->GetWorld()->SpawnActor<AGameplayCueNotify_Actor>(InstancedCue->GetClass(), TargetActor->GetActorLocation(), TargetActor->GetActorRotation(), SpawnParams);
 					auto& InnerMap = NotifyMapActor.FindOrAdd(TargetActor);
 					InnerMap.Add(DataIdx) = SpawnedInstancedCue;
 				}
diff --git a/Engine/Source/Runtime/GameplayAbilities/Private/GameplayCueNotify_Actor.cpp b/Engine/Source/Runtime/GameplayAbilities/Private/GameplayCueNotify_Actor.cpp
--- a/Engine/Source/Runtime/GameplayAbilities/Private/GameplayCueNotify_Actor.cpp
+++ b/Engine/Source/Runtime/GameplayAbilities/Private/GameplayCueNotify_Actor.cpp
@@ -49,6 +49,17 @@ void AGameplayCueNotify_Actor::Serialize(FArchive& Ar)
 	}
 }
 
+void AGameplayCueNotify_Actor::BeginPlay()
+{
+	Super::BeginPlay();
+
+	// Most likely case for this is the target actor is no longer net relevant to us and has been destroyed, so this should be destroyed too
+	if (GetOwner())
+	{
+		GetOwner()->OnDestroyed.AddDynamic(this, &AGameplayCueNotify_Actor::OnOwnerDestroyed);
+	}
+}
+
 void AGameplayCueNotify_Actor::PostInitProperties()
 {
 	Super::PostInitProperties();
@@ -102,7 +113,6 @@ void AGameplayCueNotify_Actor::OnOwnerDestroyed()
 {
 	// May need to do extra cleanup in child classes
 	Destroy();
-	//MarkPendingKill();
 }
 
 bool AGameplayCueNotify_Actor::OnExecute_Implementation(AActor* MyTarget, FGameplayCueParameters Parameters)
diff --git a/Engine/Source/Runtime/GameplayAbilities/Public/GameplayCueNotify_Actor.h b/Engine/Source/Runtime/GameplayAbilities/Public/GameplayCueNotify_Actor.h
--- a/Engine/Source/Runtime/GameplayAbilities/Public/GameplayCueNotify_Actor.h
+++ b/Engine/Source/Runtime/GameplayAbilities/Public/GameplayCueNotify_Actor.h
@@ -27,11 +27,14 @@ class GAMEPLAYABILITIES_API AGameplayCueNotify_Actor : public AActor
 	/** Does this GameplayCueNotify handle this type of GameplayCueEvent? */
 	virtual bool HandlesEvent(EGameplayCueEvent::Type EventType) const;
 
+	UFUNCTION()
 	virtual void OnOwnerDestroyed();
 
-	virtual void PostInitProperties();
+	virtual void BeginPlay() override;
 
-	virtual void Serialize(FArchive& Ar);
+	virtual void PostInitProperties() override;
+
+	virtual void Serialize(FArchive& Ar) override;
 
 	virtual void HandleGameplayCue(AActor* MyTarget, EGameplayCueEvent::Type EventType, FGameplayCueParameters Parameters);
 
