Merging

//depot/UE4-Orion/Engine/Source/Runtime/GameplayAbilities/...

to //depot/UE4/Engine/Source/Runtime/GameplayAbilities/...

diff --git a/Engine/Source/Runtime/GameplayAbilities/Private/GameplayCueNotify_Actor.cpp b/Engine/Source/Runtime/GameplayAbilities/Private/GameplayCueNotify_Actor.cpp
--- a/Engine/Source/Runtime/GameplayAbilities/Private/GameplayCueNotify_Actor.cpp
+++ b/Engine/Source/Runtime/GameplayAbilities/Private/GameplayCueNotify_Actor.cpp
@@ -12,6 +12,7 @@ AGameplayCueNotify_Actor::AGameplayCueNotify_Actor(const FObjectInitializer& Obj
 	IsOverride = true;
 	PrimaryActorTick.bCanEverTick = true;
 	bAutoDestroyOnRemove = false;
+	AutoDestroyDelay = 0.f;
 }
 
 #if WITH_EDITOR
@@ -98,7 +99,14 @@ void AGameplayCueNotify_Actor::HandleGameplayCue(AActor* MyTarget, EGameplayCueE
 
 			if (bAutoDestroyOnRemove)
 			{
-				Destroy();
+				if (AutoDestroyDelay > 0.f)
+				{
+					SetLifeSpan(AutoDestroyDelay);
+				}
+				else
+				{
+					Destroy();
+				}
 			}
 			break;
 		};
diff --git a/Engine/Source/Runtime/GameplayAbilities/Public/GameplayCueNotify_Actor.h b/Engine/Source/Runtime/GameplayAbilities/Public/GameplayCueNotify_Actor.h
--- a/Engine/Source/Runtime/GameplayAbilities/Public/GameplayCueNotify_Actor.h
+++ b/Engine/Source/Runtime/GameplayAbilities/Public/GameplayCueNotify_Actor.h
@@ -69,6 +69,10 @@ class GAMEPLAYABILITIES_API AGameplayCueNotify_Actor : public AActor
 	UPROPERTY(EditDefaultsOnly, Category = GameplayCue)
 	bool	bAutoDestroyOnRemove;
 
+	/** If bAutoDestroyOnRemove is true, the actor will stay alive for this many seconds before being auto destroyed. */
+	UPROPERTY(EditAnywhere, Category = GameplayCue)
+	float AutoDestroyDelay;
+
 	/** Does this Cue override other cues, or is it called in addition to them? E.g., If this is Damage.Physical.Slash, we wont call Damage.Physical afer we run this cue. */
 	UPROPERTY(EditDefaultsOnly, Category = GameplayCue)
 	bool IsOverride;
