Merging

//depot/UE4-Orion/Engine/Source/Runtime/GameplayAbilities/Private/GameplayCueNotify_Actor.cpp

to //depot/UE4/Engine/Source/Runtime/GameplayAbilities/Private/GameplayCueNotify_Actor.cpp

diff --git a/Engine/Source/Runtime/GameplayAbilities/Private/GameplayCueNotify_Actor.cpp b/Engine/Source/Runtime/GameplayAbilities/Private/GameplayCueNotify_Actor.cpp
--- a/Engine/Source/Runtime/GameplayAbilities/Private/GameplayCueNotify_Actor.cpp
+++ b/Engine/Source/Runtime/GameplayAbilities/Private/GameplayCueNotify_Actor.cpp
@@ -80,6 +80,9 @@ void AGameplayCueNotify_Actor::HandleGameplayCue(AActor* MyTarget, EGameplayCueE
 	{
 		K2_HandleGameplayCue(MyTarget, EventType, Parameters);
 
+		// Clear any pending auto-destroy that may have occurred from a previous OnRemove
+		SetLifeSpan(0.f);
+
 		switch (EventType)
 		{
 		case EGameplayCueEvent::OnActive:
