Fix for crash is you use AbilitySystem.ToggleDebugHUD when the game isn't running.

diff --git a/Engine/Source/Runtime/GameplayAbilities/Private/AbilitySystemDebugHUD.cpp b/Engine/Source/Runtime/GameplayAbilities/Private/AbilitySystemDebugHUD.cpp
--- a/Engine/Source/Runtime/GameplayAbilities/Private/AbilitySystemDebugHUD.cpp
+++ b/Engine/Source/Runtime/GameplayAbilities/Private/AbilitySystemDebugHUD.cpp
@@ -49,8 +49,12 @@ void AAbilitySystemDebugHUD::DrawDebugHUD(UCanvas* InCanvas, APlayerController*
 	}
 
 	UPlayer * LocalPlayer = GetWorld()->GetFirstLocalPlayerFromController();
-	PC = LocalPlayer->PlayerController;	
+	if (!LocalPlayer)
+	{
+		return;
+	}
 
+	PC = LocalPlayer->PlayerController;	
 	if (!PC)
 	{
 		return;
