Moved network helper function impls from ability task to ability.

diff --git a/Engine/Source/Runtime/GameplayAbilities/Private/Abilities/GameplayAbility.cpp b/Engine/Source/Runtime/GameplayAbilities/Private/Abilities/GameplayAbility.cpp
--- a/Engine/Source/Runtime/GameplayAbilities/Private/Abilities/GameplayAbility.cpp
+++ b/Engine/Source/Runtime/GameplayAbilities/Private/Abilities/GameplayAbility.cpp
@@ -1181,6 +1181,48 @@ bool UGameplayAbility::IsTriggered() const
 	return AbilityTriggers.Num() > 0;
 }
 
+bool UGameplayAbility::IsPredictingClient() const
+{
+	if (GetCurrentActorInfo()->OwnerActor.IsValid())
+	{
+		bool bIsLocallyControlled = GetCurrentActorInfo()->IsLocallyControlled();
+		bool bIsAuthority = GetCurrentActorInfo()->IsNetAuthority();
+
+		if (!bIsAuthority && bIsLocallyControlled && GetNetExecutionPolicy() == EGameplayAbilityNetExecutionPolicy::LocalPredicted)
+		{
+			return true;
+		}
+	}
+
+	return false;
+}
+
+bool UGameplayAbility::IsForRemoteClient() const
+{
+	if (GetCurrentActorInfo()->OwnerActor.IsValid())
+	{
+		bool bIsLocallyControlled = GetCurrentActorInfo()->IsLocallyControlled();
+		bool bIsAuthority = GetCurrentActorInfo()->IsNetAuthority();
+
+		if (bIsAuthority && !bIsLocallyControlled)
+		{
+			return true;
+		}
+	}
+
+	return false;
+}
+
+bool UGameplayAbility::IsLocallyControlled() const
+{
+	if (GetCurrentActorInfo()->OwnerActor.IsValid())
+	{
+		return GetCurrentActorInfo()->IsLocallyControlled();
+	}
+
+	return false;
+}
+
 bool UGameplayAbility::HasAuthority(const FGameplayAbilityActivationInfo* ActivationInfo) const
 {
 	return (ActivationInfo->ActivationMode == EGameplayAbilityActivationMode::Authority);
diff --git a/Engine/Source/Runtime/GameplayAbilities/Private/Abilities/Tasks/AbilityTask.cpp b/Engine/Source/Runtime/GameplayAbilities/Private/Abilities/Tasks/AbilityTask.cpp
--- a/Engine/Source/Runtime/GameplayAbilities/Private/Abilities/Tasks/AbilityTask.cpp
+++ b/Engine/Source/Runtime/GameplayAbilities/Private/Abilities/Tasks/AbilityTask.cpp
@@ -151,48 +151,19 @@ FPredictionKey UAbilityTask::GetActivationPredictionKey() const
 bool UAbilityTask::IsPredictingClient() const
 {
 	UGameplayAbility* MyAbility = Ability.Get();
-	if (MyAbility && MyAbility->GetCurrentActorInfo()->OwnerActor.IsValid())
-	{
-		bool bIsLocallyControlled = MyAbility->GetCurrentActorInfo()->IsLocallyControlled();
-		bool bIsAuthority = MyAbility->GetCurrentActorInfo()->IsNetAuthority();
-
-		if (!bIsAuthority && bIsLocallyControlled && MyAbility->GetNetExecutionPolicy() == EGameplayAbilityNetExecutionPolicy::LocalPredicted)
-		{
-			// If this is a predicted ability, forward up
-			return true;
-		}
-	}
-
-	return false;
+	return MyAbility && MyAbility->IsPredictingClient();
 }
 
 bool UAbilityTask::IsForRemoteClient() const
 {
 	UGameplayAbility* MyAbility = Ability.Get();
-	if (MyAbility && MyAbility->GetCurrentActorInfo()->OwnerActor.IsValid())
-	{
-		bool bIsLocallyControlled = MyAbility->GetCurrentActorInfo()->IsLocallyControlled();
-		bool bIsAuthority = MyAbility->GetCurrentActorInfo()->IsNetAuthority();
-
-		if (bIsAuthority && !bIsLocallyControlled)
-		{
-			// If this is a predicted ability, forward up
-			return true;
-		}
-	}
-
-	return false;
+	return MyAbility && MyAbility->IsForRemoteClient();
 }
 
 bool UAbilityTask::IsLocallyControlled() const
 {
 	UGameplayAbility* MyAbility = Ability.Get();
-	if (MyAbility && MyAbility->GetCurrentActorInfo()->OwnerActor.IsValid())
-	{
-		return MyAbility->GetCurrentActorInfo()->IsLocallyControlled();
-	}
-
-	return false;
+	return MyAbility && MyAbility->IsLocallyControlled();
 }
 
 bool UAbilityTask::CallOrAddReplicatedDelegate(EAbilityGenericReplicatedEvent::Type Event, FSimpleMulticastDelegate::FDelegate Delegate)
diff --git a/Engine/Source/Runtime/GameplayAbilities/Public/Abilities/GameplayAbility.h b/Engine/Source/Runtime/GameplayAbilities/Public/Abilities/GameplayAbility.h
--- a/Engine/Source/Runtime/GameplayAbilities/Public/Abilities/GameplayAbility.h
+++ b/Engine/Source/Runtime/GameplayAbilities/Public/Abilities/GameplayAbility.h
@@ -237,6 +237,12 @@ public:
 	/** Is this ability triggered from TriggerData (or is it triggered explicitly through input/game code) */
 	bool IsTriggered() const;
 
+	bool IsPredictingClient() const;
+
+	bool IsForRemoteClient() const;
+
+	bool IsLocallyControlled() const;
+
 	bool HasAuthority(const FGameplayAbilityActivationInfo* ActivationInfo) const;
 
 	bool HasAuthorityOrPredictionKey(const FGameplayAbilityActorInfo* ActorInfo, const FGameplayAbilityActivationInfo* ActivationInfo) const;
