Slight change to WorldReticle initialization interface for Orion.

diff --git a/Engine/Source/Runtime/GameplayAbilities/Private/Abilities/GameplayAbilityTargetActor_Trace.cpp b/Engine/Source/Runtime/GameplayAbilities/Private/Abilities/GameplayAbilityTargetActor_Trace.cpp
--- a/Engine/Source/Runtime/GameplayAbilities/Private/Abilities/GameplayAbilityTargetActor_Trace.cpp
+++ b/Engine/Source/Runtime/GameplayAbilities/Private/Abilities/GameplayAbilityTargetActor_Trace.cpp
@@ -158,7 +158,7 @@ void AGameplayAbilityTargetActor_Trace::StartTargeting(UGameplayAbility* InAbili
 		AGameplayAbilityWorldReticle* SpawnedReticleActor = GetWorld()->SpawnActor<AGameplayAbilityWorldReticle>(ReticleClass, GetActorLocation(), GetActorRotation());
 		if (SpawnedReticleActor)
 		{
-			SpawnedReticleActor->InitializeReticle(this, ReticleParams);
+			SpawnedReticleActor->InitializeReticle(this, MasterPC, ReticleParams);
 			ReticleActor = SpawnedReticleActor;
 
 			// This is to catch cases of playing on a listen server where we are using a replicated reticle actor.
diff --git a/Engine/Source/Runtime/GameplayAbilities/Private/Abilities/GameplayAbilityWorldReticle.cpp b/Engine/Source/Runtime/GameplayAbilities/Private/Abilities/GameplayAbilityWorldReticle.cpp
--- a/Engine/Source/Runtime/GameplayAbilities/Private/Abilities/GameplayAbilityWorldReticle.cpp
+++ b/Engine/Source/Runtime/GameplayAbilities/Private/Abilities/GameplayAbilityWorldReticle.cpp
@@ -28,11 +28,11 @@ void AGameplayAbilityWorldReticle::Tick(float DeltaTime)
 	FaceTowardSource(bFaceOwnerFlat);
 }
 
-void AGameplayAbilityWorldReticle::InitializeReticle(AGameplayAbilityTargetActor* InTargetingActor, FWorldReticleParameters InParameters)
+void AGameplayAbilityWorldReticle::InitializeReticle(AActor* InTargetingActor, APlayerController* PlayerController, FWorldReticleParameters InParameters)
 {
 	check(InTargetingActor);
 	TargetingActor = InTargetingActor;
-	MasterPC = InTargetingActor->MasterPC;
+	MasterPC = PlayerController;
 	AddTickPrerequisiteActor(TargetingActor);		//We want the reticle to tick after the targeting actor so that designers have the final say on the position
 	Parameters = InParameters;
 	OnParametersInitialized();
@@ -78,7 +78,7 @@ void AGameplayAbilityWorldReticle::FaceTowardSource(bool bFaceIn2D)
 	{
 		if (bFaceIn2D)
 		{
-			FVector FacingVector = (TargetingActor->StartLocation.GetTargetingTransform().GetLocation() - GetActorLocation()).GetSafeNormal2D();
+			FVector FacingVector = (TargetingActor->GetActorLocation() - GetActorLocation()).GetSafeNormal2D();
 			if (FacingVector.IsZero())
 			{
 				FacingVector = -GetActorForwardVector().GetSafeNormal2D();
@@ -90,7 +90,7 @@ void AGameplayAbilityWorldReticle::FaceTowardSource(bool bFaceIn2D)
 		}
 		else
 		{
-			FVector FacingVector = (TargetingActor->StartLocation.GetTargetingTransform().GetLocation() - GetActorLocation()).GetSafeNormal();
+			FVector FacingVector = (TargetingActor->GetActorLocation() - GetActorLocation()).GetSafeNormal();
 			if (FacingVector.IsZero())
 			{
 				FacingVector = -GetActorForwardVector().GetSafeNormal();
diff --git a/Engine/Source/Runtime/GameplayAbilities/Public/Abilities/GameplayAbilityWorldReticle.h b/Engine/Source/Runtime/GameplayAbilities/Public/Abilities/GameplayAbilityWorldReticle.h
--- a/Engine/Source/Runtime/GameplayAbilities/Public/Abilities/GameplayAbilityWorldReticle.h
+++ b/Engine/Source/Runtime/GameplayAbilities/Public/Abilities/GameplayAbilityWorldReticle.h
@@ -37,7 +37,7 @@ public:
 
 	virtual bool IsNetRelevantFor(const APlayerController* RealViewer, const AActor* ViewTarget, const FVector& SrcLocation) const override;
 
-	void InitializeReticle(AGameplayAbilityTargetActor* InTargetingActor, FWorldReticleParameters InParameters);
+	void InitializeReticle(AActor* InTargetingActor, APlayerController* PlayerController, FWorldReticleParameters InParameters);
 
 	void SetIsTargetValid(bool bNewValue);
 	void SetIsTargetAnActor(bool bNewValue);
@@ -88,5 +88,5 @@ protected:
 
 	/** In the future, we may want to grab things like sockets off of this. */
 	UPROPERTY(BlueprintReadOnly, Category = "Network")
-	AGameplayAbilityTargetActor* TargetingActor;
+	AActor* TargetingActor;
 };
