Spot integration of CL 2504277.

Removed deprecated Duration variable and associated code from UGameplayEffect.

diff --git a/Engine/Source/Editor/GameplayAbilitiesEditor/Private/AssetTypeActions_GameplayEffect.cpp b/Engine/Source/Editor/GameplayAbilitiesEditor/Private/AssetTypeActions_GameplayEffect.cpp
--- a/Engine/Source/Editor/GameplayAbilitiesEditor/Private/AssetTypeActions_GameplayEffect.cpp
+++ b/Engine/Source/Editor/GameplayAbilitiesEditor/Private/AssetTypeActions_GameplayEffect.cpp
@@ -13,8 +13,8 @@ DEFINE_LOG_CATEGORY(LogAssetTypeActions_GameplayEffect);
 
 void CopyInfoFromDataGE(const UGameplayEffect& InGE, UGameplayEffect& OutGE)
 {
+	OutGE.DurationPolicy = InGE.DurationPolicy;
 	OutGE.DurationMagnitude = InGE.DurationMagnitude;
-	OutGE.Duration = InGE.Duration;
 	OutGE.Period = InGE.Period;
 
 	for (FGameplayModifierInfo InModifier : InGE.Modifiers)
diff --git a/Engine/Source/Runtime/GameplayAbilities/Private/GameplayEffect.cpp b/Engine/Source/Runtime/GameplayAbilities/Private/GameplayEffect.cpp
--- a/Engine/Source/Runtime/GameplayAbilities/Private/GameplayEffect.cpp
+++ b/Engine/Source/Runtime/GameplayAbilities/Private/GameplayEffect.cpp
@@ -91,28 +91,6 @@ void UGameplayEffect::PostLoad()
 	ClearTagsContainer.RemoveAllTags();
 
 	UpdateInheritedTagProperties();
-
-	if (Duration.Curve.CurveTable != nullptr || Duration.Value != 0)
-	{
-		if (Duration.Value == INFINITE_DURATION)
-		{
-			DurationPolicy = EGameplayEffectDurationType::Infinite;
-		}
-		else if (Duration.Value == INSTANT_APPLICATION)
-		{
-			DurationPolicy = EGameplayEffectDurationType::Instant;
-		}
-		else
-		{
-			DurationPolicy = EGameplayEffectDurationType::HasDuration;
-		}
-
-		DurationMagnitude.ScalableFloatMagnitude = Duration;
-		DurationMagnitude.MagnitudeCalculationType = EGameplayEffectMagnitudeCalculation::ScalableFloat;
-
-		Duration.Curve.CurveTable = nullptr;
-		Duration.Value = 0;
-	}
 }
 
 void UGameplayEffect::PostInitProperties()
diff --git a/Engine/Source/Runtime/GameplayAbilities/Public/GameplayEffect.h b/Engine/Source/Runtime/GameplayAbilities/Public/GameplayEffect.h
--- a/Engine/Source/Runtime/GameplayAbilities/Public/GameplayEffect.h
+++ b/Engine/Source/Runtime/GameplayAbilities/Public/GameplayEffect.h
@@ -551,10 +551,6 @@ public:
 	UPROPERTY(EditDefaultsOnly, Category=GameplayEffect)
 	FGameplayEffectModifierMagnitude DurationMagnitude;
 
-	/** Deprecated. Use DurationMagnitude instead. */
-	UPROPERTY()
-	FScalableFloat	Duration;
-
 	/** Period in seconds. 0.0 for non-periodic effects */
 	UPROPERTY(EditDefaultsOnly, BlueprintReadOnly, Category=Period)
 	FScalableFloat	Period;
