AbilitySystem merges

Misc
-Get end point transforms (not just fvector) on target data
-GameplayCue TagRequirements blueprint function
-Blueprint accessor for SourceObject in GameplayEffectContext
-Cvar for displaying GameplayCues (AbilitySystem.DisplayGameplayCues)
-Saved a few places where we were doing work even if VLog was off or compiled out.
-Dont reuse GameplayCue actors if they are scheduled to be destroyed

GameplayEffects
-Dont invoke gameplay cues from GEs if the GE is inhibited (and reinvoke them if it becomes uninhibited)
-Build ModifiedAttributeList (for gameplaycues) for duration based GEs
-Queue GameplayCue events on replication until the entire AGE array has been replicated
-Added gameplay effect query option for specific GE class

-Optimized FGameplayEffectContext::NetSerialize

Abilities
-Added UGameplayAbility::IsEndAbilityValid

diff --git a/Engine/Source/Runtime/GameplayAbilities/Private/Abilities/GameplayAbility.cpp b/Engine/Source/Runtime/GameplayAbilities/Private/Abilities/GameplayAbility.cpp
--- a/Engine/Source/Runtime/GameplayAbilities/Private/Abilities/GameplayAbility.cpp
+++ b/Engine/Source/Runtime/GameplayAbilities/Private/Abilities/GameplayAbility.cpp
@@ -482,18 +482,32 @@ void UGameplayAbility::CancelAbility(const FGameplayAbilitySpecHandle Handle, co
 	}
 }
 
-void UGameplayAbility::EndAbility(const FGameplayAbilitySpecHandle Handle, const FGameplayAbilityActorInfo* ActorInfo, const FGameplayAbilityActivationInfo ActivationInfo, bool bReplicateEndAbility)
+bool UGameplayAbility::IsEndAbilityValid(const FGameplayAbilitySpecHandle Handle, const FGameplayAbilityActorInfo* ActorInfo) const
 {
+	check(ActorInfo);
+
 	// Protect against EndAbility being called multiple times
 	// Ending an AbilityState may cause this to be invoked again
 	if (bIsActive == false && GetInstancingPolicy() != EGameplayAbilityInstancingPolicy::NonInstanced)
 	{
-		return;
+		return false;
 	}
 
 	// check to see if this is an NonInstanced or if the ability is active.
-	FGameplayAbilitySpec* Spec = ActorInfo ? ActorInfo->AbilitySystemComponent->FindAbilitySpecFromHandle(Handle) : nullptr;
-	if ((Spec != nullptr) ? Spec->IsActive() : IsActive())
+	const FGameplayAbilitySpec* Spec = ActorInfo ? ActorInfo->AbilitySystemComponent->FindAbilitySpecFromHandle(Handle) : nullptr;
+	const bool bIsSpecActive = (Spec != nullptr) ? Spec->IsActive() : IsActive();
+
+	if (!bIsSpecActive)
+	{
+		return false;
+	}
+
+	return true;
+}
+
+void UGameplayAbility::EndAbility(const FGameplayAbilitySpecHandle Handle, const FGameplayAbilityActorInfo* ActorInfo, const FGameplayAbilityActivationInfo ActivationInfo, bool bReplicateEndAbility)
+{
+	if (IsEndAbilityValid(Handle, ActorInfo))
 	{
 		// Give blueprint a chance to react
 		K2_OnEndAbility();
@@ -671,11 +685,11 @@ bool UGameplayAbility::CheckCooldown(const FGameplayAbilitySpecHandle Handle, co
 		check(ActorInfo->AbilitySystemComponent.IsValid());
 		if (CooldownTags->Num() > 0 && ActorInfo->AbilitySystemComponent->HasAnyMatchingGameplayTags(*CooldownTags))
 		{
-			const FGameplayTag& CoooldownTag = UAbilitySystemGlobals::Get().ActivateFailCooldownTag;
+			const FGameplayTag& CooldownTag = UAbilitySystemGlobals::Get().ActivateFailCooldownTag;
 
-			if (OptionalRelevantTags && CoooldownTag.IsValid())
+			if (OptionalRelevantTags && CooldownTag.IsValid())
 			{
-				OptionalRelevantTags->AddTag(CoooldownTag);
+				OptionalRelevantTags->AddTag(CooldownTag);
 			}
 
 			return false;
@@ -1188,7 +1202,8 @@ bool UGameplayAbility::IsPredictingClient() const
 		bool bIsLocallyControlled = GetCurrentActorInfo()->IsLocallyControlled();
 		bool bIsAuthority = GetCurrentActorInfo()->IsNetAuthority();
 
-		if (!bIsAuthority && bIsLocallyControlled && GetNetExecutionPolicy() == EGameplayAbilityNetExecutionPolicy::LocalPredicted)
+		// LocalPredicted and ServerInitiated are both valid because in both those modes the ability also runs on the client
+		if (!bIsAuthority && bIsLocallyControlled && (GetNetExecutionPolicy() == EGameplayAbilityNetExecutionPolicy::LocalPredicted || GetNetExecutionPolicy() == EGameplayAbilityNetExecutionPolicy::ServerInitiated))
 		{
 			return true;
 		}
diff --git a/Engine/Source/Runtime/GameplayAbilities/Private/Abilities/GameplayAbilityTargetActor_Radius.cpp b/Engine/Source/Runtime/GameplayAbilities/Private/Abilities/GameplayAbilityTargetActor_Radius.cpp
--- a/Engine/Source/Runtime/GameplayAbilities/Private/Abilities/GameplayAbilityTargetActor_Radius.cpp
+++ b/Engine/Source/Runtime/GameplayAbilities/Private/Abilities/GameplayAbilityTargetActor_Radius.cpp
@@ -55,8 +55,8 @@ TArray<TWeakObjectPtr<AActor> >	AGameplayAbilityTargetActor_Radius::PerformOverl
 	bool bTraceComplex = false;
 	
 	FCollisionQueryParams Params(RadiusTargetingOverlap, bTraceComplex);
-	Params.bReturnPhysicalMaterial = true;
-	Params.bTraceAsyncScene = true;
+	Params.bReturnPhysicalMaterial = false;
+	Params.bTraceAsyncScene = false;
 
 	TArray<FOverlapResult> Overlaps;
 
@@ -75,4 +75,4 @@ TArray<TWeakObjectPtr<AActor> >	AGameplayAbilityTargetActor_Radius::PerformOverl
 	}
 
 	return HitActors;
-}
\ No newline at end of file
+}
diff --git a/Engine/Source/Runtime/GameplayAbilities/Private/Abilities/Tasks/AbilityTask_PlayMontageAndWait.cpp b/Engine/Source/Runtime/GameplayAbilities/Private/Abilities/Tasks/AbilityTask_PlayMontageAndWait.cpp
--- a/Engine/Source/Runtime/GameplayAbilities/Private/Abilities/Tasks/AbilityTask_PlayMontageAndWait.cpp
+++ b/Engine/Source/Runtime/GameplayAbilities/Private/Abilities/Tasks/AbilityTask_PlayMontageAndWait.cpp
@@ -66,14 +66,20 @@ UAbilityTask_PlayMontageAndWait* UAbilityTask_PlayMontageAndWait::CreatePlayMont
 
 void UAbilityTask_PlayMontageAndWait::Activate()
 {
+	UGameplayAbility* MyAbility = Ability.Get();
+	if (MyAbility == nullptr)
+	{
+		return;
+	}
+
 	bool bPlayedMontage = false;
 
-	if (AbilitySystemComponent.IsValid() && Ability.IsValid())
+	if (AbilitySystemComponent.IsValid())
 	{
-		const FGameplayAbilityActorInfo* ActorInfo = Ability->GetCurrentActorInfo();
+		const FGameplayAbilityActorInfo* ActorInfo = MyAbility->GetCurrentActorInfo();
 		if (ActorInfo->AnimInstance.IsValid())
 		{
-			if (AbilitySystemComponent->PlayMontage(Ability.Get(), Ability->GetCurrentActivationInfo(), MontageToPlay, Rate, StartSection) > 0.f)
+			if (AbilitySystemComponent->PlayMontage(MyAbility, MyAbility->GetCurrentActivationInfo(), MontageToPlay, Rate, StartSection) > 0.f)
 			{
 				// Playing a montage could potentially fire off a callback into game code which could kill this ability! Early out if we are  pending kill.
 				if (IsPendingKill())
@@ -82,7 +88,7 @@ void UAbilityTask_PlayMontageAndWait::Activate()
 					return;
 				}
 
-				InterruptedHandle = Ability->OnGameplayAbilityCancelled.AddUObject(this, &UAbilityTask_PlayMontageAndWait::OnMontageInterrupted);
+				InterruptedHandle = MyAbility->OnGameplayAbilityCancelled.AddUObject(this, &UAbilityTask_PlayMontageAndWait::OnMontageInterrupted);
 
 				BlendingOutDelegate.BindUObject(this, &UAbilityTask_PlayMontageAndWait::OnMontageBlendingOut);
 				ActorInfo->AnimInstance->Montage_SetBlendingOutDelegate(BlendingOutDelegate);
@@ -94,7 +100,7 @@ void UAbilityTask_PlayMontageAndWait::Activate()
 
 	if (!bPlayedMontage)
 	{
-		ABILITY_LOG(Warning, TEXT("UAbilityTask_PlayMontageAndWait called in Ability %s failed to play montage; Task Instance Name %s."), *Ability->GetName(), *InstanceName.ToString());
+		ABILITY_LOG(Warning, TEXT("UAbilityTask_PlayMontageAndWait called in Ability %s failed to play montage; Task Instance Name %s."), *MyAbility->GetName(), *InstanceName.ToString());
 		OnCancelled.Broadcast();
 	}
 }
diff --git a/Engine/Source/Runtime/GameplayAbilities/Private/Abilities/Tasks/AbilityTask_WaitAbilityActivate.cpp b/Engine/Source/Runtime/GameplayAbilities/Private/Abilities/Tasks/AbilityTask_WaitAbilityActivate.cpp
--- a/Engine/Source/Runtime/GameplayAbilities/Private/Abilities/Tasks/AbilityTask_WaitAbilityActivate.cpp
+++ b/Engine/Source/Runtime/GameplayAbilities/Private/Abilities/Tasks/AbilityTask_WaitAbilityActivate.cpp
@@ -13,12 +13,13 @@ UAbilityTask_WaitAbilityActivate::UAbilityTask_WaitAbilityActivate(const FObject
 	IncludeTriggeredAbilities = false;
 }
 
-UAbilityTask_WaitAbilityActivate* UAbilityTask_WaitAbilityActivate::WaitForAbilityActivate(UObject* WorldContextObject, FGameplayTag InWithTag, FGameplayTag InWithoutTag, bool InIncludeTriggeredAbilities)
+UAbilityTask_WaitAbilityActivate* UAbilityTask_WaitAbilityActivate::WaitForAbilityActivate(UObject* WorldContextObject, FGameplayTag InWithTag, FGameplayTag InWithoutTag, bool InIncludeTriggeredAbilities, bool InTriggerOnce)
 {
 	auto MyObj = NewTask<UAbilityTask_WaitAbilityActivate>(WorldContextObject);
 	MyObj->WithTag = InWithTag;
 	MyObj->WithoutTag = InWithoutTag;
 	MyObj->IncludeTriggeredAbilities = InIncludeTriggeredAbilities;
+	MyObj->TriggerOnce = InTriggerOnce;
 	return MyObj;
 }
 
@@ -46,7 +47,10 @@ void UAbilityTask_WaitAbilityActivate::OnAbilityActivate(UGameplayAbility* Activ
 
 	OnActivate.Broadcast(ActivatedAbility);
 
-	EndTask();
+	if (TriggerOnce)
+	{
+		EndTask();
+	}
 }
 
 void UAbilityTask_WaitAbilityActivate::OnDestroy(bool AbilityEnded)
diff --git a/Engine/Source/Runtime/GameplayAbilities/Private/Abilities/Tasks/AbilityTask_WaitAbilityCommit.cpp b/Engine/Source/Runtime/GameplayAbilities/Private/Abilities/Tasks/AbilityTask_WaitAbilityCommit.cpp
--- a/Engine/Source/Runtime/GameplayAbilities/Private/Abilities/Tasks/AbilityTask_WaitAbilityCommit.cpp
+++ b/Engine/Source/Runtime/GameplayAbilities/Private/Abilities/Tasks/AbilityTask_WaitAbilityCommit.cpp
@@ -13,11 +13,12 @@ UAbilityTask_WaitAbilityCommit::UAbilityTask_WaitAbilityCommit(const FObjectInit
 {
 }
 
-UAbilityTask_WaitAbilityCommit* UAbilityTask_WaitAbilityCommit::WaitForAbilityCommit(UObject* WorldContextObject, FGameplayTag InWithTag, FGameplayTag InWithoutTag)
+UAbilityTask_WaitAbilityCommit* UAbilityTask_WaitAbilityCommit::WaitForAbilityCommit(UObject* WorldContextObject, FGameplayTag InWithTag, FGameplayTag InWithoutTag, bool InTriggerOnce)
 {
 	auto MyObj = NewTask<UAbilityTask_WaitAbilityCommit>(WorldContextObject);
 	MyObj->WithTag = InWithTag;
 	MyObj->WithoutTag = InWithoutTag;
+	MyObj->TriggerOnce = InTriggerOnce;
 
 	return MyObj;
 }
@@ -51,5 +52,8 @@ void UAbilityTask_WaitAbilityCommit::OnAbilityCommit(UGameplayAbility *Activated
 
 	OnCommit.Broadcast(ActivatedAbility);
 
-	EndTask();
+	if (TriggerOnce)
+	{
+		EndTask();
+	}
 }
\ No newline at end of file
diff --git a/Engine/Source/Runtime/GameplayAbilities/Private/Abilities/Tasks/AbilityTask_WaitCancel.cpp b/Engine/Source/Runtime/GameplayAbilities/Private/Abilities/Tasks/AbilityTask_WaitCancel.cpp
new file mode 100644
--- /dev/null
+++ b/Engine/Source/Runtime/GameplayAbilities/Private/Abilities/Tasks/AbilityTask_WaitCancel.cpp
@@ -0,0 +1,74 @@
+// Copyright 1998-2015 Epic Games, Inc. All Rights Reserved.
+
+#include "AbilitySystemPrivatePCH.h"
+#include "Abilities/Tasks/AbilityTask_WaitCancel.h"
+
+#include "AbilitySystemComponent.h"
+
+#include "Abilities/GameplayAbility.h"
+
+UAbilityTask_WaitCancel::UAbilityTask_WaitCancel(const FObjectInitializer& ObjectInitializer)
+	: Super(ObjectInitializer)
+{
+	RegisteredCallbacks = false;
+
+}
+
+void UAbilityTask_WaitCancel::OnCancelCallback()
+{
+	if (AbilitySystemComponent.IsValid())
+	{
+		AbilitySystemComponent->ConsumeGenericReplicatedEvent(EAbilityGenericReplicatedEvent::GenericCancel, GetAbilitySpecHandle(), GetActivationPredictionKey());
+		OnCancel.Broadcast();
+		EndTask();
+	}
+}
+
+void UAbilityTask_WaitCancel::OnLocalCancelCallback()
+{
+	if (AbilitySystemComponent.IsValid() && IsPredictingClient())
+	{
+		AbilitySystemComponent->ServerSetReplicatedEvent(EAbilityGenericReplicatedEvent::GenericCancel, GetAbilitySpecHandle(), GetActivationPredictionKey() ,AbilitySystemComponent->ScopedPredictionKey);
+	}
+	OnCancelCallback();
+}
+
+UAbilityTask_WaitCancel* UAbilityTask_WaitCancel::WaitCancel(UObject* WorldContextObject)
+{
+	return NewTask<UAbilityTask_WaitCancel>(WorldContextObject);
+}
+
+void UAbilityTask_WaitCancel::Activate()
+{
+	if (AbilitySystemComponent.IsValid())
+	{
+		const FGameplayAbilityActorInfo* Info = Ability.Get()->GetCurrentActorInfo();
+
+		
+		if (Info->IsLocallyControlled())
+		{
+			// We have to wait for the callback from the AbilitySystemComponent.
+			AbilitySystemComponent->GenericLocalCancelCallbacks.AddDynamic(this, &UAbilityTask_WaitCancel::OnLocalCancelCallback);	// Tell me if the cancel input is pressed
+
+			RegisteredCallbacks = true;
+		}
+		else
+		{
+			if (CallOrAddReplicatedDelegate(EAbilityGenericReplicatedEvent::GenericCancel,  FSimpleMulticastDelegate::FDelegate::CreateUObject(this, &UAbilityTask_WaitCancel::OnCancelCallback)) )
+			{
+				// GenericCancel was already received from the client and we just called OnCancelCallback. The task is done.
+				return;
+			}
+		}
+	}
+}
+
+void UAbilityTask_WaitCancel::OnDestroy(bool AbilityEnding)
+{
+	if (RegisteredCallbacks && AbilitySystemComponent.IsValid())
+	{
+		AbilitySystemComponent->GenericLocalCancelCallbacks.RemoveDynamic(this, &UAbilityTask_WaitCancel::OnLocalCancelCallback);
+	}
+
+	Super::OnDestroy(AbilityEnding);
+}
diff --git a/Engine/Source/Runtime/GameplayAbilities/Private/Abilities/Tasks/AbilityTask_WaitGameplayEffectApplied_Self.cpp b/Engine/Source/Runtime/GameplayAbilities/Private/Abilities/Tasks/AbilityTask_WaitGameplayEffectApplied_Self.cpp
--- a/Engine/Source/Runtime/GameplayAbilities/Private/Abilities/Tasks/AbilityTask_WaitGameplayEffectApplied_Self.cpp
+++ b/Engine/Source/Runtime/GameplayAbilities/Private/Abilities/Tasks/AbilityTask_WaitGameplayEffectApplied_Self.cpp
@@ -13,7 +13,7 @@ UAbilityTask_WaitGameplayEffectApplied_Self::UAbilityTask_WaitGameplayEffectAppl
 {
 }
 
-UAbilityTask_WaitGameplayEffectApplied_Self* UAbilityTask_WaitGameplayEffectApplied_Self::WaitGameplayEffectAppliedToSelf(UObject* WorldContextObject, AActor* TestActor, const FGameplayTargetDataFilterHandle InFilter, FGameplayTagRequirements InSourceTagRequirements, FGameplayTagRequirements InTargetTagRequirements, bool InTriggerOnce, AActor* OptionalExternalOwner)
+UAbilityTask_WaitGameplayEffectApplied_Self* UAbilityTask_WaitGameplayEffectApplied_Self::WaitGameplayEffectAppliedToSelf(UObject* WorldContextObject, const FGameplayTargetDataFilterHandle InFilter, FGameplayTagRequirements InSourceTagRequirements, FGameplayTagRequirements InTargetTagRequirements, bool InTriggerOnce, AActor* OptionalExternalOwner)
 {
 	auto MyObj = NewTask<UAbilityTask_WaitGameplayEffectApplied_Self>(WorldContextObject);
 	MyObj->Filter = InFilter;
diff --git a/Engine/Source/Runtime/GameplayAbilities/Private/AbilitySystemBlueprintLibrary.cpp b/Engine/Source/Runtime/GameplayAbilities/Private/AbilitySystemBlueprintLibrary.cpp
--- a/Engine/Source/Runtime/GameplayAbilities/Private/AbilitySystemBlueprintLibrary.cpp
+++ b/Engine/Source/Runtime/GameplayAbilities/Private/AbilitySystemBlueprintLibrary.cpp
@@ -321,6 +321,21 @@ FVector UAbilitySystemBlueprintLibrary::GetTargetDataEndPoint(FGameplayAbilityTa
 	return FVector::ZeroVector;
 }
 
+FTransform UAbilitySystemBlueprintLibrary::GetTargetDataEndPointTransform(FGameplayAbilityTargetDataHandle TargetData, int32 Index)
+{
+	if (TargetData.Data.IsValidIndex(Index))
+	{
+		FGameplayAbilityTargetData* Data = TargetData.Data[Index].Get();
+		if (Data)
+		{
+			return Data->GetEndPointTransform();
+		}
+	}
+
+	return FTransform::Identity;
+}
+
+
 // -------------------------------------------------------------------------------------
 
 
@@ -359,6 +374,11 @@ AActor*	UAbilitySystemBlueprintLibrary::EffectContextGetEffectCauser(FGameplayEf
 	return EffectContext.GetEffectCauser();
 }
 
+UObject* UAbilitySystemBlueprintLibrary::EffectContextGetSourceObject(FGameplayEffectContextHandle EffectContext)
+{
+	return const_cast<UObject*>( EffectContext.GetSourceObject() );
+}
+
 FVector UAbilitySystemBlueprintLibrary::EffectContextGetOrigin(FGameplayEffectContextHandle EffectContext)
 {
 	if (EffectContext.HasOrigin())
@@ -491,6 +511,12 @@ bool UAbilitySystemBlueprintLibrary::GetGameplayCueDirection(AActor* TargetActor
 	return false;
 }
 
+bool UAbilitySystemBlueprintLibrary::DoesGameplayCueMeetTagRequirements(FGameplayCueParameters Parameters, FGameplayTagRequirements& SourceTagReqs, FGameplayTagRequirements& TargetTagReqs)
+{
+	return SourceTagReqs.RequirementsMet(Parameters.AggregatedSourceTags)
+		&& TargetTagReqs.RequirementsMet(Parameters.AggregatedSourceTags);
+}
+
 // ---------------------------------------------------------------------------------------
 
 FGameplayEffectSpecHandle UAbilitySystemBlueprintLibrary::AssignSetByCallerMagnitude(FGameplayEffectSpecHandle SpecHandle, FName DataName, float Magnitude)
diff --git a/Engine/Source/Runtime/GameplayAbilities/Private/AbilitySystemComponent.cpp b/Engine/Source/Runtime/GameplayAbilities/Private/AbilitySystemComponent.cpp
--- a/Engine/Source/Runtime/GameplayAbilities/Private/AbilitySystemComponent.cpp
+++ b/Engine/Source/Runtime/GameplayAbilities/Private/AbilitySystemComponent.cpp
@@ -19,15 +19,6 @@ DEFINE_LOG_CATEGORY(LogAbilitySystemComponent);
 
 #define LOCTEXT_NAMESPACE "AbilitySystemComponent"
 
-
-int32 DebugGameplayCues = 0;
-static FAutoConsoleVariableRef CVarDebugGameplayCues(
-	TEXT("AbilitySystem.DebugGameplayCues"),
-	DebugGameplayCues,
-	TEXT("Enables Debugging for GameplayCue events"),
-	ECVF_Default
-	);
-
 /** Enable to log out all render state create, destroy and updatetransform events */
 #define LOG_RENDER_STATE 0
 
@@ -580,13 +571,15 @@ FActiveGameplayEffectHandle UAbilitySystemComponent::ApplyGameplayEffectSpecToSe
 	FActiveGameplayEffectHandle	MyHandle;
 	bool bInvokeGameplayCueApplied = UGameplayEffect::INSTANT_APPLICATION != Spec.GetDuration(); // Cache this now before possibly modifying predictive instant effect to infinite duration effect.
 
+	FActiveGameplayEffect* AppliedEffect = nullptr;
+
 	FGameplayEffectSpec* OurCopyOfSpec = nullptr;
 	TSharedPtr<FGameplayEffectSpec> StackSpec;
 	float Duration = bTreatAsInfiniteDuration ? UGameplayEffect::INFINITE_DURATION : Spec.GetDuration();
 	{
 		if (Duration != UGameplayEffect::INSTANT_APPLICATION)
 		{
-			FActiveGameplayEffect* AppliedEffect = ActiveGameplayEffects.ApplyGameplayEffectSpec(Spec, PredictionKey);
+			AppliedEffect = ActiveGameplayEffects.ApplyGameplayEffectSpec(Spec, PredictionKey);
 			if (!AppliedEffect)
 			{
 				return FActiveGameplayEffectHandle();
@@ -595,13 +588,17 @@ FActiveGameplayEffectHandle UAbilitySystemComponent::ApplyGameplayEffectSpecToSe
 			MyHandle = AppliedEffect->Handle;
 			OurCopyOfSpec = &(AppliedEffect->Spec);
 
-			ABILITY_VLOG(OwnerActor, Log, TEXT("Applied %s"), *OurCopyOfSpec->Def->GetFName().ToString());
-
-			for (FGameplayModifierInfo Modifier : Spec.Def->Modifiers)
+			// Log results of applied GE spec
+			if (UE_LOG_ACTIVE(VLogAbilitySystem, Log))
 			{
-				float Magnitude = 0.f;
-				Modifier.ModifierMagnitude.AttemptCalculateMagnitude(Spec, Magnitude);
-				ABILITY_VLOG(OwnerActor, Log, TEXT("         %s: %s %f"), *Modifier.Attribute.GetName(), *EGameplayModOpToString(Modifier.ModifierOp), Magnitude);
+				ABILITY_VLOG(OwnerActor, Log, TEXT("Applied %s"), *OurCopyOfSpec->Def->GetFName().ToString());
+
+				for (FGameplayModifierInfo Modifier : Spec.Def->Modifiers)
+				{
+					float Magnitude = 0.f;
+					Modifier.ModifierMagnitude.AttemptCalculateMagnitude(Spec, Magnitude);
+					ABILITY_VLOG(OwnerActor, Log, TEXT("         %s: %s %f"), *Modifier.Attribute.GetName(), *EGameplayModOpToString(Modifier.ModifierOp), Magnitude);
+				}
 			}
 		}
 
@@ -621,16 +618,29 @@ FActiveGameplayEffectHandle UAbilitySystemComponent::ApplyGameplayEffectSpecToSe
 		}
 	}
 	
+
 	// We still probably want to apply tags and stuff even if instant?
-	if (bInvokeGameplayCueApplied)
+	if (bInvokeGameplayCueApplied && AppliedEffect && !AppliedEffect->bIsInhibited)
 	{
 		// We both added and activated the GameplayCue here.
 		// On the client, who will invoke the gameplay cue from an OnRep, he will need to look at the StartTime to determine
 		// if the Cue was actually added+activated or just added (due to relevancy)
 
 		// Fixme: what if we wanted to scale Cue magnitude based on damage? E.g, scale an cue effect when the GE is buffed?
-		InvokeGameplayCueEvent(*OurCopyOfSpec, EGameplayCueEvent::OnActive);
-		InvokeGameplayCueEvent(*OurCopyOfSpec, EGameplayCueEvent::WhileActive);
+
+		if (OurCopyOfSpec->StackCount > Spec.StackCount)
+		{
+			// Because PostReplicatedChange will get called from modifying the stack count
+			// (and not PostReplicatedAdd) we won't know which GE was modified.
+			// So instead we need to explicitly RPC the client so it knows the GC needs updating
+			NetMulticast_InvokeGameplayCueAddedAndWhileActive_FromSpec(*OurCopyOfSpec, PredictionKey);
+		}
+		else
+		{
+			// Otherwise these will get replicated to the client when the GE gets added to the replicated array
+			InvokeGameplayCueEvent(*OurCopyOfSpec, EGameplayCueEvent::OnActive);
+			InvokeGameplayCueEvent(*OurCopyOfSpec, EGameplayCueEvent::WhileActive);
+		}
 	}
 	
 	// Execute the GE at least once (if instant, this will execute once and be done. If persistent, it was added to ActiveGameplayEffects above)
@@ -736,13 +746,16 @@ void UAbilitySystemComponent::ExecuteGameplayEffect(FGameplayEffectSpec &Spec, F
 	// Effects with no period and that aren't instant application should never be executed
 	check( (Spec.GetDuration() == UGameplayEffect::INSTANT_APPLICATION || Spec.GetPeriod() != UGameplayEffect::NO_PERIOD) );
 
-	ABILITY_VLOG(OwnerActor, Log, TEXT("Executed %s"), *Spec.Def->GetFName().ToString());
-	
-	for (FGameplayModifierInfo Modifier : Spec.Def->Modifiers)
+	if (UE_LOG_ACTIVE(VLogAbilitySystem, Log))
 	{
-		float Magnitude = 0.f;
-		Modifier.ModifierMagnitude.AttemptCalculateMagnitude(Spec, Magnitude);
-		ABILITY_VLOG(OwnerActor, Log, TEXT("         %s: %s %f"), *Modifier.Attribute.GetName(), *EGameplayModOpToString(Modifier.ModifierOp), Magnitude);
+		ABILITY_VLOG(OwnerActor, Log, TEXT("Executed %s"), *Spec.Def->GetFName().ToString());
+		
+		for (FGameplayModifierInfo Modifier : Spec.Def->Modifiers)
+		{
+			float Magnitude = 0.f;
+			Modifier.ModifierMagnitude.AttemptCalculateMagnitude(Spec, Magnitude);
+			ABILITY_VLOG(OwnerActor, Log, TEXT("         %s: %s %f"), *Modifier.Attribute.GetName(), *EGameplayModOpToString(Modifier.ModifierOp), Magnitude);
+		}
 	}
 
 	ActiveGameplayEffects.ExecuteActiveEffectsFrom(Spec, PredictionKey);
@@ -840,16 +853,13 @@ void UAbilitySystemComponent::InvokeGameplayCueEvent(const FGameplayEffectSpecFo
 		ABILITY_LOG(Warning, TEXT("InvokeGameplayCueEvent Actor %s that has no gameplay effect!"), ActorAvatar ? *ActorAvatar->GetName() : TEXT("NULL"));
 		return;
 	}
-
-	if (DebugGameplayCues)
-	{
-		ABILITY_LOG(Warning, TEXT("InvokeGameplayCueEvent: %s"), *Spec.ToSimpleString());
-	}
 	
 	float ExecuteLevel = Spec.GetLevel();
 
 	FGameplayCueParameters CueParameters;
 	CueParameters.EffectContext = Spec.GetContext();
+	CueParameters.AggregatedSourceTags = Spec.AggregatedSourceTags;
+	CueParameters.AggregatedTargetTags = Spec.AggregatedTargetTags;
 
 	for (FGameplayEffectCue CueInfo : Spec.Def->GameplayCues)
 	{
@@ -872,12 +882,6 @@ void UAbilitySystemComponent::InvokeGameplayCueEvent(const FGameplayEffectSpecFo
 		CueParameters.NormalizedMagnitude = CueInfo.NormalizeLevel(ExecuteLevel);
 
 		UAbilitySystemGlobals::Get().GetGameplayCueManager()->HandleGameplayCues(ActorAvatar, CueInfo.GameplayCueTags, EventType, CueParameters);
-
-		if (DebugGameplayCues && Spec.GetContext().GetHitResult())
-		{
-			DrawDebugSphere(GetWorld(), Spec.GetContext().GetHitResult()->Location, 30.f, 32, FColor::Red, true, 30.f);
-			ABILITY_LOG(Warning, TEXT("   %s"), *CueInfo.GameplayCueTags.ToString());
-		}
 	}
 }
 
@@ -1001,6 +1005,15 @@ void UAbilitySystemComponent::NetMulticast_InvokeGameplayCueAdded_Implementation
 	}
 }
 
+void UAbilitySystemComponent::NetMulticast_InvokeGameplayCueAddedAndWhileActive_FromSpec_Implementation(const FGameplayEffectSpecForRPC& Spec, FPredictionKey PredictionKey)
+{
+	if (IsOwnerActorAuthoritative() || PredictionKey.IsLocalClientKey() == false)
+	{
+		InvokeGameplayCueEvent(Spec, EGameplayCueEvent::OnActive);
+		InvokeGameplayCueEvent(Spec, EGameplayCueEvent::WhileActive);
+	}
+}
+
 bool UAbilitySystemComponent::IsGameplayCueActive(const FGameplayTag GameplayCueTag) const
 {
 	return HasMatchingGameplayTag(GameplayCueTag);
diff --git a/Engine/Source/Runtime/GameplayAbilities/Private/AttributeSet.cpp b/Engine/Source/Runtime/GameplayAbilities/Private/AttributeSet.cpp
--- a/Engine/Source/Runtime/GameplayAbilities/Private/AttributeSet.cpp
+++ b/Engine/Source/Runtime/GameplayAbilities/Private/AttributeSet.cpp
@@ -23,12 +23,14 @@ void FGameplayAttribute::SetNumericValueChecked(float NewValue, class UAttribute
 	Dest->PreAttributeChange(*this, NewValue);
 	NumericProperty->SetFloatingPointPropertyValue(ValuePtr, NewValue);
 
+#if ENABLE_VISUAL_LOG
 	// draw a graph of the changes to the attribute in the visual logger
 	AActor* OwnerActor = Dest->GetOwningAbilitySystemComponent()->OwnerActor;
 	if (OwnerActor)
 	{
 		ABILITY_VLOG_ATTRIBUTE_GRAPH(OwnerActor, Log, GetName(), OldValue, NewValue);
 	}
+#endif
 }
 
 float FGameplayAttribute::GetNumericValue(const UAttributeSet* Src) const
diff --git a/Engine/Source/Runtime/GameplayAbilities/Private/GameplayCueManager.cpp b/Engine/Source/Runtime/GameplayAbilities/Private/GameplayCueManager.cpp
--- a/Engine/Source/Runtime/GameplayAbilities/Private/GameplayCueManager.cpp
+++ b/Engine/Source/Runtime/GameplayAbilities/Private/GameplayCueManager.cpp
@@ -12,6 +12,14 @@
 #include "UnrealEd.h"
 #endif
 
+int32 DisplayGameplayCues = 0;
+static FAutoConsoleVariableRef CVarDisplayGameplayCues(TEXT("AbilitySystem.DisplayGameplayCues"),	DisplayGameplayCues, TEXT("Display GameplayCue events in world as text."), ECVF_Default	);
+
+int32 DisableGameplayCues = 0;
+static FAutoConsoleVariableRef CVarDisableGameplayCues(TEXT("AbilitySystem.DisableGameplayCues"),	DisableGameplayCues, TEXT("Disables all GameplayCue events in the world."), ECVF_Default );
+
+float DisplayGameplayCueDuration = 5.f;
+static FAutoConsoleVariableRef CVarDurationeGameplayCues(TEXT("AbilitySystem.GameplayCue.DisplayDuration"),	DisplayGameplayCueDuration, TEXT("Disables all GameplayCue events in the world."), ECVF_Default );
 
 UGameplayCueManager::UGameplayCueManager(const FObjectInitializer& PCIP)
 : Super(PCIP)
@@ -33,12 +41,23 @@ void UGameplayCueManager::HandleGameplayCues(AActor* TargetActor, const FGamepla
 
 void UGameplayCueManager::HandleGameplayCue(AActor* TargetActor, FGameplayTag GameplayCueTag, EGameplayCueEvent::Type EventType, FGameplayCueParameters Parameters)
 {
+	if (DisableGameplayCues)
+	{
+		return;
+	}
+
 	if (TargetActor == nullptr)
 	{
 		ABILITY_LOG(Warning, TEXT("UGameplayCueManager::HandleGameplayCue called on null TargetActor. GameplayCueTag: %s."), *GameplayCueTag.ToString());
 		return;
 	}
 
+	if (DisplayGameplayCues)
+	{
+		FString DebugStr = FString::Printf(TEXT("%s - %s"), *GameplayCueTag.ToString(), *EGameplayCueEventToString(EventType) );
+		FColor DebugColor = FColor::Green;
+		DrawDebugString(TargetActor->GetWorld(), FVector(0.f, 0.f, 100.f), DebugStr, TargetActor, DebugColor, DisplayGameplayCueDuration);
+	}
 
 	// GameplayCueTags could have been removed from the dictionary but not content. When the content is resaved the old tag will be cleaned up, but it could still come through here
 	// at runtime. Since we only populate the map with dictionary gameplaycue tags, we may not find it here.
@@ -113,6 +132,11 @@ void UGameplayCueManager::HandleGameplayCueNotify_Internal(AActor* TargetActor,
 					if (auto WeakPtrPtr = InnerMap->Find(DataIdx))
 					{
 						SpawnedInstancedCue = WeakPtrPtr->Get();
+						// If the cue is scheduled to be destroyed, don't reuse it, create a new one instead
+						if (SpawnedInstancedCue && SpawnedInstancedCue->GetLifeSpan() > 0.f)
+						{
+							SpawnedInstancedCue = nullptr;
+						}
 					}
 				}
 
diff --git a/Engine/Source/Runtime/GameplayAbilities/Private/GameplayEffect.cpp b/Engine/Source/Runtime/GameplayAbilities/Private/GameplayEffect.cpp
--- a/Engine/Source/Runtime/GameplayAbilities/Private/GameplayEffect.cpp
+++ b/Engine/Source/Runtime/GameplayAbilities/Private/GameplayEffect.cpp
@@ -37,6 +37,7 @@ UGameplayEffect::UGameplayEffect(const FObjectInitializer& ObjectInitializer)
 	StackLimitCount = 0;
 	StackDurationRefreshPolicy = EGameplayEffectStackingDurationPolicy::RefreshOnSuccessfulApplication;
 	StackPeriodResetPolicy = EGameplayEffectStackingPeriodPolicy::ResetOnSuccessfulApplication;
+	bRequireModifierSuccessToTriggerCues = true;
 
 #if WITH_EDITORONLY_DATA
 	ShowAllProperties = true;
@@ -523,6 +524,8 @@ FGameplayEffectSpecForRPC::FGameplayEffectSpecForRPC(const FGameplayEffectSpec&
 	: Def(InSpec.Def)
 	, ModifiedAttributes(MoveTemp(InSpec.ModifiedAttributes))
 	, EffectContext(InSpec.GetEffectContext())
+	, AggregatedSourceTags(*InSpec.CapturedSourceTags.GetAggregatedTags())
+	, AggregatedTargetTags(*InSpec.CapturedTargetTags.GetAggregatedTags())
 	, Level(InSpec.GetLevel())
 {
 }
@@ -1095,6 +1098,8 @@ FActiveGameplayEffect::FActiveGameplayEffect()
 	, CachedStartGameStateTime(0)
 	, StartWorldTime(0.f)
 	, bIsInhibited(true)
+	, bPendingRepOnActiveGC(false)
+	, bPendingRepWhileActiveGC(false)
 	, IsPendingRemove(false)
 	, PendingNext(nullptr)
 {
@@ -1113,6 +1118,8 @@ FActiveGameplayEffect::FActiveGameplayEffect(FActiveGameplayEffectHandle InHandl
 	, CachedStartGameStateTime(InStartGameStateTime)
 	, StartWorldTime(CurrentWorldTime)
 	, bIsInhibited(true)
+	, bPendingRepOnActiveGC(false)
+	, bPendingRepWhileActiveGC(false)
 	, IsPendingRemove(false)
 	, PendingNext(nullptr)
 {
@@ -1126,6 +1133,8 @@ FActiveGameplayEffect::FActiveGameplayEffect(FActiveGameplayEffect&& Other)
 	,CachedStartGameStateTime(Other.CachedStartGameStateTime)
 	,StartWorldTime(Other.StartWorldTime)
 	,bIsInhibited(Other.bIsInhibited)
+	,bPendingRepOnActiveGC(Other.bPendingRepOnActiveGC)
+	,bPendingRepWhileActiveGC(Other.bPendingRepWhileActiveGC)
 	,IsPendingRemove(Other.IsPendingRemove)
 	,OnRemovedDelegate(Other.OnRemovedDelegate)
 	,PeriodHandle(Other.PeriodHandle)
@@ -1143,6 +1152,8 @@ FActiveGameplayEffect& FActiveGameplayEffect::operator=(FActiveGameplayEffect&&
 	CachedStartGameStateTime = Other.CachedStartGameStateTime;
 	StartWorldTime = Other.StartWorldTime;
 	bIsInhibited = Other.bIsInhibited;
+	bPendingRepOnActiveGC = Other.bPendingRepOnActiveGC;
+	bPendingRepWhileActiveGC = Other.bPendingRepWhileActiveGC;
 	IsPendingRemove = Other.IsPendingRemove;
 	OnRemovedDelegate = Other.OnRemovedDelegate;
 	PeriodHandle = Other.PeriodHandle;
@@ -1160,6 +1171,8 @@ FActiveGameplayEffect& FActiveGameplayEffect::operator=(const FActiveGameplayEff
 	CachedStartGameStateTime = Other.CachedStartGameStateTime;
 	StartWorldTime = Other.StartWorldTime;
 	bIsInhibited = Other.bIsInhibited;
+	bPendingRepOnActiveGC = Other.bPendingRepOnActiveGC;
+	bPendingRepWhileActiveGC = Other.bPendingRepWhileActiveGC;
 	IsPendingRemove = Other.IsPendingRemove;
 	OnRemovedDelegate = Other.OnRemovedDelegate;
 	PeriodHandle = Other.PeriodHandle;
@@ -1169,7 +1182,7 @@ FActiveGameplayEffect& FActiveGameplayEffect::operator=(const FActiveGameplayEff
 }
 
 /** This is the core function that turns the ActiveGE 'on' or 'off */
-void FActiveGameplayEffect::CheckOngoingTagRequirements(const FGameplayTagContainer& OwnerTags, FActiveGameplayEffectsContainer& OwningContainer)
+void FActiveGameplayEffect::CheckOngoingTagRequirements(const FGameplayTagContainer& OwnerTags, FActiveGameplayEffectsContainer& OwningContainer, bool bInvokeGameplayCueEvents)
 {
 	bool bShouldBeInhibited = !Spec.Def->OngoingTagRequirements.RequirementsMet(OwnerTags);
 
@@ -1181,11 +1194,11 @@ void FActiveGameplayEffect::CheckOngoingTagRequirements(const FGameplayTagContai
 		if (bShouldBeInhibited)
 		{
 			// Remove our ActiveGameplayEffects modifiers with our Attribute Aggregators
-			OwningContainer.RemoveActiveGameplayEffectGrantedTagsAndModifiers(*this);
+			OwningContainer.RemoveActiveGameplayEffectGrantedTagsAndModifiers(*this, bInvokeGameplayCueEvents);
 		}
 		else
 		{
-			OwningContainer.AddActiveGameplayEffectGrantedTagsAndModifiers(*this);
+			OwningContainer.AddActiveGameplayEffectGrantedTagsAndModifiers(*this, bInvokeGameplayCueEvents);
 		}
 
 		bIsInhibited = bShouldBeInhibited;
@@ -1204,7 +1217,11 @@ void FActiveGameplayEffect::PreReplicatedRemove(const struct FActiveGameplayEffe
 
 	const_cast<FActiveGameplayEffectsContainer&>(InArray).InternalOnActiveGameplayEffectRemoved(*this);	// Const cast is ok. It is there to prevent mutation of the GameplayEffects array, which this wont do.
 	
-	InArray.Owner->InvokeGameplayCueEvent(Spec, EGameplayCueEvent::Removed);
+	// Only invoke the GC event if the GE was not inhibited when it was removed
+	if (bIsInhibited == false)
+	{
+		InArray.Owner->InvokeGameplayCueEvent(Spec, EGameplayCueEvent::Removed);
+	}
 }
 
 void FActiveGameplayEffect::PostReplicatedAdd(const struct FActiveGameplayEffectsContainer &InArray)
@@ -1234,25 +1251,22 @@ void FActiveGameplayEffect::PostReplicatedAdd(const struct FActiveGameplayEffect
 
 	static const int32 MAX_DELTA_TIME = 3;
 
-	if (ShouldInvokeGameplayCueEvents)
-	{
-		InArray.Owner->InvokeGameplayCueEvent(Spec, EGameplayCueEvent::WhileActive);
-	}
-
 	// Was this actually just activated, or are we just finding out about it due to relevancy/join in progress?
 	float WorldTimeSeconds = InArray.GetWorldTime();
 	int32 GameStateTime = InArray.GetGameStateTime();
 
 	int32 DeltaGameStateTime = GameStateTime - StartGameStateTime;	// How long we think the effect has been playing (only 1 second precision!)
 
-	if (ShouldInvokeGameplayCueEvents && GameStateTime > 0 && FMath::Abs(DeltaGameStateTime) < MAX_DELTA_TIME)
-	{
-		InArray.Owner->InvokeGameplayCueEvent(Spec, EGameplayCueEvent::OnActive);
-	}
-
 	// Set our local start time accordingly
 	StartWorldTime = WorldTimeSeconds - static_cast<float>(DeltaGameStateTime);
 	CachedStartGameStateTime = StartGameStateTime;
+
+	if (ShouldInvokeGameplayCueEvents)
+	{
+		// These events will get invoked if, after the parent array has been completely updated, this GE is still not inhibited
+		bPendingRepOnActiveGC = (GameStateTime > 0 && FMath::Abs(DeltaGameStateTime) < MAX_DELTA_TIME);
+		bPendingRepWhileActiveGC = true;
+	}
 }
 
 void FActiveGameplayEffect::PostReplicatedChange(const struct FActiveGameplayEffectsContainer &InArray)
@@ -1315,9 +1329,6 @@ void FActiveGameplayEffectsContainer::RegisterWithOwner(UAbilitySystemComponent*
 /** This is the main function that executes a GameplayEffect on Attributes and ActiveGameplayEffects */
 void FActiveGameplayEffectsContainer::ExecuteActiveEffectsFrom(FGameplayEffectSpec &Spec, FPredictionKey PredictionKey)
 {
-	// If there are no modifiers, we always want to apply the GameplayCue. If there are modifiers, we only want to invoke the GameplayCue if one of them went through (could be blocked by immunity or % chance roll)
-	bool InvokeGameplayCueExecute = (Spec.Modifiers.Num() == 0);
-	
 	FGameplayEffectSpec& SpecToUse = Spec;
 
 	// Capture our own tags.
@@ -1332,13 +1343,15 @@ void FActiveGameplayEffectsContainer::ExecuteActiveEffectsFrom(FGameplayEffectSp
 	//	Modifiers
 	//		These will modify the base value of attributes
 	// ------------------------------------------------------
+	
+	bool ModifierSuccessfullyExecuted = false;
 
 	for (int32 ModIdx = 0; ModIdx < SpecToUse.Modifiers.Num(); ++ModIdx)
 	{
 		const FGameplayModifierInfo& ModDef = SpecToUse.Def->Modifiers[ModIdx];
 		
 		FGameplayModifierEvaluatedData EvalData(ModDef.Attribute, ModDef.ModifierOp, SpecToUse.GetModifierMagnitude(ModIdx, true));
-		InvokeGameplayCueExecute |= InternalExecuteMod(SpecToUse, EvalData);
+		ModifierSuccessfullyExecuted |= InternalExecuteMod(SpecToUse, EvalData);
 	}
 
 	// ------------------------------------------------------
@@ -1347,7 +1360,7 @@ void FActiveGameplayEffectsContainer::ExecuteActiveEffectsFrom(FGameplayEffectSp
 	// ------------------------------------------------------
 	
 	TArray< FGameplayEffectSpecHandle > ConditionalEffectSpecs;
-	bool ShouldHandleGameplayCues = true;
+	bool GameplayCuesWereManuallyHandled = false;
 
 	for (const FGameplayEffectExecutionDefinition& CurExecDef : SpecToUse.Def->Executions)
 	{
@@ -1390,13 +1403,13 @@ void FActiveGameplayEffectsContainer::ExecuteActiveEffectsFrom(FGameplayEffectSp
 				{
 					CurExecMod.Magnitude = GameplayEffectUtilities::ComputeStackedModifierMagnitude(CurExecMod.Magnitude, SpecStackCount, CurExecMod.ModifierOp);
 				}
-				InvokeGameplayCueExecute |= InternalExecuteMod(SpecToUse, CurExecMod);
+				ModifierSuccessfullyExecuted |= InternalExecuteMod(SpecToUse, CurExecMod);
 			}
 
 			// If execution handled GameplayCues, we dont have to.
 			if (ExecutionOutput.AreGameplayCuesHandledManually())
 			{
-				InvokeGameplayCueExecute = false;
+				GameplayCuesWereManuallyHandled = true;
 			}
 		}
 	}
@@ -1408,6 +1421,21 @@ void FActiveGameplayEffectsContainer::ExecuteActiveEffectsFrom(FGameplayEffectSp
 	// Prune the modified attributes before we replicate
 	SpecToUse.PruneModifiedAttributes();
 
+	// If there are no modifiers or we don't require modifier success to trigger, we apply the GameplayCue. 
+	bool InvokeGameplayCueExecute = (SpecToUse.Modifiers.Num() == 0) || !Spec.Def->bRequireModifierSuccessToTriggerCues;
+
+	// If there are modifiers, we only want to invoke the GameplayCue if one of them went through (could be blocked by immunity or % chance roll)
+	if (SpecToUse.Modifiers.Num() > 0 && ModifierSuccessfullyExecuted)
+	{
+		InvokeGameplayCueExecute = true;
+	}
+
+	// Don't trigger gameplay cues if one of the executions says it manually handled them
+	if (GameplayCuesWereManuallyHandled)
+	{
+		InvokeGameplayCueExecute = false;
+	}
+
 	if (InvokeGameplayCueExecute && SpecToUse.Def->GameplayCues.Num())
 	{
 		// TODO: check replication policy. Right now we will replicate every execute via a multicast RPC
@@ -1838,8 +1866,16 @@ void FActiveGameplayEffectsContainer::InternalUpdateNumericalAttribute(FGameplay
 	FOnGameplayAttributeChange* Delegate = AttributeChangeDelegates.Find(Attribute);
 	if (Delegate)
 	{
-		Delegate->Broadcast(NewValue, ModData);
+		// We should only have one: either cached CurrentModcallbackData, or explicit callback data passed directly in.
+		if (ModData && CurrentModcallbackData)
+		{
+			ABILITY_LOG(Warning, TEXT("Had passed in ModData and cached CurrentModcallbackData in FActiveGameplayEffectsContainer::InternalUpdateNumericalAttribute. For attribute %s on %s."), *Attribute.GetName(), *Owner->GetFullName() );
+		}
+
+		// Broadcast dirty delegate. If we were given explicit mod data then pass it. 
+		Delegate->Broadcast(NewValue, ModData ? ModData : CurrentModcallbackData);
 	}
+	CurrentModcallbackData = nullptr;
 }
 
 void FActiveGameplayEffectsContainer::SetAttributeBaseValue(FGameplayAttribute Attribute, float NewBaseValue)
@@ -1940,8 +1976,18 @@ void FActiveGameplayEffectsContainer::ApplyModToAttribute(const FGameplayAttribu
 		const UAttributeSet* Set = Owner->GetAttributeSubobject( Attribute.GetAttributeSetClass() );
 		check(Set);
 
+		// Cache this off here so we dont have to pipe ModData through the agggregator system. This works because the AGE OnDirty will always be called first before any other possible aggregator dependants.
+		CurrentModcallbackData = ModData;
+
 		Set->PreAttributeBaseChange(Attribute, NewBase);
 		Agg->SetBaseValue(NewBase);
+
+		if (CurrentModcallbackData)
+		{
+			// We expect this to be cleared for us in InternalUpdateNumericalAttribute
+			ABILITY_LOG(Warning, TEXT("FActiveGameplayEffectsContainer::ApplyModToAttribute CurrentModcallbackData was not consumed For attribute %s on %s."), *Attribute.GetName(), *Owner->GetFullName() );
+			CurrentModcallbackData = nullptr;
+		}
 	}
 	else
 	{
@@ -2089,6 +2135,43 @@ FActiveGameplayEffect* FActiveGameplayEffectsContainer::ApplyGameplayEffectSpec(
 	AppliedEffectSpec.CaptureAttributeDataFromTarget(Owner);
 	AppliedEffectSpec.CalculateModifierMagnitudes();
 
+	// Build ModifiedAttribute list so GCs can have magnitude info if non-period effect
+	// Note: One day may want to not call gameplay cues unless ongoing tag requirements are met (will need to move this there)
+	{
+		const bool HasModifiedAttributes = AppliedEffectSpec.ModifiedAttributes.Num() > 0;
+		const bool HasDurationAndNoPeriod = AppliedEffectSpec.GetDuration() > 0.0f && AppliedEffectSpec.GetPeriod() == UGameplayEffect::NO_PERIOD;
+		const bool HasPeriodAndNoDuration =	AppliedEffectSpec.GetDuration() == UGameplayEffect::INSTANT_APPLICATION && 
+											AppliedEffectSpec.GetPeriod() != UGameplayEffect::NO_PERIOD;
+		const bool ShouldBuildModifiedAttributeList = !HasModifiedAttributes && (HasDurationAndNoPeriod || HasPeriodAndNoDuration);
+		if (ShouldBuildModifiedAttributeList)
+		{
+			int32 ModifierIndex = -1;
+			for (const FGameplayModifierInfo& Mod : AppliedEffectSpec.Def->Modifiers)
+			{
+				++ModifierIndex;
+
+				// Take magnitude from evaluated magnitudes
+				float Magnitude = 0.0f;
+				if (AppliedEffectSpec.Modifiers.IsValidIndex(ModifierIndex))
+				{
+					const FModifierSpec& ModSpec = AppliedEffectSpec.Modifiers[ModifierIndex];
+					Magnitude = ModSpec.GetEvaluatedMagnitude();
+				}
+				
+				// Add to ModifiedAttribute list if it doesn't exist already
+				FGameplayEffectModifiedAttribute* ModifiedAttribute = AppliedEffectSpec.GetModifiedAttribute(Mod.Attribute);
+				if (!ModifiedAttribute)
+				{
+					ModifiedAttribute = AppliedEffectSpec.AddModifiedAttribute(Mod.Attribute);
+				}
+				ModifiedAttribute->TotalMagnitude += Magnitude;
+			}
+
+			// Prune list since ModifiedAttributes have only been built to
+			AppliedEffectSpec.PruneModifiedAttributes();
+		}
+	}
+
 	// Register Source and Target non snapshot capture delegates here
 	AppliedEffectSpec.CapturedRelevantAttributes.RegisterLinkedAggregatorCallbacks(AppliedActiveGE->Handle);
 	
@@ -2205,7 +2288,7 @@ void FActiveGameplayEffectsContainer::InternalOnActiveGameplayEffectAdded(FActiv
 	Effect.CheckOngoingTagRequirements(OwnerTags, *this);
 }
 
-void FActiveGameplayEffectsContainer::AddActiveGameplayEffectGrantedTagsAndModifiers(FActiveGameplayEffect& Effect)
+void FActiveGameplayEffectsContainer::AddActiveGameplayEffectGrantedTagsAndModifiers(FActiveGameplayEffect& Effect, bool bInvokeGameplayCueEvents)
 {
 	if (Effect.Spec.Def == nullptr)
 	{
@@ -2241,6 +2324,15 @@ void FActiveGameplayEffectsContainer::AddActiveGameplayEffectGrantedTagsAndModif
 	for (const FGameplayEffectCue& Cue : Effect.Spec.Def->GameplayCues)
 	{
 		Owner->UpdateTagMap(Cue.GameplayCueTags, 1);
+
+		if (bInvokeGameplayCueEvents)
+		{
+			// TODO: Optimize this so only one batched RPC is called
+			for (auto It = Cue.GameplayCueTags.CreateConstIterator(); It; ++It)
+			{
+				Owner->AddGameplayCue(*It);
+			}
+		}
 	}
 }
 
@@ -2328,6 +2420,9 @@ bool FActiveGameplayEffectsContainer::InternalRemoveActiveGameplayEffect(int32 I
 			}
 		}
 
+		// Don't invoke the GC event if the effect is inhibited, and thus the GC is already not active
+		ShouldInvokeGameplayCueEvent &= !Effect.bIsInhibited;
+
 		if (ShouldInvokeGameplayCueEvent)
 		{
 			Owner->InvokeGameplayCueEvent(Effect.Spec, EGameplayCueEvent::Removed);
@@ -2401,7 +2496,7 @@ void FActiveGameplayEffectsContainer::InternalOnActiveGameplayEffectRemoved(cons
 		RemoveActiveEffectTagDependency(Effect.Spec.Def->OngoingTagRequirements.IgnoreTags, Effect.Handle);
 		RemoveActiveEffectTagDependency(Effect.Spec.Def->OngoingTagRequirements.RequireTags, Effect.Handle);
 
-		RemoveActiveGameplayEffectGrantedTagsAndModifiers(Effect);
+		RemoveActiveGameplayEffectGrantedTagsAndModifiers(Effect, !Effect.bIsInhibited);
 	}
 	else
 	{
@@ -2412,7 +2507,7 @@ void FActiveGameplayEffectsContainer::InternalOnActiveGameplayEffectRemoved(cons
 	OnActiveGameplayEffectRemovedDelegate.Broadcast();
 }
 
-void FActiveGameplayEffectsContainer::RemoveActiveGameplayEffectGrantedTagsAndModifiers(const FActiveGameplayEffect& Effect)
+void FActiveGameplayEffectsContainer::RemoveActiveGameplayEffectGrantedTagsAndModifiers(const FActiveGameplayEffect& Effect, bool bInvokeGameplayCueEvents)
 {
 	// Update AttributeAggregators: remove mods from this ActiveGE Handle
 	if (Effect.Spec.GetPeriod() <= UGameplayEffect::NO_PERIOD)
@@ -2442,6 +2537,15 @@ void FActiveGameplayEffectsContainer::RemoveActiveGameplayEffectGrantedTagsAndMo
 	for (const FGameplayEffectCue& Cue : Effect.Spec.Def->GameplayCues)
 	{
 		Owner->UpdateTagMap(Cue.GameplayCueTags, -1);
+
+		if (bInvokeGameplayCueEvents)
+		{
+			// TODO: Optimize this so only one batched RPC is called
+			for (auto It = Cue.GameplayCueTags.CreateConstIterator(); It; ++It)
+			{
+				Owner->RemoveGameplayCue(*It);
+			}
+		}
 	}
 }
 
@@ -2519,7 +2623,7 @@ void FActiveGameplayEffectsContainer::OnOwnerTagChange(FGameplayTag TagChange, i
 			FActiveGameplayEffect* ActiveEffect = GetActiveGameplayEffect(Handle);
 			if (ActiveEffect)
 			{
-				ActiveEffect->CheckOngoingTagRequirements(OwnerTags, *this);
+				ActiveEffect->CheckOngoingTagRequirements(OwnerTags, *this, true);
 			}
 		}
 	}
@@ -2558,6 +2662,38 @@ bool FActiveGameplayEffectsContainer::IsNetAuthority() const
 	return Owner->IsOwnerActorAuthoritative();
 }
 
+bool FActiveGameplayEffectsContainer::NetDeltaSerialize(FNetDeltaSerializeInfo& DeltaParms)
+{
+	bool RetVal = FastArrayDeltaSerialize<FActiveGameplayEffect>(GameplayEffects_Internal, DeltaParms, *this);
+
+	// After the array has been replicated, invoke GC events ONLY if the effect is not inhibited
+	// We postpone this check because in the same net update we could receive multiple GEs that affect if one another is inhibited
+	
+	if (DeltaParms.Writer == nullptr && Owner != nullptr)
+	{
+		QUICK_SCOPE_CYCLE_COUNTER(STAT_ActiveGameplayEffectsContainer_NetDeltaSerialize_CheckRepGameplayCues);
+
+		for (const FActiveGameplayEffect& Effect : this)
+		{
+			if (Effect.bIsInhibited == false)
+			{
+				if (Effect.bPendingRepOnActiveGC)
+				{
+					Owner->InvokeGameplayCueEvent(Effect.Spec, EGameplayCueEvent::OnActive);
+					Effect.bPendingRepOnActiveGC = false;
+				}
+				if (Effect.bPendingRepWhileActiveGC)
+				{
+					Owner->InvokeGameplayCueEvent(Effect.Spec, EGameplayCueEvent::WhileActive);
+					Effect.bPendingRepWhileActiveGC = false;
+				}
+			}
+		}
+	}
+
+	return RetVal;
+}
+
 void FActiveGameplayEffectsContainer::PreDestroy()
 {
 	
@@ -2994,6 +3130,15 @@ bool FActiveGameplayEffectQuery::Matches(const FActiveGameplayEffect& Effect) co
 		}
 	}
 
+	// check definition
+	if (EffectDef)
+	{
+		if (Effect.Spec.Def != EffectDef)
+		{
+			return false;
+		}
+	}
+
 	// passed all the checks
 	return true;
 }
diff --git a/Engine/Source/Runtime/GameplayAbilities/Private/GameplayEffectTypes.cpp b/Engine/Source/Runtime/GameplayAbilities/Private/GameplayEffectTypes.cpp
--- a/Engine/Source/Runtime/GameplayAbilities/Private/GameplayEffectTypes.cpp
+++ b/Engine/Source/Runtime/GameplayAbilities/Private/GameplayEffectTypes.cpp
@@ -97,41 +97,78 @@ void FGameplayEffectContext::AddHitResult(const FHitResult InHitResult, bool bRe
 
 bool FGameplayEffectContext::NetSerialize(FArchive& Ar, class UPackageMap* Map, bool& bOutSuccess)
 {
-	Ar << Instigator;
-	Ar << EffectCauser;
-
-	uint8 HasActorByte = (Actors.Num() > 0);
-	Ar.SerializeBits(&HasActorByte, 1);
-	if (HasActorByte == 1)
+	uint8 RepBits = 0;
+	if (Ar.IsSaving())
 	{
-		Ar << Actors;
+		if (Instigator.IsValid() )
+		{
+			RepBits |= 1 << 0;
+		}
+		if (EffectCauser.IsValid() )
+		{
+			RepBits |= 1 << 1;
+		}
+		if (SourceObject.IsValid() )
+		{
+			RepBits |= 1 << 2;
+		}
+		if (Actors.Num() > 0)
+		{
+			RepBits |= 1 << 3;
+		}
+		if (HitResult.IsValid())
+		{
+			RepBits |= 1 << 4;
+		}
+		if (bHasWorldOrigin)
+		{
+			RepBits |= 1 << 5;
+		}
 	}
 
-	uint8 HasHitResultsByte = HitResult.IsValid();
-	Ar.SerializeBits(&HasHitResultsByte, 1);
-	if (Ar.IsLoading())
+	Ar.SerializeBits(&RepBits, 6);
+
+	if (RepBits & (1 << 0))
 	{
-		if (HasHitResultsByte)
+		Ar << Instigator;
+	}
+	if (RepBits & (1 << 1))
+	{
+		Ar << EffectCauser;
+	}
+	if (RepBits & (1 << 2))
+	{
+		Ar << SourceObject;
+	}
+	if (RepBits & (1 << 3))
+	{
+		Ar << Actors;;
+	}
+	if (RepBits & (1 << 4))
+	{
+		if (Ar.IsLoading())
 		{
 			if (!HitResult.IsValid())
 			{
 				HitResult = TSharedPtr<FHitResult>(new FHitResult());
 			}
 		}
-		AddInstigator(Instigator.Get(), EffectCauser.Get()); // Just to initialize InstigatorAbilitySystemComponent
-	}
-
-	if (HasHitResultsByte == 1)
-	{
 		HitResult->NetSerialize(Ar, Map, bOutSuccess);
 	}
-	uint8 HasWorldOriginByte = bHasWorldOrigin;
-	Ar.SerializeBits(&HasWorldOriginByte, 1);
-	bHasWorldOrigin = HasWorldOriginByte & 1;
-	if (bHasWorldOrigin)
+	if (RepBits & (1 << 5))
 	{
 		Ar << WorldOrigin;
+		bHasWorldOrigin = true;
 	}
+	else
+	{
+		bHasWorldOrigin = false;
+	}
+
+	if (Ar.IsLoading())
+	{
+		AddInstigator(Instigator.Get(), EffectCauser.Get()); // Just to initialize InstigatorAbilitySystemComponent
+	}	
 	
 	bOutSuccess = true;
 	return true;
@@ -229,6 +266,12 @@ FString EGameplayModEffectToString(int32 Type)
 	return e->GetEnumName(Type);
 }
 
+FString EGameplayCueEventToString(int32 Type)
+{
+	static UEnum *e = FindObject<UEnum>(ANY_PACKAGE, TEXT("EGameplayCueEvent"));
+	return e->GetEnumName(Type);
+}
+
 bool FGameplayTagCountContainer::HasMatchingGameplayTag(FGameplayTag TagToCheck) const
 {
 	return ExplicitTags.HasTag(TagToCheck, EGameplayTagMatchType::IncludeParentTags, EGameplayTagMatchType::Explicit);
diff --git a/Engine/Source/Runtime/GameplayAbilities/Public/Abilities/GameplayAbility.h b/Engine/Source/Runtime/GameplayAbilities/Public/Abilities/GameplayAbility.h
--- a/Engine/Source/Runtime/GameplayAbilities/Public/Abilities/GameplayAbility.h
+++ b/Engine/Source/Runtime/GameplayAbilities/Public/Abilities/GameplayAbility.h
@@ -416,6 +416,9 @@ protected:
 	UFUNCTION(BlueprintImplementableEvent, Category = Ability, DisplayName = "OnEndAbility")
 	void K2_OnEndAbility();
 
+	/** Check if the ability can be ended */
+	bool IsEndAbilityValid(const FGameplayAbilitySpecHandle Handle, const FGameplayAbilityActorInfo* ActorInfo) const;
+
 	/** Native function, called if an ability ends normally or abnormally. If bReplicate is set to true, try to replicate the ending to the client/server */
 	virtual void EndAbility(const FGameplayAbilitySpecHandle Handle, const FGameplayAbilityActorInfo* ActorInfo, const FGameplayAbilityActivationInfo ActivationInfo, bool bReplicateEndAbility);
 
@@ -578,7 +581,7 @@ public:
 	bool IsSupportedForNetworking() const override;
 
 	/** Returns the gameplay effect used to determine cooldown */
-	class UGameplayEffect* GetCooldownGameplayEffect() const;
+	virtual class UGameplayEffect* GetCooldownGameplayEffect() const;
 
 	/** Returns the gameplay effect used to apply cost */
 	class UGameplayEffect* GetCostGameplayEffect() const;
diff --git a/Engine/Source/Runtime/GameplayAbilities/Public/Abilities/GameplayAbilityTargetTypes.h b/Engine/Source/Runtime/GameplayAbilities/Public/Abilities/GameplayAbilityTargetTypes.h
--- a/Engine/Source/Runtime/GameplayAbilities/Public/Abilities/GameplayAbilityTargetTypes.h
+++ b/Engine/Source/Runtime/GameplayAbilities/Public/Abilities/GameplayAbilityTargetTypes.h
@@ -116,6 +116,11 @@ struct GAMEPLAYABILITIES_API FGameplayAbilityTargetData
 		return FVector::ZeroVector;
 	}
 
+	virtual FTransform GetEndPointTransform() const
+	{
+		return FTransform(GetEndPoint());
+	}
+
 	// -------------------------------------
 
 	virtual UScriptStruct* GetScriptStruct()
diff --git a/Engine/Source/Runtime/GameplayAbilities/Public/Abilities/Tasks/AbilityTask_WaitAbilityActivate.h b/Engine/Source/Runtime/GameplayAbilities/Public/Abilities/Tasks/AbilityTask_WaitAbilityActivate.h
--- a/Engine/Source/Runtime/GameplayAbilities/Public/Abilities/Tasks/AbilityTask_WaitAbilityActivate.h
+++ b/Engine/Source/Runtime/GameplayAbilities/Public/Abilities/Tasks/AbilityTask_WaitAbilityActivate.h
@@ -26,11 +26,12 @@ class UAbilityTask_WaitAbilityActivate : public UAbilityTask
 
 	/** Wait until a new ability (of the same or different type) is activated. Only input based abilities will be counted unless IncludeTriggeredAbilities is true. */
 	UFUNCTION(BlueprintCallable, Category="Ability|Tasks", meta = (HidePin = "WorldContextObject", DefaultToSelf = "WorldContextObject", BlueprintInternalUseOnly = "TRUE"))
-	static UAbilityTask_WaitAbilityActivate* WaitForAbilityActivate(UObject* WorldContextObject, FGameplayTag WithTag, FGameplayTag WithoutTag, bool IncludeTriggeredAbilities=false);
+	static UAbilityTask_WaitAbilityActivate* WaitForAbilityActivate(UObject* WorldContextObject, FGameplayTag WithTag, FGameplayTag WithoutTag, bool IncludeTriggeredAbilities=false, bool TriggerOnce=true);
 
 	FGameplayTag WithTag;
 	FGameplayTag WithoutTag;
 	bool IncludeTriggeredAbilities;
+	bool TriggerOnce;
 
 protected:
 
diff --git a/Engine/Source/Runtime/GameplayAbilities/Public/Abilities/Tasks/AbilityTask_WaitAbilityCommit.h b/Engine/Source/Runtime/GameplayAbilities/Public/Abilities/Tasks/AbilityTask_WaitAbilityCommit.h
--- a/Engine/Source/Runtime/GameplayAbilities/Public/Abilities/Tasks/AbilityTask_WaitAbilityCommit.h
+++ b/Engine/Source/Runtime/GameplayAbilities/Public/Abilities/Tasks/AbilityTask_WaitAbilityCommit.h
@@ -27,10 +27,11 @@ class UAbilityTask_WaitAbilityCommit : public UAbilityTask
 
 	/** Wait until a new ability (of the same or different type) is commited. Used to gracefully interrupt abilities in specific ways. */
 	UFUNCTION(BlueprintCallable, Category="Ability|Tasks", meta = (HidePin = "WorldContextObject", DefaultToSelf = "WorldContextObject", BlueprintInternalUseOnly = "TRUE", DisplayName = "Wait For New Ability Commit"))
-	static UAbilityTask_WaitAbilityCommit* WaitForAbilityCommit(UObject* WorldContextObject, FGameplayTag WithTag, FGameplayTag WithoutTage);
+	static UAbilityTask_WaitAbilityCommit* WaitForAbilityCommit(UObject* WorldContextObject, FGameplayTag WithTag, FGameplayTag WithoutTage, bool TriggerOnce=true);
 
 	FGameplayTag WithTag;
 	FGameplayTag WithoutTag;
+	bool TriggerOnce;
 
 protected:
 
diff --git a/Engine/Source/Runtime/GameplayAbilities/Public/Abilities/Tasks/AbilityTask_WaitCancel.h b/Engine/Source/Runtime/GameplayAbilities/Public/Abilities/Tasks/AbilityTask_WaitCancel.h
new file mode 100644
--- /dev/null
+++ b/Engine/Source/Runtime/GameplayAbilities/Public/Abilities/Tasks/AbilityTask_WaitCancel.h
@@ -0,0 +1,35 @@
+// Copyright 1998-2015 Epic Games, Inc. All Rights Reserved.
+#pragma once
+
+#include "AbilityTask.h"
+#include "GameplayAbility.h"
+#include "AbilityTask_WaitCancel.generated.h"
+
+DECLARE_DYNAMIC_MULTICAST_DELEGATE(FWaitCancelDelegate);
+
+UCLASS(MinimalAPI)
+class UAbilityTask_WaitCancel : public UAbilityTask
+{
+	GENERATED_UCLASS_BODY()
+
+	UPROPERTY(BlueprintAssignable)
+	FWaitCancelDelegate	OnCancel;
+	
+	UFUNCTION()
+	void OnCancelCallback();
+
+	UFUNCTION()
+	void OnLocalCancelCallback();
+
+	UFUNCTION(BlueprintCallable, meta=(HidePin = "WorldContextObject", DefaultToSelf = "WorldContextObject", BlueprintInternalUseOnly = "true", DisplayName="Wait for Cancel Input"), Category="Ability|Tasks")
+	static UAbilityTask_WaitCancel* WaitCancel(UObject* WorldContextObject);	
+
+	virtual void Activate() override;
+
+protected:
+
+	virtual void OnDestroy(bool AbilityEnding) override;
+
+	bool RegisteredCallbacks;
+	
+};
\ No newline at end of file
diff --git a/Engine/Source/Runtime/GameplayAbilities/Public/Abilities/Tasks/AbilityTask_WaitGameplayEffectApplied_Self.h b/Engine/Source/Runtime/GameplayAbilities/Public/Abilities/Tasks/AbilityTask_WaitGameplayEffectApplied_Self.h
--- a/Engine/Source/Runtime/GameplayAbilities/Public/Abilities/Tasks/AbilityTask_WaitGameplayEffectApplied_Self.h
+++ b/Engine/Source/Runtime/GameplayAbilities/Public/Abilities/Tasks/AbilityTask_WaitGameplayEffectApplied_Self.h
@@ -20,7 +20,7 @@ class UAbilityTask_WaitGameplayEffectApplied_Self : public UAbilityTask_WaitGame
 	 * Optional External Owner can be used to run this task on someone else (not the owner of the ability). By default you can leave this empty.
 	 */	 	 
 	UFUNCTION(BlueprintCallable, Category="Ability|Tasks", meta = (HidePin = "WorldContextObject", DefaultToSelf = "WorldContextObject", BlueprintInternalUseOnly = "TRUE"))
-	static UAbilityTask_WaitGameplayEffectApplied_Self* WaitGameplayEffectAppliedToSelf(UObject* WorldContextObject, AActor* TestActor, const FGameplayTargetDataFilterHandle SourceFilter, FGameplayTagRequirements SourceTagRequirements, FGameplayTagRequirements TargetTagRequirements, bool TriggerOnce=false, AActor* OptionalExternalOwner=nullptr );
+	static UAbilityTask_WaitGameplayEffectApplied_Self* WaitGameplayEffectAppliedToSelf(UObject* WorldContextObject, const FGameplayTargetDataFilterHandle SourceFilter, FGameplayTagRequirements SourceTagRequirements, FGameplayTagRequirements TargetTagRequirements, bool TriggerOnce=false, AActor* OptionalExternalOwner=nullptr );
 	
 protected:
 
diff --git a/Engine/Source/Runtime/GameplayAbilities/Public/AbilitySystemBlueprintLibrary.h b/Engine/Source/Runtime/GameplayAbilities/Public/AbilitySystemBlueprintLibrary.h
--- a/Engine/Source/Runtime/GameplayAbilities/Public/AbilitySystemBlueprintLibrary.h
+++ b/Engine/Source/Runtime/GameplayAbilities/Public/AbilitySystemBlueprintLibrary.h
@@ -98,6 +98,9 @@ class GAMEPLAYABILITIES_API UAbilitySystemBlueprintLibrary : public UBlueprintFu
 	UFUNCTION(BlueprintPure, Category = "Ability|TargetData")
 	static FVector GetTargetDataEndPoint(FGameplayAbilityTargetDataHandle TargetData, int32 Index);
 
+	UFUNCTION(BlueprintPure, Category = "Ability|TargetData")
+	static FTransform GetTargetDataEndPointTransform(FGameplayAbilityTargetDataHandle TargetData, int32 Index);
+
 	// -------------------------------------------------------------------------------
 	//		GameplayEffectContext
 	// -------------------------------------------------------------------------------
@@ -127,6 +130,10 @@ class GAMEPLAYABILITIES_API UAbilitySystemBlueprintLibrary : public UBlueprintFu
 	UFUNCTION(BlueprintPure, Category = "Ability|EffectContext", Meta = (DisplayName = "GetEffectCauser"))
 	static AActor* EffectContextGetEffectCauser(FGameplayEffectContextHandle EffectContext);
 
+	/** Gets the source object of the effect. */
+	UFUNCTION(BlueprintPure, Category = "Ability|EffectContext", Meta = (DisplayName = "GetSourceObject"))
+	static UObject* EffectContextGetSourceObject(FGameplayEffectContextHandle EffectContext);
+
 	// -------------------------------------------------------------------------------
 	//		GameplayCue
 	// -------------------------------------------------------------------------------
@@ -170,6 +177,9 @@ class GAMEPLAYABILITIES_API UAbilitySystemBlueprintLibrary : public UBlueprintFu
 	UFUNCTION(BlueprintPure, Category = "Ability|GameplayCue")
 	static bool GetGameplayCueDirection(AActor* TargetActor, FGameplayCueParameters Parameters, FVector& Direction);
 
+	/** Returns true if the aggregated source and target tags from the effect spec meets the tag requirements */
+	UFUNCTION(BlueprintPure, Category = "Ability|GameplayCue")
+	static bool DoesGameplayCueMeetTagRequirements(FGameplayCueParameters Parameters, UPARAM(ref) FGameplayTagRequirements& SourceTagReqs, UPARAM(ref) FGameplayTagRequirements& TargetTagReqs);
 
 	// -------------------------------------------------------------------------------
 	//		GameplayEffectSpec
diff --git a/Engine/Source/Runtime/GameplayAbilities/Public/AbilitySystemComponent.h b/Engine/Source/Runtime/GameplayAbilities/Public/AbilitySystemComponent.h
--- a/Engine/Source/Runtime/GameplayAbilities/Public/AbilitySystemComponent.h
+++ b/Engine/Source/Runtime/GameplayAbilities/Public/AbilitySystemComponent.h
@@ -478,6 +478,9 @@ class GAMEPLAYABILITIES_API UAbilitySystemComponent : public UActorComponent, pu
 
 	UFUNCTION(NetMulticast, unreliable)
 	void NetMulticast_InvokeGameplayCueAdded(const FGameplayTag GameplayCueTag, FPredictionKey PredictionKey, FGameplayEffectContextHandle EffectContext);
+	
+	UFUNCTION(NetMulticast, unreliable)
+	void NetMulticast_InvokeGameplayCueAddedAndWhileActive_FromSpec(const FGameplayEffectSpecForRPC& Spec, FPredictionKey PredictionKey);
 
 	// GameplayCues can also come on their own. These take an optional effect context to pass through hit result, etc
 	void ExecuteGameplayCue(const FGameplayTag GameplayCueTag, FGameplayEffectContextHandle EffectContext = FGameplayEffectContextHandle());
@@ -492,7 +495,7 @@ class GAMEPLAYABILITIES_API UAbilitySystemComponent : public UActorComponent, pu
 	/** Removes any GameplayCue added on its own, i.e. not as part of a GameplayEffect. */
 	void RemoveAllGameplayCues();
 	
-	void InvokeGameplayCueEvent(const FGameplayEffectSpecForRPC &Spec, EGameplayCueEvent::Type EventType);
+	void InvokeGameplayCueEvent(const FGameplayEffectSpecForRPC& Spec, EGameplayCueEvent::Type EventType);
 
 	void InvokeGameplayCueEvent(const FGameplayTag GameplayCueTag, EGameplayCueEvent::Type EventType, FGameplayEffectContextHandle EffectContext = FGameplayEffectContextHandle());
 
@@ -734,19 +737,11 @@ public:
 	
 	virtual void BindAbilityActivationToInputComponent(UInputComponent* InputComponent, FGameplayAbiliyInputBinds BindInfo);
 
-	virtual void AbilityInputPressed(int32 InputID);
-
-	virtual void AbilityInputReleased(int32 InputID);
-
-	void AbilitySpecInputPressed(FGameplayAbilitySpec& Spec);
-
-	void AbilitySpecInputReleased(FGameplayAbilitySpec& Spec);
-
-	UFUNCTION(BlueprintCallable, Category="Abilities")
-	virtual void InputConfirm();
-
-	UFUNCTION(BlueprintCallable, Category="Abilities")
-	virtual void InputCancel();
+	virtual void AbilityLocalInputPressed(int32 InputID);
+	virtual void AbilityLocalInputReleased(int32 InputID);
+	
+	virtual void LocalInputConfirm();
+	virtual void LocalInputCancel();
 	
 	/** Replicate that an ability has ended/canceled, to the client or server as appropriate */
 	void	ReplicateEndOrCancelAbility(FGameplayAbilitySpecHandle Handle, FGameplayAbilityActivationInfo ActivationInfo, UGameplayAbility* Ability, bool bWasCanceled);
@@ -755,15 +750,19 @@ public:
 	int32 GenericConfirmInputID;
 	int32 GenericCancelInputID;
 
-	bool IsGenericConfirmInputBound(int32 InputID) const	{ return ((InputID == GenericConfirmInputID) && ConfirmCallbacks.IsBound()); }
-	bool IsGenericCancelInputBound(int32 InputID) const		{ return ((InputID == GenericCancelInputID) && CancelCallbacks.IsBound()); }
+	bool IsGenericConfirmInputBound(int32 InputID) const	{ return ((InputID == GenericConfirmInputID) && GenericLocalConfirmCallbacks.IsBound()); }
+	bool IsGenericCancelInputBound(int32 InputID) const		{ return ((InputID == GenericCancelInputID) && GenericLocalCancelCallbacks.IsBound()); }
+
+	/** Generic local callback for generic ConfirmEvent that any ability can listen to */
+	FAbilityConfirmOrCancel	GenericLocalConfirmCallbacks;
 
+	/** Generic local callback for generic CancelEvent that any ability can listen to */
+	FAbilityConfirmOrCancel	GenericLocalCancelCallbacks;
 
-	FAbilityAbilityKey	AbilityKeyPressCallbacks;
-	FAbilityAbilityKey	AbilityKeyReleaseCallbacks;
-	FAbilityConfirmOrCancel	ConfirmCallbacks;
-	FAbilityConfirmOrCancel	CancelCallbacks;
+	/** A generic callback anytime an ability is activated (started) */	
 	FGenericAbilityDelegate AbilityActivatedCallbacks;
+
+	/** A generic callback anytime an ability is commited (cost/cooldown applied) */
 	FGenericAbilityDelegate AbilityCommitedCallbacks;
 	FAbilityFailedDelegate AbilityFailedCallbacks;
 
@@ -987,24 +986,13 @@ public:
 	UFUNCTION(Server, reliable, WithValidation)
 	void ServerSetInputReleased(FGameplayAbilitySpecHandle AbilityHandle);
 
-	void ConsumeAbilityConfirmCancel();
-
-	void ConsumeAbilityTargetData();
-
-	//Note: Ability key state is stored on the ability spec, not here on the ASC with confirm/cancel.
-	bool ReplicatedConfirmAbility;
-	bool ReplicatedCancelAbility;
-
-	FGameplayAbilityTargetDataHandle ReplicatedTargetData;
-
-	/** ReplicatedTargetData was received */
-	FAbilityTargetData	ReplicatedTargetDataDelegate;
+	/** Called on local player always. Called on server only if bReplicateInputDirectly is set on the GameplayAbility. */
+	void AbilitySpecInputPressed(FGameplayAbilitySpec& Spec);
 
-	/** ReplicatedTargetData was 'cancelled' for this activation */
-	FAbilityConfirmOrCancel	ReplicatedTargetDataCancelledDelegate;
+	/** Called on local player always. Called on server only if bReplicateInputDirectly is set on the GameplayAbility. */
+	void AbilitySpecInputReleased(FGameplayAbilitySpec& Spec);
 
-	/** Targeting actor rejected a confirmation attempt */
-	FTargetingRejectedConfirmation TargetingRejectedConfirmationDelegate;
+	// ---------------------------------------------------------------------
 
 	/** Tasks that run on simulated proxies */
 	UPROPERTY(ReplicatedUsing=OnRep_SimulatedTasks)
diff --git a/Engine/Source/Runtime/GameplayAbilities/Public/AbilitySystemLog.h b/Engine/Source/Runtime/GameplayAbilities/Public/AbilitySystemLog.h
--- a/Engine/Source/Runtime/GameplayAbilities/Public/AbilitySystemLog.h
+++ b/Engine/Source/Runtime/GameplayAbilities/Public/AbilitySystemLog.h
@@ -14,6 +14,24 @@
 GAMEPLAYABILITIES_API DECLARE_LOG_CATEGORY_EXTERN(LogAbilitySystem, Warning, All);
 GAMEPLAYABILITIES_API DECLARE_LOG_CATEGORY_EXTERN(VLogAbilitySystem, Warning, All);
 
+#if NO_LOGGING
+
+// Without logging enabled we pass ability system through to UE_LOG which only handles Fatal verbosity in NO_LOGGING
+#define ABILITY_LOG(Verbosity, Format, ...) \
+{ \
+	UE_LOG(LogAbilitySystem, Verbosity, Format, ##__VA_ARGS__); \
+}
+
+#define ABILITY_VLOG(Actor, Verbosity, Format, ...) \
+{ \
+	UE_LOG(LogAbilitySystem, Verbosity, Format, ##__VA_ARGS__); \
+	UE_VLOG(Actor, VLogAbilitySystem, Verbosity, Format, ##__VA_ARGS__); \
+}
+
+#define ABILITY_LOG_SCOPE( Format, ... ) 
+
+#else
+
 #define ABILITY_LOG(Verbosity, Format, ...) \
 { \
 	FString Str = AbilitySystemLog::Log(ELogVerbosity::Verbosity, FString::Printf(Format, ##__VA_ARGS__)); \
@@ -27,6 +45,12 @@ GAMEPLAYABILITIES_API DECLARE_LOG_CATEGORY_EXTERN(VLogAbilitySystem, Warning, Al
 	UE_VLOG(Actor, VLogAbilitySystem, Verbosity, TEXT("%s"), *Str); \
 }
 
+#define ABILITY_LOG_SCOPE( Format, ... ) AbilitySystemLogScope PREPROCESSOR_JOIN(LogScope,__LINE__)( FString::Printf(Format, ##__VA_ARGS__));
+
+#endif //NO_LOGGING
+
+#if ENABLE_VISUAL_LOG
+
 #define ABILITY_VLOG_ATTRIBUTE_GRAPH(Actor, Verbosity, AttributeName, OldValue, NewValue) \
 { \
 	const FName GraphName("Attribute Graph"); \
@@ -38,7 +62,11 @@ GAMEPLAYABILITIES_API DECLARE_LOG_CATEGORY_EXTERN(VLogAbilitySystem, Warning, Al
 	UE_VLOG_HISTOGRAM(OwnerActor, VLogAbilitySystem, Log, GraphName, LineName, NewPt); \
 }
 
-#define ABILITY_LOG_SCOPE( Format, ... ) AbilitySystemLogScope PREPROCESSOR_JOIN(LogScope,__LINE__)( FString::Printf(Format, ##__VA_ARGS__));
+#else
+
+#define ABILITY_VLOG_ATTRIBUTE_GRAPH(Actor, Verbosity, AttributeName, OldValue, NewValue)
+
+#endif //ENABLE_VISUAL_LOG
 
 struct AbilitySystemLogScope
 {
@@ -82,4 +110,4 @@ public:
 	int32	Indent;
 	TArray<AbilitySystemLogScope*>	ScopeStack;
 	bool NeedNewLine;
-};
\ No newline at end of file
+};
diff --git a/Engine/Source/Runtime/GameplayAbilities/Public/GameplayAbilitySpec.h b/Engine/Source/Runtime/GameplayAbilities/Public/GameplayAbilitySpec.h
--- a/Engine/Source/Runtime/GameplayAbilities/Public/GameplayAbilitySpec.h
+++ b/Engine/Source/Runtime/GameplayAbilities/Public/GameplayAbilitySpec.h
@@ -181,7 +181,7 @@ private:
 
 
 /** An activatable ability spec, hosted on the ability system component. This defines both what the ability is (what class, what level, input binding etc)
- *  and also holds runttime state that must be kept outside of the ability being instanced/activated.
+ *  and also holds runtime state that must be kept outside of the ability being instanced/activated.
  */
 USTRUCT()
 struct GAMEPLAYABILITIES_API FGameplayAbilitySpec : public FFastArraySerializerItem
diff --git a/Engine/Source/Runtime/GameplayAbilities/Public/GameplayEffect.h b/Engine/Source/Runtime/GameplayAbilities/Public/GameplayEffect.h
--- a/Engine/Source/Runtime/GameplayAbilities/Public/GameplayEffect.h
+++ b/Engine/Source/Runtime/GameplayAbilities/Public/GameplayEffect.h
@@ -614,6 +614,10 @@ public:
 
 	// ----------------------------------------------
 
+	/** If true, cues will only trigger when GE modifiers succeed being applied (whether through modifiers or executions) */
+	UPROPERTY(EditDefaultsOnly, BlueprintReadOnly, Category = Display)
+	bool bRequireModifierSuccessToTriggerCues;
+
 	/** Cues to trigger non-simulated reactions in response to this GameplayEffect such as sounds, particle effects, etc */
 	UPROPERTY(EditDefaultsOnly, BlueprintReadOnly, Category = Display)
 	TArray<FGameplayEffectCue>	GameplayCues;
@@ -1132,6 +1136,12 @@ struct GAMEPLAYABILITIES_API FGameplayEffectSpecForRPC
 	UPROPERTY()
 	FGameplayEffectContextHandle EffectContext; // This tells us how we got here (who / what applied us)
 
+	UPROPERTY()
+	FGameplayTagContainer AggregatedSourceTags;
+
+	UPROPERTY()
+	FGameplayTagContainer AggregatedTargetTags;
+
 	UPROPERTY()
 	float Level;
 
@@ -1200,7 +1210,7 @@ struct GAMEPLAYABILITIES_API FActiveGameplayEffect : public FFastArraySerializer
 		return Spec.GetPeriod();
 	}
 
-	void CheckOngoingTagRequirements(const FGameplayTagContainer& OwnerTags, struct FActiveGameplayEffectsContainer& OwningContainer);
+	void CheckOngoingTagRequirements(const FGameplayTagContainer& OwnerTags, struct FActiveGameplayEffectsContainer& OwningContainer, bool bInvokeGameplayCueEvents = false);
 
 	void PrintAll() const;
 
@@ -1239,6 +1249,10 @@ struct GAMEPLAYABILITIES_API FActiveGameplayEffect : public FFastArraySerializer
 	UPROPERTY()
 	bool bIsInhibited;
 
+	/** When replicated down, we cue the GC events until the entire list of active gameplay effects has been received */
+	mutable bool bPendingRepOnActiveGC;
+	mutable bool bPendingRepWhileActiveGC;
+
 	bool IsPendingRemove;
 
 	FOnActiveGameplayEffectRemoved	OnRemovedDelegate;
@@ -1268,6 +1282,7 @@ struct FActiveGameplayEffectQuery
 		, OwningTagContainer_Rejection(nullptr)
 		, EffectTagContainer_Rejection(nullptr)
 		, EffectSource(nullptr)
+		, EffectDef(nullptr)
 	{
 	}
 
@@ -1277,6 +1292,7 @@ struct FActiveGameplayEffectQuery
 		, OwningTagContainer_Rejection(nullptr)
 		, EffectTagContainer_Rejection(nullptr)
 		, EffectSource(nullptr)
+		, EffectDef(nullptr)
 	{
 	}
 
@@ -1304,6 +1320,9 @@ struct FActiveGameplayEffectQuery
 	// Matches on GameplayEffects which come from this source
 	const UObject* EffectSource;
 
+	// Matches on GameplayEffects with this definition
+	const UGameplayEffect* EffectDef;
+
 	// Handles to ignore as matches, even if other criteria is met
 	TArray<FActiveGameplayEffectHandle> IgnoreHandles;
 };
@@ -1440,10 +1459,7 @@ struct GAMEPLAYABILITIES_API FActiveGameplayEffectsContainer : public FFastArray
 
 	void CheckDuration(FActiveGameplayEffectHandle Handle);
 
-	bool NetDeltaSerialize(FNetDeltaSerializeInfo & DeltaParms)
-	{
-		return FastArrayDeltaSerialize<FActiveGameplayEffect>(GameplayEffects_Internal, DeltaParms, *this);
-	}
+	bool NetDeltaSerialize(FNetDeltaSerializeInfo& DeltaParms);
 
 	void PreDestroy();	
 
@@ -1526,6 +1542,9 @@ private:
 	TArray<FActiveGameplayEffect>	GameplayEffects_Internal;
 
 	void InternalUpdateNumericalAttribute(FGameplayAttribute Attribute, float NewValue, const FGameplayEffectModCallbackData* ModData);
+
+	/** Cached pointer to current mod data needed for callbacks. We cache it in the AGE struct to avoid passing it through all the delegate/aggregator plumbing */
+	const struct FGameplayEffectModCallbackData* CurrentModcallbackData;
 	
 	/**
 	 * Helper function to execute a mod on owned attributes
@@ -1546,8 +1565,8 @@ private:
 	void InternalOnActiveGameplayEffectAdded(FActiveGameplayEffect& Effect);
 	void InternalOnActiveGameplayEffectRemoved(const FActiveGameplayEffect& Effect);
 
-	void RemoveActiveGameplayEffectGrantedTagsAndModifiers(const FActiveGameplayEffect& Effect);
-	void AddActiveGameplayEffectGrantedTagsAndModifiers(FActiveGameplayEffect& Effect);
+	void RemoveActiveGameplayEffectGrantedTagsAndModifiers(const FActiveGameplayEffect& Effect, bool bInvokeGameplayCueEvents);
+	void AddActiveGameplayEffectGrantedTagsAndModifiers(FActiveGameplayEffect& Effect, bool bInvokeGameplayCueEvents);
 
 	/** Updates tag dependency map when a GameplayEffect is removed */
 	void RemoveActiveEffectTagDependency(const FGameplayTagContainer& Tags, FActiveGameplayEffectHandle Handle);
diff --git a/Engine/Source/Runtime/GameplayAbilities/Public/GameplayEffectTypes.h b/Engine/Source/Runtime/GameplayAbilities/Public/GameplayEffectTypes.h
--- a/Engine/Source/Runtime/GameplayAbilities/Public/GameplayEffectTypes.h
+++ b/Engine/Source/Runtime/GameplayAbilities/Public/GameplayEffectTypes.h
@@ -29,6 +29,8 @@ GAMEPLAYABILITIES_API FString EGameplayModToString(int32 Type);
 
 GAMEPLAYABILITIES_API FString EGameplayModEffectToString(int32 Type);
 
+GAMEPLAYABILITIES_API FString EGameplayCueEventToString(int32 Type);
+
 UENUM(BlueprintType)
 namespace EGameplayModOp
 {
@@ -708,6 +710,14 @@ struct FGameplayCueParameters
 	/** The original tag of the gameplay cue */
 	UPROPERTY(BlueprintReadWrite, Category=GameplayCue)
 	FGameplayTag OriginalTag;
+
+	/** The aggregated source tags taken from the effect spec */
+	UPROPERTY(BlueprintReadWrite, Category=GameplayCue)
+	FGameplayTagContainer AggregatedSourceTags;
+
+	/** The aggregated target tags taken from the effect spec */
+	UPROPERTY(BlueprintReadWrite, Category=GameplayCue)
+	FGameplayTagContainer AggregatedTargetTags;
 };
 
 UENUM(BlueprintType)
