merge fixups

diff --git a/Engine/Source/Runtime/GameplayAbilities/Private/Abilities/GameplayAbilityTargetActor.cpp b/Engine/Source/Runtime/GameplayAbilities/Private/Abilities/GameplayAbilityTargetActor.cpp
--- a/Engine/Source/Runtime/GameplayAbilities/Private/Abilities/GameplayAbilityTargetActor.cpp
+++ b/Engine/Source/Runtime/GameplayAbilities/Private/Abilities/GameplayAbilityTargetActor.cpp
@@ -33,8 +33,8 @@ void AGameplayAbilityTargetActor::Destroyed()
 			UAbilitySystemComponent* ASC = Info->AbilitySystemComponent.Get();
 			if (ASC)
 			{
-				ASC->ConfirmCallbacks.RemoveDynamic(this, &AGameplayAbilityTargetActor::ConfirmTargeting);
-				ASC->CancelCallbacks.RemoveDynamic(this, &AGameplayAbilityTargetActor::CancelTargeting);
+				ASC->GenericLocalConfirmCallbacks.RemoveDynamic(this, &AGameplayAbilityTargetActor::ConfirmTargeting);
+				ASC->GenericLocalCancelCallbacks.RemoveDynamic(this, &AGameplayAbilityTargetActor::CancelTargeting);
 			}
 		}
 	}
@@ -142,8 +142,8 @@ void AGameplayAbilityTargetActor::BindToConfirmCancelInputs()
 		if (Info->IsLocallyControlled())
 		{
 			// We have to wait for the callback from the AbilitySystemComponent. Which will always be instigated locally
-			ASC->ConfirmCallbacks.AddDynamic(this, &AGameplayAbilityTargetActor::ConfirmTargeting);	// Tell me if the confirm input is pressed
-			ASC->CancelCallbacks.AddDynamic(this, &AGameplayAbilityTargetActor::CancelTargeting);	// Tell me if the cancel input is pressed
+			ASC->GenericLocalConfirmCallbacks.AddDynamic(this, &AGameplayAbilityTargetActor::ConfirmTargeting);	// Tell me if the confirm input is pressed
+			ASC->GenericLocalCancelCallbacks.AddDynamic(this, &AGameplayAbilityTargetActor::CancelTargeting);	// Tell me if the cancel input is pressed
 		}
 		else
 		{	
diff --git a/Engine/Source/Runtime/GameplayAbilities/Private/Abilities/Tasks/AbilityTask_WaitConfirmCancel.cpp b/Engine/Source/Runtime/GameplayAbilities/Private/Abilities/Tasks/AbilityTask_WaitConfirmCancel.cpp
--- a/Engine/Source/Runtime/GameplayAbilities/Private/Abilities/Tasks/AbilityTask_WaitConfirmCancel.cpp
+++ b/Engine/Source/Runtime/GameplayAbilities/Private/Abilities/Tasks/AbilityTask_WaitConfirmCancel.cpp
@@ -68,8 +68,8 @@ void UAbilityTask_WaitConfirmCancel::Activate()
 		if (Info->IsLocallyControlled())
 		{
 			// We have to wait for the callback from the AbilitySystemComponent.
-			AbilitySystemComponent->ConfirmCallbacks.AddDynamic(this, &UAbilityTask_WaitConfirmCancel::OnLocalConfirmCallback);	// Tell me if the confirm input is pressed
-			AbilitySystemComponent->CancelCallbacks.AddDynamic(this, &UAbilityTask_WaitConfirmCancel::OnLocalCancelCallback);	// Tell me if the cancel input is pressed
+			AbilitySystemComponent->GenericLocalConfirmCallbacks.AddDynamic(this, &UAbilityTask_WaitConfirmCancel::OnLocalConfirmCallback);	// Tell me if the confirm input is pressed
+			AbilitySystemComponent->GenericLocalCancelCallbacks.AddDynamic(this, &UAbilityTask_WaitConfirmCancel::OnLocalCancelCallback);	// Tell me if the cancel input is pressed
 
 			RegisteredCallbacks = true;
 		}
@@ -93,8 +93,8 @@ void UAbilityTask_WaitConfirmCancel::OnDestroy(bool AbilityEnding)
 {
 	if (RegisteredCallbacks && AbilitySystemComponent.IsValid())
 	{
-		AbilitySystemComponent->ConfirmCallbacks.RemoveDynamic(this, &UAbilityTask_WaitConfirmCancel::OnLocalConfirmCallback);
-		AbilitySystemComponent->CancelCallbacks.RemoveDynamic(this, &UAbilityTask_WaitConfirmCancel::OnLocalCancelCallback);
+		AbilitySystemComponent->GenericLocalConfirmCallbacks.RemoveDynamic(this, &UAbilityTask_WaitConfirmCancel::OnLocalConfirmCallback);
+		AbilitySystemComponent->GenericLocalCancelCallbacks.RemoveDynamic(this, &UAbilityTask_WaitConfirmCancel::OnLocalCancelCallback);
 	}
 
 	Super::OnDestroy(AbilityEnding);
diff --git a/Engine/Source/Runtime/GameplayAbilities/Private/AbilitySystemComponent_Abilities.cpp b/Engine/Source/Runtime/GameplayAbilities/Private/AbilitySystemComponent_Abilities.cpp
--- a/Engine/Source/Runtime/GameplayAbilities/Private/AbilitySystemComponent_Abilities.cpp
+++ b/Engine/Source/Runtime/GameplayAbilities/Private/AbilitySystemComponent_Abilities.cpp
@@ -1725,14 +1725,14 @@ void UAbilitySystemComponent::BindToInputComponent(UInputComponent* InputCompone
 	// Pressed event
 	{
 		FInputActionBinding AB(ConfirmBindName, IE_Pressed);
-		AB.ActionDelegate.GetDelegateForManualSet().BindUObject(this, &UAbilitySystemComponent::InputConfirm);
+		AB.ActionDelegate.GetDelegateForManualSet().BindUObject(this, &UAbilitySystemComponent::LocalInputConfirm);
 		InputComponent->AddActionBinding(AB);
 	}
 
 	// 
 	{
 		FInputActionBinding AB(CancelBindName, IE_Pressed);
-		AB.ActionDelegate.GetDelegateForManualSet().BindUObject(this, &UAbilitySystemComponent::InputCancel);
+		AB.ActionDelegate.GetDelegateForManualSet().BindUObject(this, &UAbilitySystemComponent::LocalInputCancel);
 		InputComponent->AddActionBinding(AB);
 	}
 }
@@ -1752,14 +1752,14 @@ void UAbilitySystemComponent::BindAbilityActivationToInputComponent(UInputCompon
 		// Pressed event
 		{
 			FInputActionBinding AB(FName(*BindStr), IE_Pressed);
-			AB.ActionDelegate.GetDelegateForManualSet().BindUObject(this, &UAbilitySystemComponent::AbilityInputPressed, idx);
+			AB.ActionDelegate.GetDelegateForManualSet().BindUObject(this, &UAbilitySystemComponent::AbilityLocalInputPressed, idx);
 			InputComponent->AddActionBinding(AB);
 		}
 
 		// Released event
 		{
 			FInputActionBinding AB(FName(*BindStr), IE_Released);
-			AB.ActionDelegate.GetDelegateForManualSet().BindUObject(this, &UAbilitySystemComponent::AbilityInputReleased, idx);
+			AB.ActionDelegate.GetDelegateForManualSet().BindUObject(this, &UAbilitySystemComponent::AbilityLocalInputReleased, idx);
 			InputComponent->AddActionBinding(AB);
 		}
 	}
@@ -1768,14 +1768,14 @@ void UAbilitySystemComponent::BindAbilityActivationToInputComponent(UInputCompon
 	if (BindInfo.ConfirmTargetCommand.IsEmpty() == false)
 	{
 		FInputActionBinding AB(FName(*BindInfo.ConfirmTargetCommand), IE_Pressed);
-		AB.ActionDelegate.GetDelegateForManualSet().BindUObject(this, &UAbilitySystemComponent::InputConfirm);
+		AB.ActionDelegate.GetDelegateForManualSet().BindUObject(this, &UAbilitySystemComponent::LocalInputConfirm);
 		InputComponent->AddActionBinding(AB);
 	}
 	
 	if (BindInfo.CancelTargetCommand.IsEmpty() == false)
 	{
 		FInputActionBinding AB(FName(*BindInfo.CancelTargetCommand), IE_Pressed);
-		AB.ActionDelegate.GetDelegateForManualSet().BindUObject(this, &UAbilitySystemComponent::InputCancel);
+		AB.ActionDelegate.GetDelegateForManualSet().BindUObject(this, &UAbilitySystemComponent::LocalInputCancel);
 		InputComponent->AddActionBinding(AB);
 	}
 
@@ -1789,18 +1789,18 @@ void UAbilitySystemComponent::BindAbilityActivationToInputComponent(UInputCompon
 	}
 }
 
-void UAbilitySystemComponent::AbilityInputPressed(int32 InputID)
+void UAbilitySystemComponent::AbilityLocalInputPressed(int32 InputID)
 {
 	// Consume the input if this InputID is overloaded with GenericConfirm/Cancel and the GenericConfim/Cancel callback is bound
 	if (IsGenericConfirmInputBound(InputID))
 	{
-		InputConfirm();
+		LocalInputConfirm();
 		return;
 	}
 
 	if (IsGenericCancelInputBound(InputID))
 	{
-		InputCancel();
+		LocalInputCancel();
 		return;
 	}
 
@@ -1836,7 +1836,7 @@ void UAbilitySystemComponent::AbilityInputPressed(int32 InputID)
 	}
 }
 
-void UAbilitySystemComponent::AbilityInputReleased(int32 InputID)
+void UAbilitySystemComponent::AbilityLocalInputReleased(int32 InputID)
 {
 	ABILITYLIST_SCOPE_LOCK();
 	for (FGameplayAbilitySpec& Spec : ActivatableAbilities.Items)
@@ -1930,17 +1930,17 @@ void UAbilitySystemComponent::AbilitySpecInputReleased(FGameplayAbilitySpec& Spe
 	}
 }
 
-void UAbilitySystemComponent::InputConfirm()
+void UAbilitySystemComponent::LocalInputConfirm()
 {
-	FAbilityConfirmOrCancel Temp = ConfirmCallbacks;
-	ConfirmCallbacks.Clear();
+	FAbilityConfirmOrCancel Temp = GenericLocalConfirmCallbacks;
+	GenericLocalConfirmCallbacks.Clear();
 	Temp.Broadcast();
 }
 
-void UAbilitySystemComponent::InputCancel()
+void UAbilitySystemComponent::LocalInputCancel()
 {	
-	FAbilityConfirmOrCancel Temp = CancelCallbacks;
-	CancelCallbacks.Clear();
+	FAbilityConfirmOrCancel Temp = GenericLocalCancelCallbacks;
+	GenericLocalCancelCallbacks.Clear();
 	Temp.Broadcast();
 }
 
