fix generate new handle not being called by new cstor

diff --git a/Engine/Source/Runtime/GameplayAbilities/Private/GameplayAbilityTypes.cpp b/Engine/Source/Runtime/GameplayAbilities/Private/GameplayAbilityTypes.cpp
--- a/Engine/Source/Runtime/GameplayAbilities/Private/GameplayAbilityTypes.cpp
+++ b/Engine/Source/Runtime/GameplayAbilities/Private/GameplayAbilityTypes.cpp
@@ -172,6 +172,7 @@ FGameplayAbilitySpec:: FGameplayAbilitySpec(FGameplayAbilitySpecDef& InDef, FAct
 	, RemoveAfterActivation(false)
 	, PendingRemove(false)
 {
+	Handle.GenerateNewHandle();
 	InDef.AssignedHandle = Handle;
 	GameplayEffectHandle = InGameplayEffectHandle;
 }
