#UE4 Mering gameplaycuemanager up to head revision to get a couple misc tweaks and merge resolves.

diff --git a/Engine/Source/Runtime/GameplayAbilities/Private/GameplayCueManager.cpp b/Engine/Source/Runtime/GameplayAbilities/Private/GameplayCueManager.cpp
--- a/Engine/Source/Runtime/GameplayAbilities/Private/GameplayCueManager.cpp
+++ b/Engine/Source/Runtime/GameplayAbilities/Private/GameplayCueManager.cpp
@@ -4,6 +4,7 @@
 #include "AssetRegistryModule.h"
 #include "GameplayCueInterface.h"
 #include "GameplayCueManager.h"
+#include "GameplayCueSet.h"
 #include "GameplayTagsModule.h"
 #include "GameplayCueNotify_Static.h"
 #include "AbilitySystemComponent.h"
@@ -29,7 +30,7 @@ UGameplayCueManager::UGameplayCueManager(const FObjectInitializer& PCIP)
 	RegisteredEditorCallbacks = false;
 #endif
 
-	GlobalCueSet = NewNamedObject<UGameplayCueSet>(this, TEXT("GlobalCueSet"));
+	GlobalCueSet = NewObject<UGameplayCueSet>(this, TEXT("GlobalCueSet"));
 }
 
 
@@ -60,7 +61,6 @@ void UGameplayCueManager::HandleGameplayCue(AActor* TargetActor, FGameplayTag Ga
 		FColor DebugColor = FColor::Green;
 		DrawDebugString(TargetActor->GetWorld(), FVector(0.f, 0.f, 100.f), DebugStr, TargetActor, DebugColor, DisplayGameplayCueDuration);
 	}
-
 	// Give the global set a chance
 	check(GlobalCueSet);
 	GlobalCueSet->HandleGameplayCue(TargetActor, GameplayCueTag, EventType, Parameters);
@@ -103,7 +103,9 @@ AGameplayCueNotify_Actor* UGameplayCueManager::GetInstancedCueActor(AActor* Targ
 	AGameplayCueNotify_Actor* SpawnedCue = nullptr;
 	if (ensure(TargetActor) && ensure(CueClass))
 	{
-		SpawnedCue = TargetActor->GetWorld()->SpawnActor<AGameplayCueNotify_Actor>(CueClass, TargetActor->GetActorLocation(), TargetActor->GetActorRotation());
+		FActorSpawnParameters SpawnParams;
+		SpawnParams.Owner = TargetActor;
+		SpawnedCue = TargetActor->GetWorld()->SpawnActor<AGameplayCueNotify_Actor>(CueClass, TargetActor->GetActorLocation(), TargetActor->GetActorRotation(), SpawnParams);
 		if (ensure(SpawnedCue))
 		{
 			auto& InnerMap = NotifyMapActor.Add(TargetActor);
@@ -230,7 +232,7 @@ void UGameplayCueManager::BuildCuesToAddToGlobalSet(const TArray<FAssetData>& As
 			{
 				// Add a new NotifyData entry to our flat list for this one
 				FStringAssetReference StringRef;
-				StringRef.AssetLongPathname = *GeneratedClassTag;
+				StringRef.AssetLongPathname = FPackageName::ExportTextPathToObjectPath(*GeneratedClassTag);
 
 				OutCuesToAdd.Add(FGameplayCueReferencePair(GameplayCueTag, StringRef));
 
