GameplayAbilities: Fixing non-unity build

diff --git a/Engine/Source/Runtime/GameplayAbilities/Private/GameplayCueSet.cpp b/Engine/Source/Runtime/GameplayAbilities/Private/GameplayCueSet.cpp
--- a/Engine/Source/Runtime/GameplayAbilities/Private/GameplayCueSet.cpp
+++ b/Engine/Source/Runtime/GameplayAbilities/Private/GameplayCueSet.cpp
@@ -2,7 +2,8 @@
 
 #include "AbilitySystemPrivatePCH.h"
 #include "GameplayCueSet.h"
-
+#include "GameplayCueManager.h"
+#include "GameplayCueNotify_Static.h"
 
 // --------------------------------------------------------------------------------------------------------------------------------------------------------
 //
diff --git a/Engine/Source/Runtime/GameplayAbilities/Public/GameplayCueManager.h b/Engine/Source/Runtime/GameplayAbilities/Public/GameplayCueManager.h
--- a/Engine/Source/Runtime/GameplayAbilities/Public/GameplayCueManager.h
+++ b/Engine/Source/Runtime/GameplayAbilities/Public/GameplayCueManager.h
@@ -2,6 +2,8 @@
 
 #pragma once
 
+class UGameplayCueSet;
+
 #include "GameplayTags.h"
 #include "GameplayEffect.h"
 #include "GameplayCueNotify_Actor.h"
