Ensured that FScalableFloat::FinalCurve is updated when the curve table it's referencing is re-imported

UDN - FRichCurve in UAttributeSet becomes invalid after auto reimport

diff --git a/Engine/Source/Runtime/GameplayAbilities/Private/AttributeSet.cpp b/Engine/Source/Runtime/GameplayAbilities/Private/AttributeSet.cpp
--- a/Engine/Source/Runtime/GameplayAbilities/Private/AttributeSet.cpp
+++ b/Engine/Source/Runtime/GameplayAbilities/Private/AttributeSet.cpp
@@ -153,8 +153,10 @@ float FScalableFloat::GetValueAtLevel(float Level) const
 	{
 		if (FinalCurve == nullptr)
 		{
-			static const FString ContextString = TEXT("FScalableFloat::FinalizeCurveData");
+			static const FString ContextString = TEXT("FScalableFloat::GetValueAtLevel");
 			FinalCurve = Curve.GetCurve(ContextString);
+
+			RegisterOnCurveTablePostReimport();
 		}
 
 		if (FinalCurve != nullptr)
@@ -170,6 +172,58 @@ float FScalableFloat::GetValueAtLevel(float Level) const
 	return Value;
 }
 
+void FScalableFloat::SetValue(float NewValue)
+{
+	UnRegisterOnCurveTablePostReimport();
+
+	Value = NewValue;
+	Curve.CurveTable = nullptr;
+	Curve.RowName = NAME_None;
+	FinalCurve = nullptr;
+}
+
+void FScalableFloat::SetScalingValue(float InCoeffecient, FName InRowName, UCurveTable * InTable)
+{
+	UnRegisterOnCurveTablePostReimport();
+
+	Value = InCoeffecient;
+	Curve.RowName = InRowName;
+	Curve.CurveTable = InTable;
+	FinalCurve = nullptr;
+}
+
+void FScalableFloat::RegisterOnCurveTablePostReimport() const
+{
+#if WITH_EDITOR
+	if (!OnCurveTablePostReimportHandle.IsValid())
+	{
+		// Register our interest in knowing when our referenced curve table is changed, so that we can update FinalCurve appropriately
+		OnCurveTablePostReimportHandle = FReimportManager::Instance()->OnPostReimport().AddRaw(this, &FScalableFloat::OnCurveTablePostReimport);
+	}
+#endif // WITH_EDITOR
+}
+
+void FScalableFloat::UnRegisterOnCurveTablePostReimport() const
+{
+#if WITH_EDITOR
+	if (OnCurveTablePostReimportHandle.IsValid())
+	{
+		FReimportManager::Instance()->OnPostReimport().Remove(OnCurveTablePostReimportHandle);
+	}
+#endif // WITH_EDITOR
+}
+
+#if WITH_EDITOR
+void FScalableFloat::OnCurveTablePostReimport(UObject* InObject, bool)
+{
+	if (Curve.CurveTable && Curve.CurveTable == InObject)
+	{
+		// Reset FinalCurve so that GetValueAtLevel will re-cache it the next time it gets called
+		FinalCurve = nullptr;
+	}
+}
+#endif // WITH_EDITOR
+
 bool FGameplayAttribute::operator==(const FGameplayAttribute& Other) const
 {
 	return ((Other.Attribute == Attribute));
diff --git a/Engine/Source/Runtime/GameplayAbilities/Public/AttributeSet.h b/Engine/Source/Runtime/GameplayAbilities/Public/AttributeSet.h
--- a/Engine/Source/Runtime/GameplayAbilities/Public/AttributeSet.h
+++ b/Engine/Source/Runtime/GameplayAbilities/Public/AttributeSet.h
@@ -160,16 +160,21 @@ struct GAMEPLAYABILITIES_API FScalableFloat
 
 	FScalableFloat()
 		: Value(0.f)
-		, FinalCurve(NULL)
+		, FinalCurve(nullptr)
 	{
 	}
 
 	FScalableFloat(float InInitialValue)
 		: Value(InInitialValue)
-		, FinalCurve(NULL)
+		, FinalCurve(nullptr)
 	{
 	}
 
+	~FScalableFloat()
+	{
+		UnRegisterOnCurveTablePostReimport();
+	}
+
 public:
 
 	UPROPERTY(EditDefaultsOnly, Category=ScalableFloat)
@@ -182,24 +187,12 @@ public:
 
 	bool IsStatic() const
 	{
-		return (Curve.RowName == NAME_None);
+		return Curve.RowName.IsNone();
 	}
 
-	void SetValue(float NewValue)
-	{
-		Value = NewValue;
-		Curve.CurveTable = NULL;
-		Curve.RowName = NAME_None;
-		FinalCurve = NULL;
-	}
+	void SetValue(float NewValue);
 
-	void SetScalingValue(float InCoeffecient, FName InRowName, UCurveTable * InTable)
-	{
-		Value = InCoeffecient;
-		Curve.RowName = InRowName;
-		Curve.CurveTable = InTable;
-		FinalCurve = NULL;
-	}
+	void SetScalingValue(float InCoeffecient, FName InRowName, UCurveTable * InTable);
 
 	void LockValueAtLevel(float Level, FGlobalCurveDataOverride *GlobalOverrides)
 	{
@@ -227,6 +220,20 @@ public:
 
 private:
 
+	/** Conditionally register the OnCurveTablePostReimport function with the re-import manager */
+	void RegisterOnCurveTablePostReimport() const;
+
+	/** Conditionally unregister the OnCurveTablePostReimport function from the re-import manager */
+	void UnRegisterOnCurveTablePostReimport() const;
+
+#if WITH_EDITOR
+	/** Function to call when the curve table has been re-imported */
+	void OnCurveTablePostReimport(UObject* InObject, bool);
+
+	/** Handle to the delegate called when the curve table has been re-imported */
+	mutable FDelegateHandle OnCurveTablePostReimportHandle;
+#endif // WITH_EDITOR
+
 	// Cached direct pointer to RichCurve we should evaluate
 	mutable FRichCurve* FinalCurve;
 };
