Merging

//depot/UE4-Orion/Engine/Source/Runtime/GameplayAbilities/Public/Abilities/GameplayAbility.h

diff --git a/Engine/Source/Runtime/GameplayAbilities/Public/Abilities/GameplayAbility.h b/Engine/Source/Runtime/GameplayAbilities/Public/Abilities/GameplayAbility.h
--- a/Engine/Source/Runtime/GameplayAbilities/Public/Abilities/GameplayAbility.h
+++ b/Engine/Source/Runtime/GameplayAbilities/Public/Abilities/GameplayAbility.h
@@ -567,7 +567,7 @@ public:
 	/** Returns true if this ability can be canceled */
 	virtual bool CanBeCanceled() const;
 
-	/** Sets rather the ability should ignore cancel requests. Only valid on instanced abilities */
+	/** Sets whether the ability should ignore cancel requests. Only valid on instanced abilities */
 	UFUNCTION(BlueprintCallable, Category=Ability)
 	virtual void SetCanBeCanceled(bool bCanBeCanceled);
 
