non unity compile fixes

diff --git a/Engine/Source/Runtime/GameplayAbilities/Private/GameplayEffectTypes.cpp b/Engine/Source/Runtime/GameplayAbilities/Private/GameplayEffectTypes.cpp
--- a/Engine/Source/Runtime/GameplayAbilities/Private/GameplayEffectTypes.cpp
+++ b/Engine/Source/Runtime/GameplayAbilities/Private/GameplayEffectTypes.cpp
@@ -481,3 +481,15 @@ const FGameplayTagContainer& FTagContainerAggregator::GetSpecTags() const
 	CacheIsValid = false;
 	return CapturedSpecTags;
 }
+
+
+FGameplayEffectSpecHandle::FGameplayEffectSpecHandle()
+{
+
+}
+
+FGameplayEffectSpecHandle::FGameplayEffectSpecHandle(FGameplayEffectSpec* DataPtr)
+	: Data(DataPtr)
+{
+
+}
\ No newline at end of file
diff --git a/Engine/Source/Runtime/GameplayAbilities/Public/GameplayEffectTypes.h b/Engine/Source/Runtime/GameplayAbilities/Public/GameplayEffectTypes.h
--- a/Engine/Source/Runtime/GameplayAbilities/Public/GameplayEffectTypes.h
+++ b/Engine/Source/Runtime/GameplayAbilities/Public/GameplayEffectTypes.h
@@ -941,12 +941,8 @@ struct GAMEPLAYABILITIES_API FGameplayEffectSpecHandle
 {
 	GENERATED_USTRUCT_BODY()
 
-	FGameplayEffectSpecHandle() { }
-	FGameplayEffectSpecHandle(FGameplayEffectSpec* DataPtr)
-		: Data(DataPtr)
-	{
-
-	}
+	FGameplayEffectSpecHandle();
+	FGameplayEffectSpecHandle(FGameplayEffectSpec* DataPtr);
 
 	TSharedPtr<FGameplayEffectSpec>	Data;
 
