Merging

//depot/UE4-Orion/Engine/Source/Runtime/GameplayAbilities/Private/...

to //depot/UE4/Engine/Source/Runtime/GameplayAbilities/Private/...

diff --git a/Engine/Source/Runtime/GameplayAbilities/Private/GameplayCueManager.cpp b/Engine/Source/Runtime/GameplayAbilities/Private/GameplayCueManager.cpp
--- a/Engine/Source/Runtime/GameplayAbilities/Private/GameplayCueManager.cpp
+++ b/Engine/Source/Runtime/GameplayAbilities/Private/GameplayCueManager.cpp
@@ -91,16 +91,21 @@ void UGameplayCueManager::EndGameplayCuesFor(AActor* TargetActor)
 
 AGameplayCueNotify_Actor* UGameplayCueManager::GetInstancedCueActor(AActor* TargetActor, UClass* CueClass)
 {
+	AGameplayCueNotify_Actor* SpawnedCue = nullptr;
 	if (auto InnerMap = NotifyMapActor.Find(TargetActor))
 	{
 		if (auto WeakPtrPtr = InnerMap->Find(CueClass))
 		{
-			return WeakPtrPtr->Get();
+			SpawnedCue = WeakPtrPtr->Get();
+			// If the cue is scheduled to be destroyed, don't reuse it, create a new one instead
+			if (SpawnedCue && SpawnedCue->GetLifeSpan() > 0.0f)
+			{
+				return SpawnedCue;
+			}
 		}
 	}
 
 	// We don't have an instance for this, and we need one, so make one
-	AGameplayCueNotify_Actor* SpawnedCue = nullptr;
 	if (ensure(TargetActor) && ensure(CueClass))
 	{
 		FActorSpawnParameters SpawnParams;
diff --git a/Engine/Source/Runtime/GameplayAbilities/Private/GameplayCueSet.cpp b/Engine/Source/Runtime/GameplayAbilities/Private/GameplayCueSet.cpp
--- a/Engine/Source/Runtime/GameplayAbilities/Private/GameplayCueSet.cpp
+++ b/Engine/Source/Runtime/GameplayAbilities/Private/GameplayCueSet.cpp
@@ -42,15 +42,22 @@ void UGameplayCueSet::AddCues(const TArray<FGameplayCueReferencePair>& CuesToAdd
 			const FStringAssetReference& StringRef = CueRefPair.StringRef;
 
 			// Check for duplicates: we may want to remove this eventually.
+			bool bDupe = false;
 			for (FGameplayCueNotifyData Data : GameplayCueData)
 			{
 				if (Data.GameplayCueTag == GameplayCueTag)
 				{
 					ABILITY_LOG(Warning, TEXT("AddGameplayCueData_Internal called for [%s,%s] when it already existed [%s,%s]. Skipping."), *GameplayCueTag.ToString(), *StringRef.ToString(), *Data.GameplayCueTag.ToString(), *Data.GameplayCueNotifyObj.ToString());
-					return;
+					bDupe = true;
+					break;
 				}
 			}
 
+			if (bDupe)
+			{
+				continue;
+			}
+
 			FGameplayCueNotifyData NewData;
 			NewData.GameplayCueNotifyObj = StringRef;
 			NewData.GameplayCueTag = GameplayCueTag;
