UE4-Orion-To-UE4(2525237): Fixed leaking AGameplayCueNotify_Actor;

* The life span check used for finding AGameplayCueNotify_Actor was reversed.

diff --git a/Engine/Source/Runtime/GameplayAbilities/Private/GameplayCueManager.cpp b/Engine/Source/Runtime/GameplayAbilities/Private/GameplayCueManager.cpp
--- a/Engine/Source/Runtime/GameplayAbilities/Private/GameplayCueManager.cpp
+++ b/Engine/Source/Runtime/GameplayAbilities/Private/GameplayCueManager.cpp
@@ -98,7 +98,7 @@ AGameplayCueNotify_Actor* UGameplayCueManager::GetInstancedCueActor(AActor* Targ
 		{
 			SpawnedCue = WeakPtrPtr->Get();
 			// If the cue is scheduled to be destroyed, don't reuse it, create a new one instead
-			if (SpawnedCue && SpawnedCue->GetLifeSpan() > 0.0f)
+			if (SpawnedCue && SpawnedCue->GetLifeSpan() <= 0.0f)
 			{
 				return SpawnedCue;
 			}
