Better logging for activation count overflow

diff --git a/Engine/Source/Runtime/GameplayAbilities/Private/AbilitySystemComponent_Abilities.cpp b/Engine/Source/Runtime/GameplayAbilities/Private/AbilitySystemComponent_Abilities.cpp
--- a/Engine/Source/Runtime/GameplayAbilities/Private/AbilitySystemComponent_Abilities.cpp
+++ b/Engine/Source/Runtime/GameplayAbilities/Private/AbilitySystemComponent_Abilities.cpp
@@ -1067,7 +1067,7 @@ bool UAbilitySystemComponent::InternalTryActivateAbility(FGameplayAbilitySpecHan
 	}
 	else
 	{
-		ABILITY_LOG(Warning, TEXT("TryActivateAbility called when the Spec->ActiveCount >= UINT8_MAX"));
+		ABILITY_LOG(Warning, TEXT("TryActivateAbility %s called when the Spec->ActiveCount (%d) >= UINT8_MAX"), *Ability->GetName(), (int32)Spec->ActiveCount);
 	}
 
 	// Setup a fresh ActivationInfo for this AbilitySpec.
