Merging

Fix issue with gameplay cue actor inner map getting stomped/recreated each time a new class came through

diff --git a/Engine/Source/Runtime/GameplayAbilities/Private/GameplayCueManager.cpp b/Engine/Source/Runtime/GameplayAbilities/Private/GameplayCueManager.cpp
--- a/Engine/Source/Runtime/GameplayAbilities/Private/GameplayCueManager.cpp
+++ b/Engine/Source/Runtime/GameplayAbilities/Private/GameplayCueManager.cpp
@@ -113,7 +113,7 @@ AGameplayCueNotify_Actor* UGameplayCueManager::GetInstancedCueActor(AActor* Targ
 		SpawnedCue = TargetActor->GetWorld()->SpawnActor<AGameplayCueNotify_Actor>(CueClass, TargetActor->GetActorLocation(), TargetActor->GetActorRotation(), SpawnParams);
 		if (ensure(SpawnedCue))
 		{
-			auto& InnerMap = NotifyMapActor.Add(TargetActor);
+			auto& InnerMap = NotifyMapActor.FindOrAdd(TargetActor);
 			InnerMap.Add(CueClass) = SpawnedCue;
 		}
 	}
