Merging
Fix issue with removing predictive gameplay cue notifies. Gameplay cues can be predictively added but are not always predictively removed. Track removal prediction separately.

diff --git a/Engine/Source/Runtime/GameplayAbilities/Private/AbilitySystemComponent.cpp b/Engine/Source/Runtime/GameplayAbilities/Private/AbilitySystemComponent.cpp
--- a/Engine/Source/Runtime/GameplayAbilities/Private/AbilitySystemComponent.cpp
+++ b/Engine/Source/Runtime/GameplayAbilities/Private/AbilitySystemComponent.cpp
@@ -988,6 +988,7 @@ void UAbilitySystemComponent::RemoveGameplayCue(const FGameplayTag GameplayCueTa
 	}
 	else if (ScopedPredictionKey.IsLocalClientKey())
 	{
+		ActiveGameplayCues.PredictiveRemove(GameplayCueTag);
 		InvokeGameplayCueEvent(GameplayCueTag, EGameplayCueEvent::Removed);
 	}
 }
diff --git a/Engine/Source/Runtime/GameplayAbilities/Private/GameplayCueInterface.cpp b/Engine/Source/Runtime/GameplayAbilities/Private/GameplayCueInterface.cpp
--- a/Engine/Source/Runtime/GameplayAbilities/Private/GameplayCueInterface.cpp
+++ b/Engine/Source/Runtime/GameplayAbilities/Private/GameplayCueInterface.cpp
@@ -115,7 +115,9 @@ void IGameplayCueInterface::ForwardGameplayCueToParent()
 
 void FActiveGameplayCue::PreReplicatedRemove(const struct FActiveGameplayCueContainer &InArray)
 {
-	if (PredictionKey.IsLocalClientKey() == false)
+	// We don't check the PredictionKey here like we do in PostReplicatedAdd. PredictionKey tells us
+	// if we were predictely created, but this doesn't mean we will predictively remove ourselves.
+	if (bPredictivelyRemoved == false)
 	{
 		// If predicted ignore the add/remove
 		InArray.Owner->InvokeGameplayCueEvent(GameplayCueTag, EGameplayCueEvent::Removed);
@@ -161,4 +163,17 @@ void FActiveGameplayCueContainer::RemoveCue(const FGameplayTag& Tag)
 			return;
 		}
 	}
+}
+
+void FActiveGameplayCueContainer::PredictiveRemove(const FGameplayTag& Tag)
+{
+	for (int32 idx=0; idx < GameplayCues.Num(); ++idx)
+	{
+		FActiveGameplayCue& Cue = GameplayCues[idx];
+		if (Cue.GameplayCueTag == Tag)
+		{			
+			Cue.bPredictivelyRemoved = true;
+			return;
+		}
+	}
 }
\ No newline at end of file
diff --git a/Engine/Source/Runtime/GameplayAbilities/Public/GameplayCueInterface.h b/Engine/Source/Runtime/GameplayAbilities/Public/GameplayCueInterface.h
--- a/Engine/Source/Runtime/GameplayAbilities/Public/GameplayCueInterface.h
+++ b/Engine/Source/Runtime/GameplayAbilities/Public/GameplayCueInterface.h
@@ -63,12 +63,21 @@ struct FActiveGameplayCue : public FFastArraySerializerItem
 {
 	GENERATED_USTRUCT_BODY()
 
+	FActiveGameplayCue()	
+	{
+		bPredictivelyRemoved = false;
+	}
+
 	UPROPERTY()
 	FGameplayTag GameplayCueTag;
 
 	UPROPERTY()
 	FPredictionKey PredictionKey;
 
+	/** Has this been predictively removed on the client? */
+	UPROPERTY(NotReplicated)
+	bool bPredictivelyRemoved;
+
 	void PreReplicatedRemove(const struct FActiveGameplayCueContainer &InArray);
 	void PostReplicatedAdd(const struct FActiveGameplayCueContainer &InArray);
 	void PostReplicatedChange(const struct FActiveGameplayCueContainer &InArray) { }
@@ -88,6 +97,9 @@ struct FActiveGameplayCueContainer : public FFastArraySerializer
 	void AddCue(const FGameplayTag& Tag, const FPredictionKey& PredictionKey);
 	void RemoveCue(const FGameplayTag& Tag);
 
+	/** Marks as predictively removed so that we dont invoke remove event twice due to onrep */
+	void PredictiveRemove(const FGameplayTag& Tag);
+
 	bool NetDeltaSerialize(FNetDeltaSerializeInfo & DeltaParms)
 	{
 		return FastArrayDeltaSerialize<FActiveGameplayCue>(GameplayCues, DeltaParms, *this);
@@ -105,4 +117,4 @@ struct TStructOpsTypeTraits< FActiveGameplayCueContainer > : public TStructOpsTy
 	{
 		WithNetDeltaSerializer = true,
 	};
-};
\ No newline at end of file
+};
