Add optional vector interpolation curve to AbilityTask_MoveToLocation to support per-axis interpolation

diff --git a/Engine/Source/Runtime/GameplayAbilities/Private/Abilities/Tasks/AbilityTask_MoveToLocation.cpp b/Engine/Source/Runtime/GameplayAbilities/Private/Abilities/Tasks/AbilityTask_MoveToLocation.cpp
--- a/Engine/Source/Runtime/GameplayAbilities/Private/Abilities/Tasks/AbilityTask_MoveToLocation.cpp
+++ b/Engine/Source/Runtime/GameplayAbilities/Private/Abilities/Tasks/AbilityTask_MoveToLocation.cpp
@@ -14,7 +14,7 @@ UAbilityTask_MoveToLocation::UAbilityTask_MoveToLocation(const FObjectInitialize
 	bIsFinished = false;
 }
 
-UAbilityTask_MoveToLocation* UAbilityTask_MoveToLocation::MoveToLocation(class UObject* WorldContextObject, FName TaskInstanceName, FVector Location, float Duration, UCurveFloat* OptionalInterpolationCurve)
+UAbilityTask_MoveToLocation* UAbilityTask_MoveToLocation::MoveToLocation(UObject* WorldContextObject, FName TaskInstanceName, FVector Location, float Duration, UCurveFloat* OptionalInterpolationCurve, UCurveVector* OptionalVectorInterpolationCurve)
 {
 	auto MyObj = NewTask<UAbilityTask_MoveToLocation>(WorldContextObject, TaskInstanceName);
 
@@ -28,6 +28,7 @@ UAbilityTask_MoveToLocation* UAbilityTask_MoveToLocation::MoveToLocation(class U
 	MyObj->TimeMoveStarted = MyObj->GetWorld()->GetTimeSeconds();
 	MyObj->TimeMoveWillEnd = MyObj->TimeMoveStarted + MyObj->DurationOfMovement;
 	MyObj->LerpCurve = OptionalInterpolationCurve;
+	MyObj->LerpCurveVector = OptionalVectorInterpolationCurve;
 
 	return MyObj;
 }
@@ -85,13 +86,25 @@ void UAbilityTask_MoveToLocation::TickTask(float DeltaTime)
 		}
 		else
 		{
+			FVector NewLocation;
+
 			float MoveFraction = (CurrentTime - TimeMoveStarted) / DurationOfMovement;
-			if (LerpCurve)
+			if (LerpCurveVector)
+			{
+				const FVector ComponentInterpolationFraction = LerpCurveVector->GetVectorValue(MoveFraction);
+				NewLocation = FMath::Lerp<FVector, FVector>(StartLocation, TargetLocation, ComponentInterpolationFraction);
+			}
+			else
 			{
-				MoveFraction = LerpCurve->GetFloatValue(MoveFraction);
+				if (LerpCurve)
+				{
+					MoveFraction = LerpCurve->GetFloatValue(MoveFraction);
+				}
+
+				NewLocation = FMath::Lerp<FVector, float>(StartLocation, TargetLocation, MoveFraction);
 			}
 
-			MyActor->SetActorLocation(FMath::Lerp<FVector, float>(StartLocation, TargetLocation, MoveFraction));
+			MyActor->SetActorLocation(NewLocation);
 		}
 	}
 	else
diff --git a/Engine/Source/Runtime/GameplayAbilities/Public/Abilities/Tasks/AbilityTask_MoveToLocation.h b/Engine/Source/Runtime/GameplayAbilities/Public/Abilities/Tasks/AbilityTask_MoveToLocation.h
--- a/Engine/Source/Runtime/GameplayAbilities/Public/Abilities/Tasks/AbilityTask_MoveToLocation.h
+++ b/Engine/Source/Runtime/GameplayAbilities/Public/Abilities/Tasks/AbilityTask_MoveToLocation.h
@@ -30,9 +30,9 @@ class UAbilityTask_MoveToLocation : public UAbilityTask
 
 	virtual void InitSimulatedTask(UAbilitySystemComponent* InAbilitySystemComponent) override;
 
-	/** Move to the specified location, using the curve (range 0 - 1) or linearly if no curve is specified */
+	/** Move to the specified location, using the vector curve (range 0 - 1) if specified, otherwise the float curve (range 0 - 1) or fallback to linear interpolation */
 	UFUNCTION(BlueprintCallable, Category = "Ability|Tasks", meta = (HidePin = "WorldContextObject", DefaultToSelf = "WorldContextObject", BlueprintInternalUseOnly = "TRUE"))
-	static UAbilityTask_MoveToLocation* MoveToLocation(UObject* WorldContextObject, FName TaskInstanceName, FVector Location, float Duration, UCurveFloat* OptionalInterpolationCurve);
+	static UAbilityTask_MoveToLocation* MoveToLocation(UObject* WorldContextObject, FName TaskInstanceName, FVector Location, float Duration, UCurveFloat* OptionalInterpolationCurve, UCurveVector* OptionalVectorInterpolationCurve);
 
 	virtual void Activate() override;
 
@@ -64,5 +64,8 @@ protected:
 	UPROPERTY(Replicated)
 	UCurveFloat* LerpCurve;
 
+	UPROPERTY(Replicated)
+	UCurveVector* LerpCurveVector;
+
 	
 };
\ No newline at end of file
