Fix AbilityTask_MoveToLocation vector lerp curve not affecting simulated clients

diff --git a/Engine/Source/Runtime/GameplayAbilities/Private/Abilities/Tasks/AbilityTask_MoveToLocation.cpp b/Engine/Source/Runtime/GameplayAbilities/Private/Abilities/Tasks/AbilityTask_MoveToLocation.cpp
--- a/Engine/Source/Runtime/GameplayAbilities/Private/Abilities/Tasks/AbilityTask_MoveToLocation.cpp
+++ b/Engine/Source/Runtime/GameplayAbilities/Private/Abilities/Tasks/AbilityTask_MoveToLocation.cpp
@@ -120,6 +120,7 @@ void UAbilityTask_MoveToLocation::GetLifetimeReplicatedProps(TArray< FLifetimePr
 	DOREPLIFETIME(UAbilityTask_MoveToLocation, TargetLocation);
 	DOREPLIFETIME(UAbilityTask_MoveToLocation, DurationOfMovement);
 	DOREPLIFETIME(UAbilityTask_MoveToLocation, LerpCurve);
+	DOREPLIFETIME(UAbilityTask_MoveToLocation, LerpCurveVector);
 }
 
 void UAbilityTask_MoveToLocation::OnDestroy(bool AbilityIsEnding)
