Merging CL 2538383

Expose some gameplaycueset data as public for gameplaycue manager to access.

diff --git a/Engine/Source/Runtime/GameplayAbilities/Public/GameplayCueSet.h b/Engine/Source/Runtime/GameplayAbilities/Public/GameplayCueSet.h
--- a/Engine/Source/Runtime/GameplayAbilities/Public/GameplayCueSet.h
+++ b/Engine/Source/Runtime/GameplayAbilities/Public/GameplayCueSet.h
@@ -62,17 +62,17 @@ class GAMEPLAYABILITIES_API UGameplayCueSet : public UDataAsset
 	virtual void Empty();
 
 	virtual void PrintCues() const;
-
-protected:
-	virtual bool HandleGameplayCueNotify_Internal(AActor* TargetActor, int32 DataIdx, EGameplayCueEvent::Type EventType, FGameplayCueParameters Parameters);
-	virtual void BuildAccelerationMap_Internal();
 	
-	static FGameplayTag	BaseGameplayCueTag();
-
-protected:
 	UPROPERTY(EditAnywhere, Category=CueSet)
 	TArray<FGameplayCueNotifyData> GameplayCueData;
 
 	/** Maps GameplayCue Tag to index into above GameplayCues array. */
 	TMap<FGameplayTag, int32> GameplayCueDataMap;
+
+protected:
+	virtual bool HandleGameplayCueNotify_Internal(AActor* TargetActor, int32 DataIdx, EGameplayCueEvent::Type EventType, FGameplayCueParameters Parameters);
+	virtual void BuildAccelerationMap_Internal();
+	
+	static FGameplayTag	BaseGameplayCueTag();
+	
 };
\ No newline at end of file
