Merging

remove ensure that was reinstated after merge

diff --git a/Engine/Source/Runtime/GameplayAbilities/Private/AbilitySystemComponent_Abilities.cpp b/Engine/Source/Runtime/GameplayAbilities/Private/AbilitySystemComponent_Abilities.cpp
--- a/Engine/Source/Runtime/GameplayAbilities/Private/AbilitySystemComponent_Abilities.cpp
+++ b/Engine/Source/Runtime/GameplayAbilities/Private/AbilitySystemComponent_Abilities.cpp
@@ -897,8 +897,15 @@ bool UAbilitySystemComponent::TryActivateAbility(FGameplayAbilitySpecHandle Abil
 		return false;
 	}
 
+		
 	ENetRole NetMode = ActorInfo->AvatarActor->Role;
-	ensure(NetMode != ROLE_SimulatedProxy);
+
+	// This should only come from button presses/local instigation (AI, etc).
+	if (NetMode == ROLE_SimulatedProxy)
+	{
+		return false;
+	}
+
 	bool bIsLocal = AbilityActorInfo->IsLocallyControlled();
 
 	// Check to see if this a local only or server only ability, if so either remotely execute or fail
