GameplayCueNotify tags should select only within GameplayCue.* tags

diff --git a/Engine/Source/Runtime/GameplayAbilities/Public/GameplayCueNotify_Actor.h b/Engine/Source/Runtime/GameplayAbilities/Public/GameplayCueNotify_Actor.h
--- a/Engine/Source/Runtime/GameplayAbilities/Public/GameplayCueNotify_Actor.h
+++ b/Engine/Source/Runtime/GameplayAbilities/Public/GameplayCueNotify_Actor.h
@@ -58,7 +58,7 @@ class GAMEPLAYABILITIES_API AGameplayCueNotify_Actor : public AActor
 	UFUNCTION(BlueprintNativeEvent, Category = "GameplayCueNotify")
 	bool OnRemove(AActor* MyTarget, FGameplayCueParameters Parameters);
 
-	UPROPERTY(EditDefaultsOnly, Category = GameplayCue)
+	UPROPERTY(EditDefaultsOnly, Category=GameplayCue, meta=(Categories="GameplayCue"))
 	FGameplayTag	GameplayCueTag;
 
 	/** Mirrors GameplayCueTag in order to be asset registry searchable */
diff --git a/Engine/Source/Runtime/GameplayAbilities/Public/GameplayCueNotify_Static.h b/Engine/Source/Runtime/GameplayAbilities/Public/GameplayCueNotify_Static.h
--- a/Engine/Source/Runtime/GameplayAbilities/Public/GameplayCueNotify_Static.h
+++ b/Engine/Source/Runtime/GameplayAbilities/Public/GameplayCueNotify_Static.h
@@ -55,7 +55,7 @@ class GAMEPLAYABILITIES_API UGameplayCueNotify_Static : public UObject
 	UFUNCTION(BlueprintNativeEvent, BlueprintPure, Category = "GameplayCueNotify")
 	bool OnRemove(AActor* MyTarget, FGameplayCueParameters Parameters) const;
 
-	UPROPERTY(EditDefaultsOnly, Category = GameplayCue)
+	UPROPERTY(EditDefaultsOnly, Category = GameplayCue, meta=(Categories="GameplayCue"))
 	FGameplayTag	GameplayCueTag;
 
 	/** Mirrors GameplayCueTag in order to be asset registry searchable */
