Fix issues with gameplay cue notifies twice in some cases when related to inhibited gameplay effects.
-Do not replicate GameplayCue events triggered from inhibited effects
-Clear pending GameplayCue events on post rep if the gameplay effect is inhibited.

diff --git a/Engine/Source/Runtime/GameplayAbilities/Private/GameplayEffect.cpp b/Engine/Source/Runtime/GameplayAbilities/Private/GameplayEffect.cpp
--- a/Engine/Source/Runtime/GameplayAbilities/Private/GameplayEffect.cpp
+++ b/Engine/Source/Runtime/GameplayAbilities/Private/GameplayEffect.cpp
@@ -2346,11 +2346,8 @@ void FActiveGameplayEffectsContainer::AddActiveGameplayEffectGrantedTagsAndModif
 
 		if (bInvokeGameplayCueEvents)
 		{
-			// TODO: Optimize this so only one batched RPC is called
-			for (auto It = Cue.GameplayCueTags.CreateConstIterator(); It; ++It)
-			{
-				Owner->AddGameplayCue(*It);
-			}
+			Owner->InvokeGameplayCueEvent(Effect.Spec, EGameplayCueEvent::OnActive);
+			Owner->InvokeGameplayCueEvent(Effect.Spec, EGameplayCueEvent::WhileActive);
 		}
 	}
 }
@@ -2699,14 +2696,16 @@ bool FActiveGameplayEffectsContainer::NetDeltaSerialize(FNetDeltaSerializeInfo&
 				if (Effect.bPendingRepOnActiveGC)
 				{
 					Owner->InvokeGameplayCueEvent(Effect.Spec, EGameplayCueEvent::OnActive);
-					Effect.bPendingRepOnActiveGC = false;
+					
 				}
 				if (Effect.bPendingRepWhileActiveGC)
 				{
 					Owner->InvokeGameplayCueEvent(Effect.Spec, EGameplayCueEvent::WhileActive);
-					Effect.bPendingRepWhileActiveGC = false;
 				}
 			}
+
+			Effect.bPendingRepOnActiveGC = false;
+			Effect.bPendingRepWhileActiveGC = false;
 		}
 	}
 
