fix infinite loop if a gameplay cue tag is removed from a gameplay cue notify

diff --git a/Engine/Source/Runtime/GameplayAbilities/Private/GameplayCueSet.cpp b/Engine/Source/Runtime/GameplayAbilities/Private/GameplayCueSet.cpp
--- a/Engine/Source/Runtime/GameplayAbilities/Private/GameplayCueSet.cpp
+++ b/Engine/Source/Runtime/GameplayAbilities/Private/GameplayCueSet.cpp
@@ -236,7 +236,7 @@ void UGameplayCueSet::BuildAccelerationMap_Internal()
 	for (FGameplayCueNotifyData& Data : GameplayCueData)
 	{
 		FGameplayTag Parent = Data.GameplayCueTag.RequestDirectParent();
-		while (Parent != BaseGameplayCueTag())
+		while (Parent != BaseGameplayCueTag() && Parent.IsValid())
 		{
 			int32* idxPtr = GameplayCueDataMap.Find(Parent);
 			if (idxPtr)
