adding in GetGameplayEffectDefForHandle() to the AbilitySystemComp to get the UGameplayEffect* Def for the handle that is passed in

diff --git a/Engine/Source/Runtime/GameplayAbilities/Private/AbilitySystemComponent.cpp b/Engine/Source/Runtime/GameplayAbilities/Private/AbilitySystemComponent.cpp
--- a/Engine/Source/Runtime/GameplayAbilities/Private/AbilitySystemComponent.cpp
+++ b/Engine/Source/Runtime/GameplayAbilities/Private/AbilitySystemComponent.cpp
@@ -827,6 +827,17 @@ void UAbilitySystemComponent::CheckDurationExpired(FActiveGameplayEffectHandle H
 	ActiveGameplayEffects.CheckDuration(Handle);
 }
 
+const UGameplayEffect* UAbilitySystemComponent::GetGameplayEffectDefForHandle(FActiveGameplayEffectHandle Handle)
+{
+	FActiveGameplayEffect* ActiveGE = ActiveGameplayEffects.GetActiveGameplayEffect(Handle);
+	if (ActiveGE)
+	{
+		return ActiveGE->Spec.Def;
+	}
+
+	return nullptr;
+}
+
 bool UAbilitySystemComponent::RemoveActiveGameplayEffect(FActiveGameplayEffectHandle Handle, int32 StacksToRemove)
 {
 	return ActiveGameplayEffects.RemoveActiveGameplayEffect(Handle, StacksToRemove);
diff --git a/Engine/Source/Runtime/GameplayAbilities/Public/AbilitySystemComponent.h b/Engine/Source/Runtime/GameplayAbilities/Public/AbilitySystemComponent.h
--- a/Engine/Source/Runtime/GameplayAbilities/Public/AbilitySystemComponent.h
+++ b/Engine/Source/Runtime/GameplayAbilities/Public/AbilitySystemComponent.h
@@ -258,6 +258,9 @@ class GAMEPLAYABILITIES_API UAbilitySystemComponent : public UActorComponent, pu
 	UFUNCTION(BlueprintCallable, Category = GameplayEffects, meta=(DisplayName = "ApplyGameplayEffectSpecToSelf"))
 	FActiveGameplayEffectHandle BP_ApplyGameplayEffectSpecToSelf(UPARAM(ref) FGameplayEffectSpecHandle& SpecHandle);
 	
+	/** Gets the FActiveGameplayEffect based on the passed in Handle */
+	const UGameplayEffect* GetGameplayEffectDefForHandle(FActiveGameplayEffectHandle Handle);
+
 	/** Removes GameplayEffect by Handle. StacksToRemove=-1 will remove all stacks. */
 	UFUNCTION(BlueprintCallable, BlueprintAuthorityOnly, Category = GameplayEffects)
 	bool RemoveActiveGameplayEffect(FActiveGameplayEffectHandle Handle, int32 StacksToRemove=-1);
