fix for predictive gameplay cue removals not removing the gameplay tag

diff --git a/Engine/Source/Runtime/GameplayAbilities/Private/AbilitySystemComponent.cpp b/Engine/Source/Runtime/GameplayAbilities/Private/AbilitySystemComponent.cpp
--- a/Engine/Source/Runtime/GameplayAbilities/Private/AbilitySystemComponent.cpp
+++ b/Engine/Source/Runtime/GameplayAbilities/Private/AbilitySystemComponent.cpp
@@ -1034,7 +1034,6 @@ void UAbilitySystemComponent::RemoveGameplayCue(const FGameplayTag GameplayCueTa
 	else if (ScopedPredictionKey.IsLocalClientKey())
 	{
 		ActiveGameplayCues.PredictiveRemove(GameplayCueTag);
-		InvokeGameplayCueEvent(GameplayCueTag, EGameplayCueEvent::Removed);
 	}
 }
 
diff --git a/Engine/Source/Runtime/GameplayAbilities/Private/GameplayCueInterface.cpp b/Engine/Source/Runtime/GameplayAbilities/Private/GameplayCueInterface.cpp
--- a/Engine/Source/Runtime/GameplayAbilities/Private/GameplayCueInterface.cpp
+++ b/Engine/Source/Runtime/GameplayAbilities/Private/GameplayCueInterface.cpp
@@ -172,7 +172,11 @@ void FActiveGameplayCueContainer::PredictiveRemove(const FGameplayTag& Tag)
 		FActiveGameplayCue& Cue = GameplayCues[idx];
 		if (Cue.GameplayCueTag == Tag)
 		{			
+			// Predictive remove: mark the cue as predictive remove, invoke remove event, update tag map.
+			// DONT remove from the replicated array.
 			Cue.bPredictivelyRemoved = true;
+			Owner->InvokeGameplayCueEvent(Tag, EGameplayCueEvent::Removed);
+			Owner->UpdateTagMap(Tag, -1);
 			return;
 		}
 	}
