Fix for nonunity build

diff --git a/Engine/Source/Runtime/GameplayAbilities/Private/GameplayAbilitiesModule.cpp b/Engine/Source/Runtime/GameplayAbilities/Private/GameplayAbilitiesModule.cpp
--- a/Engine/Source/Runtime/GameplayAbilities/Private/GameplayAbilitiesModule.cpp
+++ b/Engine/Source/Runtime/GameplayAbilities/Private/GameplayAbilitiesModule.cpp
@@ -2,7 +2,7 @@
 
 #include "AbilitySystemPrivatePCH.h"
 #include "AbilitySystemGlobals.h"
-
+#include "AbilitySystemComponent.h"
 
 class FGameplayAbilitiesModule : public IGameplayAbilitiesModule
 {
