Intergration from UE4-Orion to UE4

Adding a new function to the FAggregator that takes in everything it needs to update the mods and only call dirty once the update in completed. This will prevent the UAttributeSets stats to get an update inbetween the remove and re-add of the mods. This fixes OR-3272.

diff --git a/Engine/Source/Runtime/GameplayAbilities/Private/GameplayEffect.cpp b/Engine/Source/Runtime/GameplayAbilities/Private/GameplayEffect.cpp
--- a/Engine/Source/Runtime/GameplayAbilities/Private/GameplayEffect.cpp
+++ b/Engine/Source/Runtime/GameplayAbilities/Private/GameplayEffect.cpp
@@ -1667,20 +1667,8 @@ void FActiveGameplayEffectsContainer::UpdateAggregatorModMagnitudes(const TSet<F
 		FAggregator* Aggregator = FindOrCreateAttributeAggregator(Attribute).Get();
 		check(Aggregator);
 
-		// Remove ALL of our mods (this is all we can do! We don't keep track of aggregator mods <=> spec mod mappings)
-		// Todo: it would be nice to have a better way of updating this. Some ideas: instead of remove, call an 'invalidate' to NAN out all floats for a given handle, then call an 'update' to set real values.
-		Aggregator->RemoveAggregatorMod(ActiveEffect.Handle);
-
-		// Now re-add ALL of our mods
-		for(int32 ModIdx = 0; ModIdx < Spec.Modifiers.Num(); ++ModIdx)
-		{
-			const FGameplayModifierInfo& ModDef = Spec.Def->Modifiers[ModIdx];
-
-			if (ModDef.Attribute == Attribute)
-			{
-				Aggregator->AddAggregatorMod(Spec.GetModifierMagnitude(ModIdx, true), ModDef.ModifierOp, &ModDef.SourceTags, &ModDef.TargetTags, ActiveEffect.PredictionKey.WasLocallyGenerated(), ActiveEffect.Handle);
-			}
-		}
+		// Update the aggregator Mods.
+		Aggregator->UpdateAggregatorMod(ActiveEffect.Handle, Attribute, Spec, ActiveEffect.PredictionKey.WasLocallyGenerated(), ActiveEffect.Handle);
 	}
 }
 
diff --git a/Engine/Source/Runtime/GameplayAbilities/Private/GameplayEffectAggregator.cpp b/Engine/Source/Runtime/GameplayAbilities/Private/GameplayEffectAggregator.cpp
--- a/Engine/Source/Runtime/GameplayAbilities/Private/GameplayEffectAggregator.cpp
+++ b/Engine/Source/Runtime/GameplayAbilities/Private/GameplayEffectAggregator.cpp
@@ -188,6 +188,35 @@ void FAggregator::AddAggregatorMod(float EvaluatedMagnitude, TEnumAsByte<EGamepl
 }
 
 void FAggregator::RemoveAggregatorMod(FActiveGameplayEffectHandle ActiveHandle)
+{
+	InternalRemoveAggregatorMod(ActiveHandle);
+
+	// mark it as dirty so that all the stats get updated
+	BroadcastOnDirty();
+}
+
+void FAggregator::UpdateAggregatorMod(FActiveGameplayEffectHandle ActiveHandle, const FGameplayAttribute& Attribute, const FGameplayEffectSpec& Spec, bool bWasLocalyGenerated, FActiveGameplayEffectHandle InHandle)
+{
+	// remove the mods but dont mark it as dirty until we re-add the aggregators, we are doing this so the UAttributeSets stats only know about the delta change.
+	InternalRemoveAggregatorMod(ActiveHandle);
+
+	// Now re-add ALL of our mods
+	for (int32 ModIdx = 0; ModIdx < Spec.Modifiers.Num(); ++ModIdx)
+	{
+		const FGameplayModifierInfo& ModDef = Spec.Def->Modifiers[ModIdx];
+
+		if (ModDef.Attribute == Attribute)
+		{
+			AddAggregatorMod(Spec.GetModifierMagnitude(ModIdx, true), ModDef.ModifierOp, &ModDef.SourceTags, &ModDef.TargetTags, bWasLocalyGenerated, InHandle);
+		}
+	}
+
+	// mark it as dirty so that all the stats get updated
+	BroadcastOnDirty();
+}
+
+// moved the remove into a common function since now we will need to remove w/o marking it as dirty so we can support updating the aggregators 
+void FAggregator::InternalRemoveAggregatorMod(FActiveGameplayEffectHandle ActiveHandle)
 {
 	if (ActiveHandle.IsValid())
 	{
@@ -196,8 +225,6 @@ void FAggregator::RemoveAggregatorMod(FActiveGameplayEffectHandle ActiveHandle)
 		RemoveModsWithActiveHandle(Mods[EGameplayModOp::Division], ActiveHandle);
 		RemoveModsWithActiveHandle(Mods[EGameplayModOp::Override], ActiveHandle);
 	}
-
-	BroadcastOnDirty();
 }
 
 void FAggregator::AddModsFrom(const FAggregator& SourceAggregator)
diff --git a/Engine/Source/Runtime/GameplayAbilities/Public/GameplayEffectAggregator.h b/Engine/Source/Runtime/GameplayAbilities/Public/GameplayEffectAggregator.h
--- a/Engine/Source/Runtime/GameplayAbilities/Public/GameplayEffectAggregator.h
+++ b/Engine/Source/Runtime/GameplayAbilities/Public/GameplayEffectAggregator.h
@@ -49,8 +49,13 @@ struct GAMEPLAYABILITIES_API FAggregator : public TSharedFromThis<FAggregator>
 	static float StaticExecModOnBaseValue(float BaseValue, TEnumAsByte<EGameplayModOp::Type> ModifierOp, float EvaluatedMagnitude);
 
 	void AddAggregatorMod(float EvaluatedData, TEnumAsByte<EGameplayModOp::Type> ModifierOp, const FGameplayTagRequirements*	SourceTagReqs, const FGameplayTagRequirements* TargetTagReqs, bool IsPredicted, FActiveGameplayEffectHandle ActiveHandle = FActiveGameplayEffectHandle() );
+
+	/** Removes all mods for the passed in handle and marks this as dirty to recalculate the aggregator*/
 	void RemoveAggregatorMod(FActiveGameplayEffectHandle ActiveHandle);
 
+	/** Updates the aggregators for the past in handle, this will handle it so the UAttributeSets stats only get one update for the delta change */
+	void UpdateAggregatorMod(FActiveGameplayEffectHandle ActiveHandle, const FGameplayAttribute& Attribute, const FGameplayEffectSpec& Spec, bool bWasLocalyGenerated, FActiveGameplayEffectHandle InHandle);
+
 	/** Evaluates the Aggregator with the internal base value and given parameters */
 	float Evaluate(const FAggregatorEvaluateParameters& Parameters) const;
 
@@ -83,6 +88,9 @@ private:
 	float SumMods(const TArray<FAggregatorMod> &Mods, float Bias, const FAggregatorEvaluateParameters& Parameters) const;
 	void RemoveModsWithActiveHandle(TArray<FAggregatorMod>& Mods, FActiveGameplayEffectHandle ActiveHandle);
 
+	/** Helper function that will remove all mods that are based on the passed in active handle */
+	void InternalRemoveAggregatorMod(FActiveGameplayEffectHandle ActiveHandle);
+
 	float	BaseValue;
 	TArray<FAggregatorMod>	Mods[EGameplayModOp::Max];
 
