Update AGameplayAbilityTargetActor_Radius to use a collision object query instead of a trace channel.

--------
Integrated using branch UE4-To-UE4-Orion (reversed) of change#2544082 by [REDACTED_PERSON] on 2015/05/08 20:08:53.

diff --git a/Engine/Source/Runtime/GameplayAbilities/Private/Abilities/GameplayAbilityTargetActor_Radius.cpp b/Engine/Source/Runtime/GameplayAbilities/Private/Abilities/GameplayAbilityTargetActor_Radius.cpp
--- a/Engine/Source/Runtime/GameplayAbilities/Private/Abilities/GameplayAbilityTargetActor_Radius.cpp
+++ b/Engine/Source/Runtime/GameplayAbilities/Private/Abilities/GameplayAbilityTargetActor_Radius.cpp
@@ -60,7 +60,7 @@ TArray<TWeakObjectPtr<AActor> >	AGameplayAbilityTargetActor_Radius::PerformOverl
 
 	TArray<FOverlapResult> Overlaps;
 
-	SourceActor->GetWorld()->OverlapMultiByChannel(Overlaps, Origin, FQuat::Identity, ECC_Pawn, FCollisionShape::MakeSphere(Radius), Params);
+	SourceActor->GetWorld()->OverlapMultiByObjectType(Overlaps, Origin, FQuat::Identity, FCollisionObjectQueryParams(ECC_Pawn), FCollisionShape::MakeSphere(Radius), Params);
 
 	TArray<TWeakObjectPtr<AActor>>	HitActors;
 
