Merge of 2555458 to main

Adds virtual OnTagUpdated function to ability system component.

diff --git a/Engine/Source/Runtime/GameplayAbilities/Private/AbilitySystemComponent.cpp b/Engine/Source/Runtime/GameplayAbilities/Private/AbilitySystemComponent.cpp
--- a/Engine/Source/Runtime/GameplayAbilities/Private/AbilitySystemComponent.cpp
+++ b/Engine/Source/Runtime/GameplayAbilities/Private/AbilitySystemComponent.cpp
@@ -496,12 +496,23 @@ void UAbilitySystemComponent::RemoveLooseGameplayTags(const FGameplayTagContaine
 
 void UAbilitySystemComponent::UpdateTagMap(const FGameplayTag& BaseTag, int32 CountDelta)
 {
+	const bool bHadTag = GameplayTagCountContainer.HasMatchingGameplayTag(BaseTag);
 	GameplayTagCountContainer.UpdateTagCount(BaseTag, CountDelta);
+	const bool bHasTag = GameplayTagCountContainer.HasMatchingGameplayTag(BaseTag);
+
+	if (bHadTag != bHasTag)
+	{
+		OnTagUpdated(BaseTag, bHasTag);
+	}
 }
 
 void UAbilitySystemComponent::UpdateTagMap(const FGameplayTagContainer& Container, int32 CountDelta)
 {
-	GameplayTagCountContainer.UpdateTagCount(Container, CountDelta);
+	for (auto TagIt = Container.CreateConstIterator(); TagIt; ++TagIt)
+	{
+		const FGameplayTag& Tag = *TagIt;
+		UpdateTagMap(Tag, CountDelta);
+	}
 }
 
 // ------------------------------------------------------------------------
diff --git a/Engine/Source/Runtime/GameplayAbilities/Public/AbilitySystemComponent.h b/Engine/Source/Runtime/GameplayAbilities/Public/AbilitySystemComponent.h
--- a/Engine/Source/Runtime/GameplayAbilities/Public/AbilitySystemComponent.h
+++ b/Engine/Source/Runtime/GameplayAbilities/Public/AbilitySystemComponent.h
@@ -1076,6 +1076,8 @@ protected:
 	void UpdateTagMap(const FGameplayTag& BaseTag, int32 CountDelta);
 	
 	void UpdateTagMap(const FGameplayTagContainer& Container, int32 CountDelta);
+
+	virtual void OnTagUpdated(const FGameplayTag& Tag, bool TagExists) {};
 	
 	// ---------------------------------------------
 
