More strict restrictions related to WeakObjectPtr exposed to blueprint. No arrays on weakobjptr. No lazy pointers.  UE-15033

diff --git a/Engine/Config/BaseEngine.ini b/Engine/Config/BaseEngine.ini
--- a/Engine/Config/BaseEngine.ini
+++ b/Engine/Config/BaseEngine.ini
@@ -1152,6 +1152,7 @@ bUseStreamingPause=false
 +K2FieldRedirects=(OldFieldName="KismetMathLibrary.InverseTransformPosition", NewFieldName="KismetMathLibrary.InverseTransformLocation")
 +K2ParamRedirects=(NodeName="/Script/BlueprintGraph.K2Node_CallFunction", OldParamName="KismetMathLibrary.InverseTransformPosition.Position", NewParamName="Location")
 +K2ParamRedirects=(NodeName="/Script/BlueprintGraph.K2Node_CallFunction", OldParamName="KismetMathLibrary.InverseTransformLocation.Position", NewParamName="Location")
++K2FieldRedirects=(OldFieldName="PrimitiveComponent.GetMoveIgnoreActors", NewFieldName="PrimitiveComponent.CopyArrayOfMoveIgnoreActors")
 
 +EnumRedirects=(EnumName="EMaxConcurrentResolutionRule",OldEnumEntry="EMaxConcurrentResolutionRule::StopFarthest",NewEnumEntry="EMaxConcurrentResolutionRule::StopFarthestThenPreventNew")
 
diff --git a/Engine/Source/Programs/UnrealHeaderTool/Private/HeaderParser.cpp b/Engine/Source/Programs/UnrealHeaderTool/Private/HeaderParser.cpp
--- a/Engine/Source/Programs/UnrealHeaderTool/Private/HeaderParser.cpp
+++ b/Engine/Source/Programs/UnrealHeaderTool/Private/HeaderParser.cpp
@@ -717,8 +717,8 @@ namespace
 		}
 		if (auto ArrayProperty = Cast<const UArrayProperty>(Property))
 		{
-			// Inner Property can be handled as a member variable
-			return IsPropertySupportedByBlueprint(ArrayProperty->Inner, true);
+			// Script VM doesn't support array of weak ptrs.
+			return IsPropertySupportedByBlueprint(ArrayProperty->Inner, false);
 		}
 
 		const bool bSupportedType = Property->IsA<UInterfaceProperty>()
@@ -733,10 +733,9 @@ namespace
 			|| Property->IsA<UBoolProperty>()
 			|| Property->IsA<UStrProperty>()
 			|| Property->IsA<UTextProperty>()
-			|| Property->IsA<UMulticastDelegateProperty>()
 			|| Property->IsA<UDelegateProperty>();
 
-		const bool bIsSupportedMemberVariable = Property->IsA<UObjectPropertyBase>();
+		const bool bIsSupportedMemberVariable = Property->IsA<UWeakObjectProperty>() || Property->IsA<UMulticastDelegateProperty>();
 
 		return bSupportedType || (bIsSupportedMemberVariable && bMemberVariable);
 	}
diff --git a/Engine/Source/Runtime/Engine/Classes/Components/PrimitiveComponent.h b/Engine/Source/Runtime/Engine/Classes/Components/PrimitiveComponent.h
--- a/Engine/Source/Runtime/Engine/Classes/Components/PrimitiveComponent.h
+++ b/Engine/Source/Runtime/Engine/Classes/Components/PrimitiveComponent.h
@@ -489,7 +489,12 @@ public:
 	/**
 	 * Returns the list of actors we currently ignore when moving.
 	 */
-	UFUNCTION(BlueprintCallable, Category = "Collision")
+	UFUNCTION(BlueprintCallable, meta=(DisplayName="GetMoveIgnoreActors"), Category = "Collision")
+	TArray<AActor*> CopyArrayOfMoveIgnoreActors();
+
+	/**
+	 * Returns the list of actors (as WeakObjectPtr) we currently ignore when moving.
+	 */
 	TArray<TWeakObjectPtr<AActor> > & GetMoveIgnoreActors();
 
 	/**
diff --git a/Engine/Source/Runtime/Engine/Private/Components/PrimitiveComponent.cpp b/Engine/Source/Runtime/Engine/Private/Components/PrimitiveComponent.cpp
--- a/Engine/Source/Runtime/Engine/Private/Components/PrimitiveComponent.cpp
+++ b/Engine/Source/Runtime/Engine/Private/Components/PrimitiveComponent.cpp
@@ -2223,6 +2223,26 @@ void UPrimitiveComponent::IgnoreActorWhenMoving(AActor* Actor, bool bShouldIgnor
 	}
 }
 
+TArray<AActor*> UPrimitiveComponent::CopyArrayOfMoveIgnoreActors()
+{
+	TArray<AActor*> TempMoveIgnoreActors;
+	for (int32 Idx = 0; Idx < MoveIgnoreActors.Num();)
+	{
+		AActor* Actor = MoveIgnoreActors[Idx].Get();
+		if (Actor)
+		{
+			TempMoveIgnoreActors.Add(Actor);
+			Idx++;
+		}
+		else
+		{
+			MoveIgnoreActors.RemoveAtSwap(Idx, 1, false);
+		}
+	}
+	MoveIgnoreActors.Shrink();
+	return TempMoveIgnoreActors;
+}
+
 TArray<TWeakObjectPtr<AActor> > & UPrimitiveComponent::GetMoveIgnoreActors()
 {
 	// Clean up stale references
diff --git a/Engine/Source/Runtime/GameplayAbilities/Private/AbilitySystemBlueprintLibrary.cpp b/Engine/Source/Runtime/GameplayAbilities/Private/AbilitySystemBlueprintLibrary.cpp
--- a/Engine/Source/Runtime/GameplayAbilities/Private/AbilitySystemBlueprintLibrary.cpp
+++ b/Engine/Source/Runtime/GameplayAbilities/Private/AbilitySystemBlueprintLibrary.cpp
@@ -80,7 +80,7 @@ FGameplayAbilityTargetDataHandle UAbilitySystemBlueprintLibrary::AbilityTargetDa
 	FGameplayAbilityTargetDataHandle		Handle(NewData);
 	return Handle;
 }
-FGameplayAbilityTargetDataHandle UAbilitySystemBlueprintLibrary::AbilityTargetDataFromActorArray(TArray<TWeakObjectPtr<AActor>> ActorArray, bool OneTargetPerHandle)
+FGameplayAbilityTargetDataHandle UAbilitySystemBlueprintLibrary::AbilityTargetDataFromActorArray(TArray<AActor*> ActorArray, bool OneTargetPerHandle)
 {
 	// Construct TargetData
 	if (OneTargetPerHandle)
@@ -88,9 +88,9 @@ FGameplayAbilityTargetDataHandle UAbilitySystemBlueprintLibrary::AbilityTargetDa
 		FGameplayAbilityTargetDataHandle		Handle;
 		for (int32 i = 0; i < ActorArray.Num(); ++i)
 		{
-			if (ActorArray[i].IsValid())
+			if (IsValid(ActorArray[i]))
 			{
-				FGameplayAbilityTargetDataHandle TempHandle = AbilityTargetDataFromActor(ActorArray[i].Get());
+				FGameplayAbilityTargetDataHandle TempHandle = AbilityTargetDataFromActor(ActorArray[i]);
 				Handle.Append(&TempHandle);
 			}
 		}
@@ -99,7 +99,11 @@ FGameplayAbilityTargetDataHandle UAbilitySystemBlueprintLibrary::AbilityTargetDa
 	else
 	{
 		FGameplayAbilityTargetData_ActorArray*	NewData = new FGameplayAbilityTargetData_ActorArray();
-		NewData->TargetActorArray = ActorArray;
+		NewData->TargetActorArray.Reset();
+		for (auto Actor : ActorArray)
+		{
+			NewData->TargetActorArray.Add(Actor);
+		}
 		FGameplayAbilityTargetDataHandle		Handle(NewData);
 		return Handle;
 	}
diff --git a/Engine/Source/Runtime/GameplayAbilities/Public/Abilities/GameplayAbilityTargetTypes.h b/Engine/Source/Runtime/GameplayAbilities/Public/Abilities/GameplayAbilityTargetTypes.h
--- a/Engine/Source/Runtime/GameplayAbilities/Public/Abilities/GameplayAbilityTargetTypes.h
+++ b/Engine/Source/Runtime/GameplayAbilities/Public/Abilities/GameplayAbilityTargetTypes.h
@@ -423,7 +423,7 @@ struct GAMEPLAYABILITIES_API FGameplayAbilityTargetData_ActorArray : public FGam
 	FGameplayAbilityTargetingLocationInfo SourceLocation;
 
 	/** Rather than targeting a single point, this type of targeting selects multiple actors. */
-	UPROPERTY(BlueprintReadWrite, EditAnywhere, Category = Targeting)
+	UPROPERTY(EditAnywhere, Category = Targeting)
 	TArray<TWeakObjectPtr<AActor> > TargetActorArray;
 
 	virtual TArray<TWeakObjectPtr<AActor> >	GetActors() const override
diff --git a/Engine/Source/Runtime/GameplayAbilities/Public/AbilitySystemBlueprintLibrary.h b/Engine/Source/Runtime/GameplayAbilities/Public/AbilitySystemBlueprintLibrary.h
--- a/Engine/Source/Runtime/GameplayAbilities/Public/AbilitySystemBlueprintLibrary.h
+++ b/Engine/Source/Runtime/GameplayAbilities/Public/AbilitySystemBlueprintLibrary.h
@@ -55,7 +55,7 @@ class GAMEPLAYABILITIES_API UAbilitySystemBlueprintLibrary : public UBlueprintFu
 	static FGameplayAbilityTargetDataHandle	AbilityTargetDataFromActor(AActor* Actor);
 
 	UFUNCTION(BlueprintPure, Category = "Ability|TargetData")
-	static FGameplayAbilityTargetDataHandle	AbilityTargetDataFromActorArray(TArray<TWeakObjectPtr<AActor>> ActorArray, bool OneTargetPerHandle);
+	static FGameplayAbilityTargetDataHandle	AbilityTargetDataFromActorArray(TArray<AActor*> ActorArray, bool OneTargetPerHandle);
 
 	/** Create a new target data handle with filtration performed on the data */
 	UFUNCTION(BlueprintPure, Category = "Ability|TargetData")
