add functions for getting IsLocallyControlledPlayer (vs just IsLocallyControlled which returns true for AI on server)

--------
Integrated using branch UE4-To-UE4-Orion (reversed) of change#2561568 by [REDACTED_PERSON] on 2015/05/21 17:29:27.

diff --git a/Engine/Source/Runtime/GameplayAbilities/Private/AbilitySystemBlueprintLibrary.cpp b/Engine/Source/Runtime/GameplayAbilities/Private/AbilitySystemBlueprintLibrary.cpp
--- a/Engine/Source/Runtime/GameplayAbilities/Private/AbilitySystemBlueprintLibrary.cpp
+++ b/Engine/Source/Runtime/GameplayAbilities/Private/AbilitySystemBlueprintLibrary.cpp
@@ -398,6 +398,11 @@ bool UAbilitySystemBlueprintLibrary::IsInstigatorLocallyControlled(FGameplayCueP
 	return Parameters.EffectContext.IsLocallyControlled();
 }
 
+bool UAbilitySystemBlueprintLibrary::IsInstigatorLocallyControlledPlayer(FGameplayCueParameters Parameters)
+{
+	return Parameters.EffectContext.IsLocallyControlledPlayer();
+}
+
 int32 UAbilitySystemBlueprintLibrary::GetActorCount(FGameplayCueParameters Parameters)
 {
 	return Parameters.EffectContext.GetActors().Num();
diff --git a/Engine/Source/Runtime/GameplayAbilities/Private/GameplayAbilityTypes.cpp b/Engine/Source/Runtime/GameplayAbilities/Private/GameplayAbilityTypes.cpp
--- a/Engine/Source/Runtime/GameplayAbilities/Private/GameplayAbilityTypes.cpp
+++ b/Engine/Source/Runtime/GameplayAbilities/Private/GameplayAbilityTypes.cpp
@@ -84,6 +84,16 @@ bool FGameplayAbilityActorInfo::IsLocallyControlled() const
 	return false;
 }
 
+bool FGameplayAbilityActorInfo::IsLocallyControlledPlayer() const
+{
+	if (PlayerController.IsValid())
+	{
+		return PlayerController->IsLocalController();
+	}
+
+	return false;
+}
+
 bool FGameplayAbilityActorInfo::IsNetAuthority() const
 {
 	// Make sure this works on pending kill actors
diff --git a/Engine/Source/Runtime/GameplayAbilities/Private/GameplayEffectTypes.cpp b/Engine/Source/Runtime/GameplayAbilities/Private/GameplayEffectTypes.cpp
--- a/Engine/Source/Runtime/GameplayAbilities/Private/GameplayEffectTypes.cpp
+++ b/Engine/Source/Runtime/GameplayAbilities/Private/GameplayEffectTypes.cpp
@@ -188,6 +188,20 @@ bool FGameplayEffectContext::IsLocallyControlled() const
 	return false;
 }
 
+bool FGameplayEffectContext::IsLocallyControlledPlayer() const
+{
+	APawn* Pawn = Cast<APawn>(Instigator.Get());
+	if (!Pawn)
+	{
+		Pawn = Cast<APawn>(EffectCauser.Get());
+	}
+	if (Pawn && Pawn->Controller)
+	{
+		return Pawn->Controller->IsLocalPlayerController();
+	}
+	return false;
+}
+
 void FGameplayEffectContext::AddOrigin(FVector InOrigin)
 {
 	bHasWorldOrigin = true;
diff --git a/Engine/Source/Runtime/GameplayAbilities/Public/Abilities/GameplayAbilityTypes.h b/Engine/Source/Runtime/GameplayAbilities/Public/Abilities/GameplayAbilityTypes.h
--- a/Engine/Source/Runtime/GameplayAbilities/Public/Abilities/GameplayAbilityTypes.h
+++ b/Engine/Source/Runtime/GameplayAbilities/Public/Abilities/GameplayAbilityTypes.h
@@ -146,6 +146,8 @@ struct GAMEPLAYABILITIES_API FGameplayAbilityActorInfo
 	/** Returns true if this actor is locally controlled. Only true for players on the client that owns them */
 	bool IsLocallyControlled() const;
 
+	bool IsLocallyControlledPlayer() const;
+
 	/** Returns true if the owning actor has net authority */
 	bool IsNetAuthority() const;
 
diff --git a/Engine/Source/Runtime/GameplayAbilities/Public/AbilitySystemBlueprintLibrary.h b/Engine/Source/Runtime/GameplayAbilities/Public/AbilitySystemBlueprintLibrary.h
--- a/Engine/Source/Runtime/GameplayAbilities/Public/AbilitySystemBlueprintLibrary.h
+++ b/Engine/Source/Runtime/GameplayAbilities/Public/AbilitySystemBlueprintLibrary.h
@@ -141,6 +141,9 @@ class GAMEPLAYABILITIES_API UAbilitySystemBlueprintLibrary : public UBlueprintFu
 	UFUNCTION(BlueprintPure, Category="Ability|GameplayCue")
 	static bool IsInstigatorLocallyControlled(FGameplayCueParameters Parameters);
 
+	UFUNCTION(BlueprintPure, Category="Ability|GameplayCue")
+	static bool IsInstigatorLocallyControlledPlayer(FGameplayCueParameters Parameters);
+
 	UFUNCTION(BlueprintPure, Category = "Ability|GameplayCue")
 	static int32 GetActorCount(FGameplayCueParameters Parameters);
 
diff --git a/Engine/Source/Runtime/GameplayAbilities/Public/GameplayEffectTypes.h b/Engine/Source/Runtime/GameplayAbilities/Public/GameplayEffectTypes.h
--- a/Engine/Source/Runtime/GameplayAbilities/Public/GameplayEffectTypes.h
+++ b/Engine/Source/Runtime/GameplayAbilities/Public/GameplayEffectTypes.h
@@ -379,6 +379,8 @@ struct GAMEPLAYABILITIES_API FGameplayEffectContext
 
 	virtual bool IsLocallyControlled() const;
 
+	virtual bool IsLocallyControlledPlayer() const;
+
 	virtual bool NetSerialize(FArchive& Ar, class UPackageMap* Map, bool& bOutSuccess);
 
 protected:
@@ -564,6 +566,15 @@ struct FGameplayEffectContextHandle
 		return false;
 	}
 
+	bool IsLocallyControlledPlayer() const
+	{
+		if (IsValid())
+		{
+			return Data->IsLocallyControlledPlayer();
+		}
+		return false;
+	}
+
 	void AddActors(TArray<TWeakObjectPtr<AActor>> InActors, bool bReset = false)
 	{
 		if (IsValid())
