Disable gameplay cues on dedicated server. This can be toggled with the cvar AbilitySystem.GameplayCue.RunOnDedicatedServer.

--------
Integrated using branch UE4-To-UE4-Orion (reversed) of change#2572809 by [REDACTED_PERSON] on 2015/06/01 14:35:26.

diff --git a/Engine/Source/Runtime/GameplayAbilities/Private/GameplayCueManager.cpp b/Engine/Source/Runtime/GameplayAbilities/Private/GameplayCueManager.cpp
--- a/Engine/Source/Runtime/GameplayAbilities/Private/GameplayCueManager.cpp
+++ b/Engine/Source/Runtime/GameplayAbilities/Private/GameplayCueManager.cpp
@@ -22,6 +22,9 @@ static FAutoConsoleVariableRef CVarDisableGameplayCues(TEXT("AbilitySystem.Disab
 float DisplayGameplayCueDuration = 5.f;
 static FAutoConsoleVariableRef CVarDurationeGameplayCues(TEXT("AbilitySystem.GameplayCue.DisplayDuration"),	DisplayGameplayCueDuration, TEXT("Disables all GameplayCue events in the world."), ECVF_Default );
 
+int32 GameplayCueRunOnDedicatedServer = 0;
+static FAutoConsoleVariableRef CVarDedicatedServerGameplayCues(TEXT("AbilitySystem.GameplayCue.RunOnDedicatedServer"), GameplayCueRunOnDedicatedServer, TEXT("Run gameplay cue events on dedicated server"), ECVF_Default );
+
 UGameplayCueManager::UGameplayCueManager(const FObjectInitializer& PCIP)
 : Super(PCIP)
 {
@@ -33,9 +36,25 @@ UGameplayCueManager::UGameplayCueManager(const FObjectInitializer& PCIP)
 	GlobalCueSet = NewObject<UGameplayCueSet>(this, TEXT("GlobalCueSet"));
 }
 
+bool IsDedicatedServerForGameplayCue()
+{
+#if WITH_EDITOR
+	// This will handle dedicated server PIE case properly
+	return GEngine->ShouldAbsorbCosmeticOnlyEvent();
+#else
+	// When in standalone non editor, this is the fastest way to check
+	return IsRunningDedicatedServer();
+#endif
+}
+
 
 void UGameplayCueManager::HandleGameplayCues(AActor* TargetActor, const FGameplayTagContainer& GameplayCueTags, EGameplayCueEvent::Type EventType, FGameplayCueParameters Parameters)
 {
+	if (GameplayCueRunOnDedicatedServer == 0 && IsDedicatedServerForGameplayCue())
+	{
+		return;
+	}
+
 	for (auto It = GameplayCueTags.CreateConstIterator(); It; ++It)
 	{
 		HandleGameplayCue(TargetActor, *It, EventType, Parameters);
@@ -49,6 +68,11 @@ void UGameplayCueManager::HandleGameplayCue(AActor* TargetActor, FGameplayTag Ga
 		return;
 	}
 
+	if (GameplayCueRunOnDedicatedServer == 0 && IsDedicatedServerForGameplayCue())
+	{
+		return;
+	}
+
 	if (TargetActor == nullptr)
 	{
 		ABILITY_LOG(Warning, TEXT("UGameplayCueManager::HandleGameplayCue called on null TargetActor. GameplayCueTag: %s."), *GameplayCueTag.ToString());
