support for stateful predictive adding of gameplay cue tags. Previously only OnActive/WhileActive events would be predicted but the stateful (tag) aspect of gameplaycues was not predicted.

--------
Integrated using branch UE4-To-UE4-Orion (reversed) of change#2583399 by [REDACTED_PERSON] on 2015/06/10 17:26:35.

diff --git a/Engine/Source/Runtime/GameplayAbilities/Private/AbilitySystemComponent.cpp b/Engine/Source/Runtime/GameplayAbilities/Private/AbilitySystemComponent.cpp
--- a/Engine/Source/Runtime/GameplayAbilities/Private/AbilitySystemComponent.cpp
+++ b/Engine/Source/Runtime/GameplayAbilities/Private/AbilitySystemComponent.cpp
@@ -1020,6 +1020,8 @@ void UAbilitySystemComponent::AddGameplayCue(const FGameplayTag GameplayCueTag,
 	}
 	else if (ScopedPredictionKey.IsLocalClientKey())
 	{
+		ActiveGameplayCues.PredictiveAdd(GameplayCueTag, ScopedPredictionKey);
+
 		// Allow for predictive gameplaycue events? Needs more thought
 		InvokeGameplayCueEvent(GameplayCueTag, EGameplayCueEvent::OnActive, EffectContext);
 		InvokeGameplayCueEvent(GameplayCueTag, EGameplayCueEvent::WhileActive);
diff --git a/Engine/Source/Runtime/GameplayAbilities/Private/GameplayCueInterface.cpp b/Engine/Source/Runtime/GameplayAbilities/Private/GameplayCueInterface.cpp
--- a/Engine/Source/Runtime/GameplayAbilities/Private/GameplayCueInterface.cpp
+++ b/Engine/Source/Runtime/GameplayAbilities/Private/GameplayCueInterface.cpp
@@ -132,11 +132,12 @@ void FActiveGameplayCue::PreReplicatedRemove(const struct FActiveGameplayCueCont
 
 void FActiveGameplayCue::PostReplicatedAdd(const struct FActiveGameplayCueContainer &InArray)
 {
+	InArray.Owner->UpdateTagMap(GameplayCueTag, 1);
+
 	if (PredictionKey.IsLocalClientKey() == false)
 	{
 		// If predicted ignore the add/remove
 		InArray.Owner->InvokeGameplayCueEvent(GameplayCueTag, EGameplayCueEvent::WhileActive);
-		InArray.Owner->UpdateTagMap(GameplayCueTag, 1);
 	}
 }
 
@@ -185,4 +186,10 @@ void FActiveGameplayCueContainer::PredictiveRemove(const FGameplayTag& Tag)
 			return;
 		}
 	}
+}
+
+void FActiveGameplayCueContainer::PredictiveAdd(const FGameplayTag& Tag, FPredictionKey& PredictionKey)
+{
+	Owner->UpdateTagMap(Tag, 1);	
+	PredictionKey.NewRejectOrCaughtUpDelegate(FPredictionKeyEvent::CreateUObject(Owner, &UAbilitySystemComponent::RemoveOneTagCount_NoReturn, Tag));
 }
\ No newline at end of file
diff --git a/Engine/Source/Runtime/GameplayAbilities/Public/AbilitySystemComponent.h b/Engine/Source/Runtime/GameplayAbilities/Public/AbilitySystemComponent.h
--- a/Engine/Source/Runtime/GameplayAbilities/Public/AbilitySystemComponent.h
+++ b/Engine/Source/Runtime/GameplayAbilities/Public/AbilitySystemComponent.h
@@ -310,6 +310,12 @@ class GAMEPLAYABILITIES_API UAbilitySystemComponent : public UGameplayTasksCompo
 		RemoveActiveGameplayEffect(Handle, StacksToRemove);
 	}
 
+	/** Needed for delegate callback for tag prediction */
+	void RemoveOneTagCount_NoReturn(FGameplayTag Tag)
+	{
+		UpdateTagMap(Tag, -1);
+	}
+
 	float GetGameplayEffectDuration(FActiveGameplayEffectHandle Handle) const;
 
 	float GetGameplayEffectDuration() const;
diff --git a/Engine/Source/Runtime/GameplayAbilities/Public/GameplayCueInterface.h b/Engine/Source/Runtime/GameplayAbilities/Public/GameplayCueInterface.h
--- a/Engine/Source/Runtime/GameplayAbilities/Public/GameplayCueInterface.h
+++ b/Engine/Source/Runtime/GameplayAbilities/Public/GameplayCueInterface.h
@@ -103,6 +103,8 @@ struct FActiveGameplayCueContainer : public FFastArraySerializer
 	/** Marks as predictively removed so that we dont invoke remove event twice due to onrep */
 	void PredictiveRemove(const FGameplayTag& Tag);
 
+	void PredictiveAdd(const FGameplayTag& Tag, FPredictionKey& PredictionKey);
+
 	bool NetDeltaSerialize(FNetDeltaSerializeInfo & DeltaParms)
 	{
 		return FastArrayDeltaSerialize<FActiveGameplayCue>(GameplayCues, DeltaParms, *this);
