merge from Orion, adding in TryActivateAbilityByClass() to the ability system comp.

diff --git a/Engine/Source/Runtime/GameplayAbilities/Private/AbilitySystemComponent_Abilities.cpp b/Engine/Source/Runtime/GameplayAbilities/Private/AbilitySystemComponent_Abilities.cpp
--- a/Engine/Source/Runtime/GameplayAbilities/Private/AbilitySystemComponent_Abilities.cpp
+++ b/Engine/Source/Runtime/GameplayAbilities/Private/AbilitySystemComponent_Abilities.cpp
@@ -830,6 +830,22 @@ bool UAbilitySystemComponent::TryActivateAbilitiesByTag(const FGameplayTagContai
 	return bSuccess;
 }
 
+bool UAbilitySystemComponent::TryActivateAbilityByClass(TSubclassOf<UGameplayAbility> InAbilityToActivate, bool bAllowRemoteActivation)
+{
+	bool bSuccess = false;
+
+	for (const FGameplayAbilitySpec& Spec : ActivatableAbilities.Items)
+	{
+		if (Spec.Ability == InAbilityToActivate.GetDefaultObject())
+		{
+			bSuccess |= TryActivateAbility(Spec.Handle, bAllowRemoteActivation);
+			break;
+		}
+	}
+
+	return bSuccess;
+}
+
 bool UAbilitySystemComponent::TryActivateAbility(FGameplayAbilitySpecHandle AbilityToActivate, bool bAllowRemoteActivation)
 {
 	FGameplayTagContainer FailureTags;
diff --git a/Engine/Source/Runtime/GameplayAbilities/Public/AbilitySystemComponent.h b/Engine/Source/Runtime/GameplayAbilities/Public/AbilitySystemComponent.h
--- a/Engine/Source/Runtime/GameplayAbilities/Public/AbilitySystemComponent.h
+++ b/Engine/Source/Runtime/GameplayAbilities/Public/AbilitySystemComponent.h
@@ -589,6 +589,14 @@ class GAMEPLAYABILITIES_API UAbilitySystemComponent : public UGameplayTasksCompo
 	UFUNCTION(BlueprintCallable, Category = "Abilities")
 	bool TryActivateAbilitiesByTag(const FGameplayTagContainer& GameplayTagContainer, bool bAllowRemoteActivation = true);
 
+	/**
+	 * Attempts to activate the ability that is passed in. This will check costs and requirements before doing so.
+	 * Returns true if it thinks it activated, but it may return false positives due to failure later in activation.
+	 * If bAllowRemoteActivation is true, it will remotely activate local/server abilities, if false it will only try to locally activate the ability
+	 */
+	UFUNCTION(BlueprintCallable, Category = "Abilities")
+	bool TryActivateAbilityByClass(TSubclassOf<UGameplayAbility> InAbilityToActivate, bool bAllowRemoteActivation = true);
+
 	/** 
 	 * Attempts to activate the given ability, will check costs and requirements before doing so.
 	 * Returns true if it thinks it activated, but it may return false positives due to failure later in activation.
