make struct parameter a reference

--------
Integrated using branch UE4-To-UE4-Orion (reversed) of change#2590803 by [REDACTED_PERSON] on 2015/06/17 16:04:00.

diff --git a/Engine/Source/Runtime/GameplayAbilities/Private/GameplayCueSet.cpp b/Engine/Source/Runtime/GameplayAbilities/Private/GameplayCueSet.cpp
--- a/Engine/Source/Runtime/GameplayAbilities/Private/GameplayCueSet.cpp
+++ b/Engine/Source/Runtime/GameplayAbilities/Private/GameplayCueSet.cpp
@@ -114,7 +114,7 @@ void UGameplayCueSet::PrintCues() const
 	}
 }
 
-bool UGameplayCueSet::HandleGameplayCueNotify_Internal(AActor* TargetActor, int32 DataIdx, EGameplayCueEvent::Type EventType, FGameplayCueParameters Parameters)
+bool UGameplayCueSet::HandleGameplayCueNotify_Internal(AActor* TargetActor, int32 DataIdx, EGameplayCueEvent::Type EventType, const FGameplayCueParameters& Parameters)
 {	
 	bool bReturnVal = false;
 
diff --git a/Engine/Source/Runtime/GameplayAbilities/Public/GameplayCueSet.h b/Engine/Source/Runtime/GameplayAbilities/Public/GameplayCueSet.h
--- a/Engine/Source/Runtime/GameplayAbilities/Public/GameplayCueSet.h
+++ b/Engine/Source/Runtime/GameplayAbilities/Public/GameplayCueSet.h
@@ -73,6 +73,6 @@ class GAMEPLAYABILITIES_API UGameplayCueSet : public UDataAsset
 	static FGameplayTag	BaseGameplayCueTag();
 
 protected:
-	virtual bool HandleGameplayCueNotify_Internal(AActor* TargetActor, int32 DataIdx, EGameplayCueEvent::Type EventType, FGameplayCueParameters Parameters);
+	virtual bool HandleGameplayCueNotify_Internal(AActor* TargetActor, int32 DataIdx, EGameplayCueEvent::Type EventType, const FGameplayCueParameters& Parameters);
 	virtual void BuildAccelerationMap_Internal();	
 };
\ No newline at end of file
