update default blueprint classes

--------
Integrated using branch UE4-To-UE4-Orion (reversed) of change#2546304 by [REDACTED_PERSON] on 2015/05/11 20:12:10.

diff --git a/Engine/Source/Runtime/GameplayAbilities/Public/GameplayEffect.h b/Engine/Source/Runtime/GameplayAbilities/Public/GameplayEffect.h
--- a/Engine/Source/Runtime/GameplayAbilities/Public/GameplayEffect.h
+++ b/Engine/Source/Runtime/GameplayAbilities/Public/GameplayEffect.h
@@ -516,7 +516,7 @@ enum class EGameplayEffectStackingPeriodPolicy : uint8
  *	The GameplayEffect definition. This is the data asset defined in the editor that drives everything.
  *  This is only blueprintable to allow for templating gameplay effects. Gameplay effects should NOT contain blueprint graphs.
  */
-UCLASS(Blueprintable)
+UCLASS(Blueprintable, meta = (ShortTooltip="A GameplayEffect modifies attributes and tags."))
 class GAMEPLAYABILITIES_API UGameplayEffect : public UObject, public IGameplayTagAssetInterface
 {
 
