Authority checks around AbilitySystemComponent remove gameplay effect functions

--------
Integrated using branch UE4-To-UE4-Orion (reversed) of change#2599580 by [REDACTED_PERSON] on 2015/06/24 17:05:49.

diff --git a/Engine/Source/Runtime/GameplayAbilities/Private/AbilitySystemComponent.cpp b/Engine/Source/Runtime/GameplayAbilities/Private/AbilitySystemComponent.cpp
--- a/Engine/Source/Runtime/GameplayAbilities/Private/AbilitySystemComponent.cpp
+++ b/Engine/Source/Runtime/GameplayAbilities/Private/AbilitySystemComponent.cpp
@@ -1137,8 +1137,7 @@ void UAbilitySystemComponent::ModifyActiveEffectStartTime(FActiveGameplayEffectH
 
 void UAbilitySystemComponent::RemoveActiveEffectsWithTags(const FGameplayTagContainer Tags)
 {
-	AActor *OwningActor = GetOwner();
-	if (OwningActor && OwningActor->Role == ROLE_Authority)
+	if (IsOwnerActorAuthoritative())
 	{
 		RemoveActiveEffects(FActiveGameplayEffectQuery(&Tags));
 	}
@@ -1146,7 +1145,10 @@ void UAbilitySystemComponent::RemoveActiveEffectsWithTags(const FGameplayTagCont
 
 void UAbilitySystemComponent::RemoveActiveEffects(const FActiveGameplayEffectQuery Query, int32 StacksToRemove)
 {
-	ActiveGameplayEffects.RemoveActiveEffects(Query, StacksToRemove);
+	if (IsOwnerActorAuthoritative())
+	{
+		ActiveGameplayEffects.RemoveActiveEffects(Query, StacksToRemove);
+	}
 }
 
 // ---------------------------------------------------------------------------------------
