Fix issues with Remove GameplayCue event:
-Issues where it would be called on client when effect was inhibited
-Issues where it could be double called on server or predicting clients

--------
Integrated using branch UE4-To-UE4-Orion (reversed) of change#2599733 by [REDACTED_PERSON] on 2015/06/24 18:22:05.

diff --git a/Engine/Source/Runtime/GameplayAbilities/Private/GameplayEffect.cpp b/Engine/Source/Runtime/GameplayAbilities/Private/GameplayEffect.cpp
--- a/Engine/Source/Runtime/GameplayAbilities/Private/GameplayEffect.cpp
+++ b/Engine/Source/Runtime/GameplayAbilities/Private/GameplayEffect.cpp
@@ -1230,13 +1230,7 @@ void FActiveGameplayEffect::PreReplicatedRemove(const struct FActiveGameplayEffe
 
 	ABILITY_LOG(Verbose, TEXT("PreReplicatedRemove: %s %s Marked as Pending Remove: %s"), *Handle.ToString(), *Spec.Def->GetName(), IsPendingRemove ? TEXT("TRUE") : TEXT("FALSE"));
 
-	const_cast<FActiveGameplayEffectsContainer&>(InArray).InternalOnActiveGameplayEffectRemoved(*this);	// Const cast is ok. It is there to prevent mutation of the GameplayEffects array, which this wont do.
-	
-	// Only invoke the GC event if the GE was not inhibited when it was removed
-	if (bIsInhibited == false)
-	{
-		InArray.Owner->InvokeGameplayCueEvent(Spec, EGameplayCueEvent::Removed);
-	}
+	const_cast<FActiveGameplayEffectsContainer&>(InArray).InternalOnActiveGameplayEffectRemoved(*this, !bIsInhibited);	// Const cast is ok. It is there to prevent mutation of the GameplayEffects array, which this wont do.
 }
 
 void FActiveGameplayEffect::PostReplicatedAdd(const struct FActiveGameplayEffectsContainer &InArray)
@@ -2427,12 +2421,7 @@ bool FActiveGameplayEffectsContainer::InternalRemoveActiveGameplayEffect(int32 I
 		// Don't invoke the GC event if the effect is inhibited, and thus the GC is already not active
 		ShouldInvokeGameplayCueEvent &= !Effect.bIsInhibited;
 
-		if (ShouldInvokeGameplayCueEvent)
-		{
-			Owner->InvokeGameplayCueEvent(Effect.Spec, EGameplayCueEvent::Removed);
-		}
-
-		InternalOnActiveGameplayEffectRemoved(Effect);
+		InternalOnActiveGameplayEffectRemoved(Effect, ShouldInvokeGameplayCueEvent);
 
 		if (Effect.DurationHandle.IsValid())
 		{
@@ -2492,7 +2481,7 @@ bool FActiveGameplayEffectsContainer::InternalRemoveActiveGameplayEffect(int32 I
 }
 
 /** Called by client and server: This does cleanup that has to happen whether the effect is being removed locally or due to replication */
-void FActiveGameplayEffectsContainer::InternalOnActiveGameplayEffectRemoved(const FActiveGameplayEffect& Effect)
+void FActiveGameplayEffectsContainer::InternalOnActiveGameplayEffectRemoved(const FActiveGameplayEffect& Effect, bool bInvokeGameplayCueEvents)
 {
 	if (Effect.Spec.Def)
 	{
@@ -2500,7 +2489,7 @@ void FActiveGameplayEffectsContainer::InternalOnActiveGameplayEffectRemoved(cons
 		RemoveActiveEffectTagDependency(Effect.Spec.Def->OngoingTagRequirements.IgnoreTags, Effect.Handle);
 		RemoveActiveEffectTagDependency(Effect.Spec.Def->OngoingTagRequirements.RequireTags, Effect.Handle);
 
-		RemoveActiveGameplayEffectGrantedTagsAndModifiers(Effect, !Effect.bIsInhibited);
+		RemoveActiveGameplayEffectGrantedTagsAndModifiers(Effect, bInvokeGameplayCueEvents);
 	}
 	else
 	{
@@ -2539,14 +2528,10 @@ void FActiveGameplayEffectsContainer::RemoveActiveGameplayEffectGrantedTagsAndMo
 	Owner->UpdateTagMap(Effect.Spec.DynamicGrantedTags, -1);
 
 	for (const FGameplayEffectCue& Cue : Effect.Spec.Def->GameplayCues)
-	{		
+	{
 		if (bInvokeGameplayCueEvents)
 		{
-			// TODO: Optimize this so only one batched RPC is called
-			for (auto It = Cue.GameplayCueTags.CreateConstIterator(); It; ++It)
-			{
-				Owner->RemoveGameplayCue(*It);
-			}
+			Owner->InvokeGameplayCueEvent(Effect.Spec, EGameplayCueEvent::Removed);			
 		}
 
 		Owner->UpdateTagMap(Cue.GameplayCueTags, -1);
diff --git a/Engine/Source/Runtime/GameplayAbilities/Public/GameplayEffect.h b/Engine/Source/Runtime/GameplayAbilities/Public/GameplayEffect.h
--- a/Engine/Source/Runtime/GameplayAbilities/Public/GameplayEffect.h
+++ b/Engine/Source/Runtime/GameplayAbilities/Public/GameplayEffect.h
@@ -1585,7 +1585,7 @@ private:
 	
 	/** Called both in server side creation and replication creation/deletion */
 	void InternalOnActiveGameplayEffectAdded(FActiveGameplayEffect& Effect);
-	void InternalOnActiveGameplayEffectRemoved(const FActiveGameplayEffect& Effect);
+	void InternalOnActiveGameplayEffectRemoved(const FActiveGameplayEffect& Effect, bool bInvokeGameplayCueEvents);
 
 	void RemoveActiveGameplayEffectGrantedTagsAndModifiers(const FActiveGameplayEffect& Effect, bool bInvokeGameplayCueEvents);
 	void AddActiveGameplayEffectGrantedTagsAndModifiers(FActiveGameplayEffect& Effect, bool bInvokeGameplayCueEvents);
