GameplayCue additions:
-Ability to instance a GameplayCueNotify actor per instigator and source object
-Removed nested maps in GameplayCueManager, replaced with key struct to combine class, source, target, instigator data.

--------
Integrated using branch UE4-To-UE4-Orion (reversed) of change#2599735 by [REDACTED_PERSON] on 2015/06/24 18:22:18.

diff --git a/Engine/Source/Runtime/GameplayAbilities/Private/GameplayCueManager.cpp b/Engine/Source/Runtime/GameplayAbilities/Private/GameplayCueManager.cpp
--- a/Engine/Source/Runtime/GameplayAbilities/Private/GameplayCueManager.cpp
+++ b/Engine/Source/Runtime/GameplayAbilities/Private/GameplayCueManager.cpp
@@ -108,33 +108,38 @@ void UGameplayCueManager::HandleGameplayCue(AActor* TargetActor, FGameplayTag Ga
 
 void UGameplayCueManager::EndGameplayCuesFor(AActor* TargetActor)
 {
-	TMap<TWeakObjectPtr<UClass>, TWeakObjectPtr<AGameplayCueNotify_Actor>> FoundMapActor;
-	if (NotifyMapActor.RemoveAndCopyValue(TargetActor, FoundMapActor))
+	for (auto It = NotifyMapActor.CreateIterator(); It; ++It)
 	{
-		for (auto It = FoundMapActor.CreateConstIterator(); It; ++It)
+		FGCNotifyActorKey& Key = It.Key();
+		if (Key.TargetActor == TargetActor)
 		{
 			AGameplayCueNotify_Actor* InstancedCue = It.Value().Get();
 			if (InstancedCue)
 			{
 				InstancedCue->OnOwnerDestroyed();
 			}
+			It.RemoveCurrent();
 		}
 	}
 }
 
-AGameplayCueNotify_Actor* UGameplayCueManager::GetInstancedCueActor(AActor* TargetActor, UClass* CueClass)
+AGameplayCueNotify_Actor* UGameplayCueManager::GetInstancedCueActor(AActor* TargetActor, UClass* CueClass, const FGameplayCueParameters& Parameters)
 {
+	QUICK_SCOPE_CYCLE_COUNTER(STAT_GameplayCueManager_GetInstancedCueActor);
+
+	AGameplayCueNotify_Actor* CDO = Cast<AGameplayCueNotify_Actor>(CueClass->ClassDefaultObject);
+	FGCNotifyActorKey	NotifyKey(TargetActor, CueClass, 
+							CDO->bUniqueInstancePerInstigator ? Parameters.EffectContext.GetInstigator() : nullptr, 
+							CDO->bUniqueInstancePerSourceObject ? Parameters.EffectContext.GetSourceObject() : nullptr);
+
 	AGameplayCueNotify_Actor* SpawnedCue = nullptr;
-	if (auto InnerMap = NotifyMapActor.Find(TargetActor))
-	{
-		if (auto WeakPtrPtr = InnerMap->Find(CueClass))
+	if (TWeakObjectPtr<AGameplayCueNotify_Actor>* WeakPtrPtr = NotifyMapActor.Find(NotifyKey))
+	{		
+		SpawnedCue = WeakPtrPtr->Get();
+		// If the cue is scheduled to be destroyed, don't reuse it, create a new one instead
+		if (SpawnedCue && SpawnedCue->GetLifeSpan() <= 0.0f)
 		{
-			SpawnedCue = WeakPtrPtr->Get();
-			// If the cue is scheduled to be destroyed, don't reuse it, create a new one instead
-			if (SpawnedCue && SpawnedCue->GetLifeSpan() <= 0.0f)
-			{
-				return SpawnedCue;
-			}
+			return SpawnedCue;
 		}
 	}
 
@@ -146,8 +151,7 @@ AGameplayCueNotify_Actor* UGameplayCueManager::GetInstancedCueActor(AActor* Targ
 		SpawnedCue = TargetActor->GetWorld()->SpawnActor<AGameplayCueNotify_Actor>(CueClass, TargetActor->GetActorLocation(), TargetActor->GetActorRotation(), SpawnParams);
 		if (ensure(SpawnedCue))
 		{
-			auto& InnerMap = NotifyMapActor.FindOrAdd(TargetActor);
-			InnerMap.Add(CueClass) = SpawnedCue;
+			NotifyMapActor.Add(NotifyKey, SpawnedCue);
 		}
 	}
 
diff --git a/Engine/Source/Runtime/GameplayAbilities/Private/GameplayCueNotify_Actor.cpp b/Engine/Source/Runtime/GameplayAbilities/Private/GameplayCueNotify_Actor.cpp
--- a/Engine/Source/Runtime/GameplayAbilities/Private/GameplayCueNotify_Actor.cpp
+++ b/Engine/Source/Runtime/GameplayAbilities/Private/GameplayCueNotify_Actor.cpp
@@ -13,6 +13,8 @@ AGameplayCueNotify_Actor::AGameplayCueNotify_Actor(const FObjectInitializer& Obj
 	PrimaryActorTick.bCanEverTick = true;
 	bAutoDestroyOnRemove = false;
 	AutoDestroyDelay = 0.f;
+	bUniqueInstancePerSourceObject = false;
+	bUniqueInstancePerInstigator = false;
 }
 
 #if WITH_EDITOR
diff --git a/Engine/Source/Runtime/GameplayAbilities/Private/GameplayCueSet.cpp b/Engine/Source/Runtime/GameplayAbilities/Private/GameplayCueSet.cpp
--- a/Engine/Source/Runtime/GameplayAbilities/Private/GameplayCueSet.cpp
+++ b/Engine/Source/Runtime/GameplayAbilities/Private/GameplayCueSet.cpp
@@ -175,7 +175,7 @@ bool UGameplayCueSet::HandleGameplayCueNotify_Internal(AActor* TargetActor, int3
 				if (ensure(CueManager))
 				{
 					//Get our instance. We should probably have a flag or something to determine if we want to reuse or stack instances. That would mean changing our map to have a list of active instances.
-					AGameplayCueNotify_Actor* SpawnedInstancedCue = CueManager->GetInstancedCueActor(TargetActor, CueData.LoadedGameplayCueClass);
+					AGameplayCueNotify_Actor* SpawnedInstancedCue = CueManager->GetInstancedCueActor(TargetActor, CueData.LoadedGameplayCueClass, Parameters);
 					if (ensure(SpawnedInstancedCue))
 					{
 						SpawnedInstancedCue->HandleGameplayCue(TargetActor, EventType, Parameters);
diff --git a/Engine/Source/Runtime/GameplayAbilities/Public/GameplayCueManager.h b/Engine/Source/Runtime/GameplayAbilities/Public/GameplayCueManager.h
--- a/Engine/Source/Runtime/GameplayAbilities/Public/GameplayCueManager.h
+++ b/Engine/Source/Runtime/GameplayAbilities/Public/GameplayCueManager.h
@@ -146,7 +146,7 @@ class GAMEPLAYABILITIES_API UGameplayCueManager : public UDataAsset
 	virtual void EndGameplayCuesFor(AActor* TargetActor);
 
 	/** Returns the cached instance cue. Creates it if it doesn't exist */
-	virtual AGameplayCueNotify_Actor* GetInstancedCueActor(AActor* TargetActor, UClass* CueClass);
+	virtual AGameplayCueNotify_Actor* GetInstancedCueActor(AActor* TargetActor, UClass* CueClass, const FGameplayCueParameters& Parameters);
 
 	// -------------------------------------------------------------
 	//  Loading GameplayCueNotifies from ObjectLibraries
@@ -175,8 +175,37 @@ class GAMEPLAYABILITIES_API UGameplayCueManager : public UDataAsset
 	UPROPERTY(transient)
 	FStreamableManager	StreamableManager;
 
-	// Fixme: we can combine the AActor* and the FGameplayTag into a single struct with a decent hash and avoid double map lookups
-	TMap<TWeakObjectPtr<AActor>, TMap<TWeakObjectPtr<UClass>, TWeakObjectPtr<AGameplayCueNotify_Actor>>>		NotifyMapActor;
+	struct FGCNotifyActorKey
+	{
+		FGCNotifyActorKey(AActor* InTargetActor, UClass* InCueClass, AActor* InInstigatorActor=nullptr, const UObject* InSourceObj=nullptr)
+		{
+			TargetActor = InTargetActor;
+			OptionalInstigatorActor = InInstigatorActor;
+			OptionalSourceObject = InSourceObj;
+			CueClass = InCueClass;
+		}
+
+		TWeakObjectPtr<AActor>	TargetActor;
+		TWeakObjectPtr<AActor>	OptionalInstigatorActor;
+		TWeakObjectPtr<UObject>	OptionalSourceObject;
+		TWeakObjectPtr<UClass>	CueClass;
+
+		FORCEINLINE bool operator==(const FGCNotifyActorKey& Other) const
+		{
+			return TargetActor == Other.TargetActor && CueClass == Other.CueClass &&
+					OptionalInstigatorActor == Other.OptionalInstigatorActor && OptionalSourceObject == Other.OptionalSourceObject;
+		}
+	};
+
+	FORCEINLINE friend uint32 GetTypeHash(const FGCNotifyActorKey& Key)
+	{
+		return GetTypeHash(Key.TargetActor)	^
+				GetTypeHash(Key.OptionalInstigatorActor) ^
+				GetTypeHash(Key.OptionalSourceObject) ^
+				GetTypeHash(Key.CueClass);
+	}
+	
+	TMap<FGCNotifyActorKey, TWeakObjectPtr<AGameplayCueNotify_Actor> >		NotifyMapActor;
 
 	void PrintGameplayCueNotifyMap();
 
diff --git a/Engine/Source/Runtime/GameplayAbilities/Public/GameplayCueNotify_Actor.h b/Engine/Source/Runtime/GameplayAbilities/Public/GameplayCueNotify_Actor.h
--- a/Engine/Source/Runtime/GameplayAbilities/Public/GameplayCueNotify_Actor.h
+++ b/Engine/Source/Runtime/GameplayAbilities/Public/GameplayCueNotify_Actor.h
@@ -77,6 +77,20 @@ class GAMEPLAYABILITIES_API AGameplayCueNotify_Actor : public AActor
 	UPROPERTY(EditDefaultsOnly, Category = GameplayCue)
 	bool IsOverride;
 
+	/**
+	 *	Does this cue get a new instance for each instigator? For example if two instigators apply a GC to the same source, do we create two of these GameplayCue Notify actors or just one?
+	 *	If the notify is simply playing FX or sounds on the source, it should not need unique instances. If this Notify is attaching a beam from the instigator to the target, it does need a unique instance per instigator.
+	 */ 	 
+	UPROPERTY(EditDefaultsOnly, Category = GameplayCue)
+	bool bUniqueInstancePerInstigator;
+
+	/**
+	 *	Does this cue get a new instance for each source object? For example if two source objects apply a GC to the same source, do we create two of these GameplayCue Notify actors or just one?
+	 *	If the notify is simply playing FX or sounds on the source, it should not need unique instances. If this Notify is attaching a beam from the source object to the target, it does need a unique instance per instigator.
+	 */ 	 
+	UPROPERTY(EditDefaultsOnly, Category = GameplayCue)
+	bool bUniqueInstancePerSourceObject;
+
 private:
 	virtual void DeriveGameplayCueTagFromAssetName();
 };
