Refactored Gameplay Effect Queries
- added new query structure that uses tag queries. old structure will be deprecated.
- added pin customization for tag queries
- added helpers for creating tag and effect queries effeciently in code

Also fixed some tag query issues
- fixed parse error in evaluator
- fixed potential crashes generating debug strings

diff --git a/Engine/Source/Editor/GameplayTagsEditor/Private/GameplayTagsGraphPanelPinFactory.h b/Engine/Source/Editor/GameplayTagsEditor/Private/GameplayTagsGraphPanelPinFactory.h
--- a/Engine/Source/Editor/GameplayTagsEditor/Private/GameplayTagsGraphPanelPinFactory.h
+++ b/Engine/Source/Editor/GameplayTagsEditor/Private/GameplayTagsGraphPanelPinFactory.h
@@ -5,6 +5,7 @@
 #include "EdGraphUtilities.h"
 #include "SGameplayTagGraphPin.h"
 #include "SGameplayTagContainerGraphPin.h"
+#include "SGameplayTagQueryGraphPin.h"
 #include "GameplayTags.h"
 
 class FGameplayTagsGraphPanelPinFactory: public FGraphPanelPinFactory
@@ -24,6 +25,11 @@ class FGameplayTagsGraphPanelPinFactory: public FGraphPanelPinFactory
 		{
 			return SNew(SGameplayTagContainerGraphPin, InPin);
 		}
-		return NULL;
+		if (InPin->PinType.PinCategory == K2Schema->PC_Struct && InPin->PinType.PinSubCategoryObject == FGameplayTagQuery::StaticStruct())
+		{
+			return SNew(SGameplayTagQueryGraphPin, InPin);
+		}
+
+		return nullptr;
 	}
 };
diff --git a/Engine/Source/Editor/GameplayTagsEditor/Private/SGameplayTagQueryGraphPin.cpp b/Engine/Source/Editor/GameplayTagsEditor/Private/SGameplayTagQueryGraphPin.cpp
new file mode 100644
--- /dev/null
+++ b/Engine/Source/Editor/GameplayTagsEditor/Private/SGameplayTagQueryGraphPin.cpp
@@ -0,0 +1,92 @@
+// Copyright 1998-2015 Epic Games, Inc. All Rights Reserved.
+
+#include "GameplayTagsEditorModulePrivatePCH.h"
+#include "SGameplayTagQueryGraphPin.h"
+#include "GameplayTagsModule.h"
+#include "GameplayTags.h"
+
+#define LOCTEXT_NAMESPACE "GameplayTagQueryGraphPin"
+
+void SGameplayTagQueryGraphPin::Construct(const FArguments& InArgs, UEdGraphPin* InGraphPinObj)
+{
+	TagQuery = MakeShareable(new FGameplayTagQuery());
+	SGraphPin::Construct( SGraphPin::FArguments(), InGraphPinObj );
+}
+
+TSharedRef<SWidget>	SGameplayTagQueryGraphPin::GetDefaultValueWidget()
+{
+	ParseDefaultValueData();
+
+	//Create widget
+	return SNew(SVerticalBox)
+		+SVerticalBox::Slot()
+		.AutoHeight()
+		[
+			SAssignNew( ComboButton, SComboButton )
+			.OnGetMenuContent(this, &SGameplayTagQueryGraphPin::GetListContent)
+			.ContentPadding( FMargin( 2.0f, 2.0f ) )
+			.Visibility( this, &SGraphPin::GetDefaultValueVisibility )
+			.ButtonContent()
+			[
+				SNew( STextBlock )
+				.Text( LOCTEXT("GameplayTagQueryWidget_Edit", "Edit") )
+			]
+		]
+		+SVerticalBox::Slot()
+		.AutoHeight()
+		[
+			QueryDesc()
+		];
+}
+
+void SGameplayTagQueryGraphPin::ParseDefaultValueData()
+{
+	FString const TagQueryString = GraphPinObj->GetDefaultAsString();
+
+	UProperty* const TQProperty = FindField<UProperty>(UEditableGameplayTagQuery::StaticClass(), TEXT("TagQueryExportText_Helper"));
+	if (TQProperty)
+	{
+		FGameplayTagQuery* const TQ = TagQuery.Get();
+		TQProperty->ImportText(*TagQueryString, TQ, 0, nullptr, GLog);
+	}
+}
+
+TSharedRef<SWidget> SGameplayTagQueryGraphPin::GetListContent()
+{
+	EditableQueries.Empty();
+	EditableQueries.Add(SGameplayTagQueryWidget::FEditableGameplayTagQueryDatum(GraphPinObj, TagQuery.Get(), &TagQueryExportText));
+
+	return SNew( SVerticalBox )
+		+SVerticalBox::Slot()
+		.AutoHeight()
+		.MaxHeight( 400 )
+		[
+			SNew(SGameplayTagQueryWidget, EditableQueries)
+			.OnQueryChanged(this, &SGameplayTagQueryGraphPin::OnQueryChanged)
+			.Visibility( this, &SGraphPin::GetDefaultValueVisibility )
+			.AutoSave(true)
+		];
+}
+
+void SGameplayTagQueryGraphPin::OnQueryChanged()
+{
+	// Set Pin Data
+	GraphPinObj->GetSchema()->TrySetDefaultValue(*GraphPinObj, TagQueryExportText);
+	QueryDescription = TagQuery->GetDescription();
+}
+
+TSharedRef<SWidget> SGameplayTagQueryGraphPin::QueryDesc()
+{
+	QueryDescription = TagQuery->GetDescription();
+
+	return SNew(STextBlock)
+		.Text(this, &SGameplayTagQueryGraphPin::GetQueryDescText)
+		.AutoWrapText(true);
+}
+
+FText SGameplayTagQueryGraphPin::GetQueryDescText() const
+{
+	return FText::FromString(QueryDescription);
+}
+
+#undef LOCTEXT_NAMESPACE
diff --git a/Engine/Source/Editor/GameplayTagsEditor/Private/SGameplayTagQueryGraphPin.h b/Engine/Source/Editor/GameplayTagsEditor/Private/SGameplayTagQueryGraphPin.h
new file mode 100644
--- /dev/null
+++ b/Engine/Source/Editor/GameplayTagsEditor/Private/SGameplayTagQueryGraphPin.h
@@ -0,0 +1,53 @@
+// Copyright 1998-2015 Epic Games, Inc. All Rights Reserved.
+
+#pragma once
+
+#include "SlateBasics.h"
+#include "SGameplayTagQueryWidget.h"
+#include "SGraphPin.h"
+
+class SGameplayTagQueryGraphPin : public SGraphPin
+{
+public:
+	SLATE_BEGIN_ARGS(SGameplayTagQueryGraphPin) {}
+	SLATE_END_ARGS()
+
+	void Construct(const FArguments& InArgs, UEdGraphPin* InGraphPinObj);
+
+	// Begin SGraphPin interface
+	virtual TSharedRef<SWidget>	GetDefaultValueWidget() override;
+	// End SGraphPin interface
+
+private:
+
+	/** Parses the Data from the pin to fill in the names of the array. */
+	void ParseDefaultValueData();
+
+	/** Callback function to create content for the combo button. */
+	TSharedRef<SWidget> GetListContent();
+
+	void OnQueryChanged();
+
+	/**
+	 * Creates SelectedTags List.
+	 * @return widget that contains the read only tag names for displaying on the node.
+	 */
+	TSharedRef<SWidget> QueryDesc();
+
+private:
+
+	// Combo Button for the drop down list.
+	TSharedPtr<SComboButton> ComboButton;
+
+	// Tag Container used for the GameplayTagWidget.
+	TSharedPtr<FGameplayTagQuery> TagQuery;
+
+	FString TagQueryExportText;
+
+	// Datum uses for the GameplayTagWidget.
+	TArray<SGameplayTagQueryWidget::FEditableGameplayTagQueryDatum> EditableQueries;
+
+	FText GetQueryDescText() const;
+
+	FString QueryDescription;
+};
diff --git a/Engine/Source/Editor/GameplayTagsEditor/Private/SGameplayTagQueryWidget.cpp b/Engine/Source/Editor/GameplayTagsEditor/Private/SGameplayTagQueryWidget.cpp
--- a/Engine/Source/Editor/GameplayTagsEditor/Private/SGameplayTagQueryWidget.cpp
+++ b/Engine/Source/Editor/GameplayTagsEditor/Private/SGameplayTagQueryWidget.cpp
@@ -18,8 +18,10 @@ void SGameplayTagQueryWidget::Construct(const FArguments& InArgs, const TArray<F
 	TagQueries = EditableTagQueries;
 
 	bReadOnly = InArgs._ReadOnly;
+	bAutoSave = InArgs._AutoSave;
 	OnSaveAndClose = InArgs._OnSaveAndClose;
 	OnCancel = InArgs._OnCancel;
+	OnQueryChanged = InArgs._OnQueryChanged;
 
 	// Tag the assets as transactional so they can support undo/redo
 	for (int32 AssetIdx = 0; AssetIdx < TagQueries.Num(); ++AssetIdx)
@@ -44,6 +46,7 @@ void SGameplayTagQueryWidget::Construct(const FArguments& InArgs, const TArray<F
 	FPropertyEditorModule& PropertyModule = FModuleManager::LoadModuleChecked<FPropertyEditorModule>("PropertyEditor");
 	Details = PropertyModule.CreateDetailView(ViewArgs);
 	Details->SetObject(EQ);
+	Details->OnFinishedChangingProperties().AddSP(this, &SGameplayTagQueryWidget::OnFinishedChangingProperties);
 
 	ChildSlot
 	[
@@ -68,6 +71,7 @@ void SGameplayTagQueryWidget::Construct(const FArguments& InArgs, const TArray<F
 					[
 						SNew(SButton)
 						.IsEnabled(!bReadOnly)
+						.Visibility(this, &SGameplayTagQueryWidget::GetSaveAndCloseButtonVisibility)
 						.OnClicked(this, &SGameplayTagQueryWidget::OnSaveAndCloseClicked)
 						.Text(LOCTEXT("GameplayTagQueryWidget_SaveAndClose", "Save and Close"))
 					]
@@ -75,6 +79,7 @@ void SGameplayTagQueryWidget::Construct(const FArguments& InArgs, const TArray<F
 					.AutoWidth()
 					[
 						SNew(SButton)
+						.Visibility(this, &SGameplayTagQueryWidget::GetCancelButtonVisibility)
 						.OnClicked(this, &SGameplayTagQueryWidget::OnCancelClicked)
 						.Text(LOCTEXT("GameplayTagQueryWidget_Cancel", "Close Without Saving"))
 					]
@@ -89,6 +94,26 @@ void SGameplayTagQueryWidget::Construct(const FArguments& InArgs, const TArray<F
 	];
 }
 
+void SGameplayTagQueryWidget::OnFinishedChangingProperties(const FPropertyChangedEvent& PropertyChangedEvent)
+{
+	if (bAutoSave)
+	{
+		SaveToTagQuery();
+	}
+
+	OnQueryChanged.ExecuteIfBound();
+}
+
+EVisibility SGameplayTagQueryWidget::GetSaveAndCloseButtonVisibility() const
+{
+	return bAutoSave ? EVisibility::Collapsed : EVisibility::Visible;
+}
+
+EVisibility SGameplayTagQueryWidget::GetCancelButtonVisibility() const
+{
+	return bAutoSave ? EVisibility::Collapsed : EVisibility::Visible;
+}
+
 UEditableGameplayTagQuery* SGameplayTagQueryWidget::CreateEditableQuery(FGameplayTagQuery& Q)
 {
 	UEditableGameplayTagQuery* const AnEditableQuery = Q.CreateEditableQuery();
@@ -111,7 +136,8 @@ SGameplayTagQueryWidget::~SGameplayTagQueryWidget()
 	}
 }
 
-FReply SGameplayTagQueryWidget::OnSaveAndCloseClicked()
+
+void SGameplayTagQueryWidget::SaveToTagQuery()
 {
 	// translate obj tree to token stream
 	if (EditableQuery.IsValid() && !bReadOnly)
@@ -120,9 +146,18 @@ FReply SGameplayTagQueryWidget::OnSaveAndCloseClicked()
 		for (auto& TQ : TagQueries)
 		{
 			TQ.TagQuery->BuildFromEditableQuery(*EditableQuery.Get());
+			if (TQ.TagQueryExportText != nullptr)
+			{
+				*TQ.TagQueryExportText = EditableQuery.Get()->GetTagQueryExportText(*TQ.TagQuery);
+			}
 			TQ.TagQueryOwner->MarkPackageDirty();
 		}
 	}
+}
+
+FReply SGameplayTagQueryWidget::OnSaveAndCloseClicked()
+{
+	SaveToTagQuery();
 
 	OnSaveAndClose.ExecuteIfBound();
 	return FReply::Handled();
diff --git a/Engine/Source/Editor/GameplayTagsEditor/Private/SGameplayTagQueryWidget.h b/Engine/Source/Editor/GameplayTagsEditor/Private/SGameplayTagQueryWidget.h
--- a/Engine/Source/Editor/GameplayTagsEditor/Private/SGameplayTagQueryWidget.h
+++ b/Engine/Source/Editor/GameplayTagsEditor/Private/SGameplayTagQueryWidget.h
@@ -9,25 +9,24 @@
 class SGameplayTagQueryWidget : public SCompoundWidget
 {
 public:
-
-	DECLARE_DELEGATE(FOnSaveAndClose)
-	DECLARE_DELEGATE(FOnCancel)
-	
 	SLATE_BEGIN_ARGS( SGameplayTagQueryWidget )
-	: _ReadOnly(false)
+		: _ReadOnly(false), _AutoSave(false)
 	{}
-		SLATE_ARGUMENT( bool, ReadOnly ) // Flag to set if the list is read only
-		SLATE_EVENT(FOnSaveAndClose, OnSaveAndClose) // Called when "Save and Close" button clicked
-		SLATE_EVENT(FOnCancel, OnCancel) // Called when "Close Without Saving" button clicked
+		SLATE_ARGUMENT( bool, ReadOnly )	// Flag to set if the list is read only
+		SLATE_ARGUMENT( bool, AutoSave)		// Flag to set if edits should be applied automatically (hides buttons)
+		SLATE_EVENT(FSimpleDelegate, OnSaveAndClose) // Called when "Save and Close" button clicked
+		SLATE_EVENT(FSimpleDelegate, OnCancel) // Called when "Close Without Saving" button clicked
+		SLATE_EVENT(FSimpleDelegate, OnQueryChanged)	// Called when the user has modified the query
 	SLATE_END_ARGS()
 
 	/** Simple struct holding a tag query and its owner for generic re-use of the widget */
 	struct FEditableGameplayTagQueryDatum
 	{
 		/** Constructor */
-		FEditableGameplayTagQueryDatum(class UObject* InOwnerObj, struct FGameplayTagQuery* InTagQuery)
+		FEditableGameplayTagQueryDatum(class UObject* InOwnerObj, struct FGameplayTagQuery* InTagQuery, FString* InTagExportText=nullptr)
 			: TagQueryOwner(InOwnerObj)
 			, TagQuery(InTagQuery)
+			, TagQueryExportText(InTagExportText)
 		{}
 
 		/** Owning UObject of the query being edited */
@@ -35,6 +34,9 @@ public:
 
 		/** Tag query to edit */
 		struct FGameplayTagQuery* TagQuery; 
+
+		/** The export text for FGameplayTagQuery, useful in some circumstances */
+		FString* TagQueryExportText;
 	};
 
 	/** Construct the actual widget */
@@ -45,16 +47,22 @@ public:
 private:
 
 	/* Flag to set if the list is read only*/
-	bool bReadOnly;
+	uint32 bReadOnly : 1;
+
+	/** If true, query will be written immediately on all changes. Otherwise, will only be written on user prompt (via the buttons). */
+	uint32 bAutoSave : 1;
 
 	/** Containers to modify */
 	TArray<FEditableGameplayTagQueryDatum> TagQueries;
 
 	/** Called when "save and close" is clicked */
-	FOnSaveAndClose OnSaveAndClose;
+	FSimpleDelegate OnSaveAndClose;
+
+	/** Called when the user has modified the query */
+	FSimpleDelegate OnQueryChanged;
 
 	/** Called when "cancel" is clicked */
-	FOnCancel OnCancel;
+	FSimpleDelegate OnCancel;
 
 	/** Properties Tab */
 	TSharedPtr<class IDetailsView> Details;
@@ -68,5 +76,13 @@ private:
 	/** Called when the user clicks the "Expand All" button; Expands the entire tag tree */
 	FReply OnCancelClicked();
 
+	void OnFinishedChangingProperties(const FPropertyChangedEvent& PropertyChangedEvent);
+
+	/** Controls visibility of the "save and close" button */
+	EVisibility GetSaveAndCloseButtonVisibility() const;
+	/** Controls visibility of the "cancel without saving" button */
+	EVisibility GetCancelButtonVisibility() const;
+
+	void SaveToTagQuery();
 };
 
diff --git a/Engine/Source/Runtime/GameplayAbilities/Private/AbilitySystemComponent.cpp b/Engine/Source/Runtime/GameplayAbilities/Private/AbilitySystemComponent.cpp
--- a/Engine/Source/Runtime/GameplayAbilities/Private/AbilitySystemComponent.cpp
+++ b/Engine/Source/Runtime/GameplayAbilities/Private/AbilitySystemComponent.cpp
@@ -1115,21 +1115,36 @@ bool UAbilitySystemComponent::CanApplyAttributeModifiers(const UGameplayEffect *
 	return ActiveGameplayEffects.CanApplyAttributeModifiers(GameplayEffect, Level, EffectContext);
 }
 
+// #deprecated, use FGameplayEffectQuery version
 TArray<float> UAbilitySystemComponent::GetActiveEffectsTimeRemaining(const FActiveGameplayEffectQuery Query) const
 {
 	return ActiveGameplayEffects.GetActiveEffectsTimeRemaining(Query);
 }
+TArray<float> UAbilitySystemComponent::GetActiveEffectsTimeRemaining(const FGameplayEffectQuery Query) const
+{
+	return ActiveGameplayEffects.GetActiveEffectsTimeRemaining(Query);
+}
 
+// #deprecated
 TArray<float> UAbilitySystemComponent::GetActiveEffectsDuration(const FActiveGameplayEffectQuery Query) const
 {
 	return ActiveGameplayEffects.GetActiveEffectsDuration(Query);
 }
+TArray<float> UAbilitySystemComponent::GetActiveEffectsDuration(const FGameplayEffectQuery Query) const
+{
+	return ActiveGameplayEffects.GetActiveEffectsDuration(Query);
+}
 
 TArray<FActiveGameplayEffectHandle> UAbilitySystemComponent::GetActiveEffects(const FActiveGameplayEffectQuery Query) const
 {
 	return ActiveGameplayEffects.GetActiveEffects(Query);
 }
 
+TArray<FActiveGameplayEffectHandle> UAbilitySystemComponent::GetActiveEffects(const FGameplayEffectQuery Query) const
+{
+	return ActiveGameplayEffects.GetActiveEffects(Query);
+}
+
 void UAbilitySystemComponent::ModifyActiveEffectStartTime(FActiveGameplayEffectHandle Handle, float StartTimeDiff)
 {
 	ActiveGameplayEffects.ModifyActiveEffectStartTime(Handle, StartTimeDiff);
diff --git a/Engine/Source/Runtime/GameplayAbilities/Private/GameplayEffect.cpp b/Engine/Source/Runtime/GameplayAbilities/Private/GameplayEffect.cpp
--- a/Engine/Source/Runtime/GameplayAbilities/Private/GameplayEffect.cpp
+++ b/Engine/Source/Runtime/GameplayAbilities/Private/GameplayEffect.cpp
@@ -2787,6 +2787,7 @@ bool FActiveGameplayEffectsContainer::CanApplyAttributeModifiers(const UGameplay
 	return true;
 }
 
+// #deprecated
 TArray<float> FActiveGameplayEffectsContainer::GetActiveEffectsTimeRemaining(const FActiveGameplayEffectQuery Query) const
 {
 	SCOPE_CYCLE_COUNTER(STAT_GameplayEffectsGetActiveEffectsTimeRemaining);
@@ -2811,7 +2812,33 @@ TArray<float> FActiveGameplayEffectsContainer::GetActiveEffectsTimeRemaining(con
 	// Note: keep one return location to avoid copy operation.
 	return ReturnList;
 }
+TArray<float> FActiveGameplayEffectsContainer::GetActiveEffectsTimeRemaining(const FGameplayEffectQuery Query) const
+{
+	SCOPE_CYCLE_COUNTER(STAT_GameplayEffectsGetActiveEffectsTimeRemaining);
+
+	float CurrentTime = GetWorldTime();
+
+	TArray<float>	ReturnList;
 
+	for (const FActiveGameplayEffect& Effect : this)
+	{
+		if (!Query.Matches(Effect))
+		{
+			continue;
+		}
+
+		float Elapsed = CurrentTime - Effect.StartWorldTime;
+		float Duration = Effect.GetDuration();
+
+		ReturnList.Add(Duration - Elapsed);
+	}
+
+	// Note: keep one return location to avoid copy operation.
+	return ReturnList;
+}
+
+
+// #deprecated
 TArray<float> FActiveGameplayEffectsContainer::GetActiveEffectsDuration(const FActiveGameplayEffectQuery Query) const
 {
 	SCOPE_CYCLE_COUNTER(STAT_GameplayEffectsGetActiveEffectsDuration);
@@ -2832,6 +2859,27 @@ TArray<float> FActiveGameplayEffectsContainer::GetActiveEffectsDuration(const FA
 	return ReturnList;
 }
 
+TArray<float> FActiveGameplayEffectsContainer::GetActiveEffectsDuration(const FGameplayEffectQuery Query) const
+{
+	SCOPE_CYCLE_COUNTER(STAT_GameplayEffectsGetActiveEffectsDuration);
+
+	TArray<float>	ReturnList;
+
+	for (const FActiveGameplayEffect& Effect : this)
+	{
+		if (!Query.Matches(Effect))
+		{
+			continue;
+		}
+
+		ReturnList.Add(Effect.GetDuration());
+	}
+
+	// Note: keep one return location to avoid copy operation.
+	return ReturnList;
+}
+
+// #deprecated
 TArray<FActiveGameplayEffectHandle> FActiveGameplayEffectsContainer::GetActiveEffects(const FActiveGameplayEffectQuery Query) const
 {
 	SCOPE_CYCLE_COUNTER(STAT_GameplayEffectsGetActiveEffects);
@@ -2851,6 +2899,25 @@ TArray<FActiveGameplayEffectHandle> FActiveGameplayEffectsContainer::GetActiveEf
 	return ReturnList;
 }
 
+TArray<FActiveGameplayEffectHandle> FActiveGameplayEffectsContainer::GetActiveEffects(const FGameplayEffectQuery Query) const
+{
+	SCOPE_CYCLE_COUNTER(STAT_GameplayEffectsGetActiveEffects);
+
+	TArray<FActiveGameplayEffectHandle> ReturnList;
+
+	for (const FActiveGameplayEffect& Effect : this)
+	{
+		if (!Query.Matches(Effect))
+		{
+			continue;
+		}
+
+		ReturnList.Add(Effect.Handle);
+	}
+
+	return ReturnList;
+}
+
 void FActiveGameplayEffectsContainer::ModifyActiveEffectStartTime(FActiveGameplayEffectHandle Handle, float StartTimeDiff)
 {
 	SCOPE_CYCLE_COUNTER(STAT_GameplayEffectsModifyActiveEffectStartTime);
@@ -2899,6 +2966,21 @@ int32 FActiveGameplayEffectsContainer::GetActiveEffectCount(const FActiveGamepla
 	return Count;
 }
 
+int32 FActiveGameplayEffectsContainer::GetActiveEffectCount(const FGameplayEffectQuery& Query) const
+{
+	int32 Count = 0;
+
+	for (const FActiveGameplayEffect& Effect : this)
+	{
+		if (Query.Matches(Effect))
+		{
+			Count += Effect.Spec.StackCount;
+		}
+	}
+
+	return Count;
+}
+
 FOnGameplayAttributeChange& FActiveGameplayEffectsContainer::RegisterGameplayAttributeEvent(FGameplayAttribute Attribute)
 {
 	return AttributeChangeDelegates.FindOrAdd(Attribute);
@@ -3058,44 +3140,187 @@ void FActiveGameplayEffectHandle::RemoveFromGlobalMap()
 
 // -----------------------------------------------------------------
 
-bool FActiveGameplayEffectQuery::Matches(const FActiveGameplayEffect& Effect) const
+FGameplayEffectQuery::FGameplayEffectQuery()
+	: EffectSource(nullptr),
+	EffectDefinition(nullptr)
+{
+}
+
+FGameplayEffectQuery::FGameplayEffectQuery(FActiveGameplayEffectQueryCustomMatch InCustomMatchDelegate)
+	: CustomMatchDelegate(InCustomMatchDelegate),
+	EffectSource(nullptr),
+	EffectDefinition(nullptr)
+{
+}
+
+
+FGameplayEffectQuery::FGameplayEffectQuery(FGameplayEffectQuery&& Other)
+{
+	*this = Other;
+}
+
+FGameplayEffectQuery& FGameplayEffectQuery::operator=(FGameplayEffectQuery&& Other)
 {
+	CustomMatchDelegate = MoveTemp(Other.CustomMatchDelegate);
+	CustomMatchDelegate_BP = MoveTemp(Other.CustomMatchDelegate_BP);
+	OwningTagQuery = MoveTemp(Other.OwningTagQuery);
+	EffectTagQuery = MoveTemp(Other.EffectTagQuery);
+	ModifyingAttribute = MoveTemp(Other.ModifyingAttribute);
+	EffectSource = Other.EffectSource;
+	EffectDefinition = Other.EffectDefinition;
+	IgnoreHandles = MoveTemp(Other.IgnoreHandles);
+	return *this;
+}
+
+bool FGameplayEffectQuery::Matches(const FActiveGameplayEffect& Effect) const
+{
+	// since all of these query conditions must be met to be considered a match, failing
+	// any one of them means we can return false
+
 	// Anything in the ignore handle list is an immediate non-match
 	if (IgnoreHandles.Contains(Effect.Handle))
 	{
 		return false;
 	}
 
-	if (CustomMatch.IsBound())
+	if (CustomMatchDelegate.IsBound())
 	{
-		return CustomMatch.Execute(Effect);
+		if (CustomMatchDelegate.Execute(Effect) == false)
+		{
+			return false;
+		}
+	}
+
+	if (CustomMatchDelegate_BP.IsBound())
+	{
+		bool bDelegateMatches = false;
+		CustomMatchDelegate_BP.Execute(Effect, bDelegateMatches);
+		if (bDelegateMatches == false)
+		{
+			return false;
+		}
 	}
 
-	// if we are looking for owning tags check them on the Granted Tags and Owned Tags Container
-	if (OwningTagContainer)
+	// if it fails to match any of these, fail to match overall (assuming AND relationship)
+	if (OwningTagQuery.IsEmpty() == false)
 	{
-		if (!Effect.Spec.Def->InheritableOwnedTagsContainer.CombinedTags.MatchesAny(*OwningTagContainer, true) &&
-			!Effect.Spec.DynamicGrantedTags.MatchesAny(*OwningTagContainer, false))
+		FGameplayTagContainer const& CombinedTags = Effect.Spec.Def->InheritableOwnedTagsContainer.CombinedTags;
+		if ((CombinedTags.Num() > 0) && (OwningTagQuery.Matches(CombinedTags) == false))		// note: matches on empty
 		{
-			// if the GameplayEffect didn't match check the spec for tags that were added when this effect was created
-			if (!Effect.Spec.CapturedSourceTags.GetSpecTags().MatchesAny(*OwningTagContainer, false))
+			if (OwningTagQuery.Matches(Effect.Spec.DynamicGrantedTags) == false)
 			{
+				if (OwningTagQuery.Matches(Effect.Spec.CapturedSourceTags.GetSpecTags()) == false)
+				{
+					// none of these match, so EffectTagQuery fails to match, so entire query fails to match, we can be done
+					return false;
+				}
+			}
+		}
+	}
+
+	if (EffectTagQuery.IsEmpty() == false)
+	{
+		FGameplayTagContainer const& CombinedTags = Effect.Spec.Def->InheritableOwnedTagsContainer.CombinedTags;
+		if ((CombinedTags.Num() > 0) && (OwningTagQuery.Matches(CombinedTags) == false))		// note: matches on empty
+		{
+			if (OwningTagQuery.Matches(Effect.Spec.DynamicAssetTags) == false)
+			{
+				// none of these match, so EffectTagQuery fails to match, so entire query fails to match, we can be done
 				return false;
 			}
 		}
-	}	
-	
-	// if we are just looking for Tags on the Effect then look at the Gameplay Effect Tags
-	if (EffectTagContainer)
+	}
+
+	// if we are looking for ModifyingAttribute go over each of the Spec Modifiers and check the Attributes
+	if (ModifyingAttribute.IsValid())
+	{
+		bool bEffectModifiesThisAttribute = false;
+
+		for (int32 ModIdx = 0; ModIdx < Effect.Spec.Modifiers.Num(); ++ModIdx)
+		{
+			const FGameplayModifierInfo& ModDef = Effect.Spec.Def->Modifiers[ModIdx];
+			const FModifierSpec& ModSpec = Effect.Spec.Modifiers[ModIdx];
+
+			if (ModDef.Attribute == ModifyingAttribute)
+			{
+				bEffectModifiesThisAttribute = true;
+				break;
+			}
+		}
+		if (bEffectModifiesThisAttribute == false)
+		{
+			// effect doesn't modify the attribute we are looking for, no match
+			return false;
+		}
+	}
+
+	// check source object
+	if (EffectSource != nullptr)
+	{
+		if (Effect.Spec.GetEffectContext().GetSourceObject() != EffectSource)
+		{
+			return false;
+		}
+	}
+
+	// check definition
+	if (EffectDefinition != nullptr)
 	{
-		if (!Effect.Spec.Def->InheritableGameplayEffectTags.CombinedTags.MatchesAny(*EffectTagContainer, true) &&
-			!Effect.Spec.DynamicAssetTags.MatchesAny(*EffectTagContainer, false))
+		if (Effect.Spec.Def != EffectDefinition)
 		{
-			// this doesn't match our Tags so bail
 			return false;
 		}
 	}
 
+	return true;
+}
+
+// static
+FGameplayEffectQuery FGameplayEffectQuery::MakeQuery_MatchAnyOwningTags(const FGameplayTagContainer& InOwningTagContainer)
+{
+	FGameplayEffectQuery OutQuery;
+	OutQuery.OwningTagQuery = FGameplayTagQuery::MakeQuery_MatchAnyTag(InOwningTagContainer);
+	return OutQuery;
+}
+
+bool FActiveGameplayEffectQuery::Matches(const FActiveGameplayEffect& Effect) const
+{
+	// Anything in the ignore handle list is an immediate non-match
+	if (IgnoreHandles.Contains(Effect.Handle))
+	{
+		return false;
+	}
+
+	if (CustomMatch.IsBound())
+	{
+		return CustomMatch.Execute(Effect);
+	}
+
+ 	// if we are looking for owning tags check them on the Granted Tags and Owned Tags Container
+ 	if (OwningTagContainer)
+ 	{
+ 		if (!Effect.Spec.Def->InheritableOwnedTagsContainer.CombinedTags.MatchesAny(*OwningTagContainer, true) &&
+ 			!Effect.Spec.DynamicGrantedTags.MatchesAny(*OwningTagContainer, false))
+ 		{
+ 			// if the GameplayEffect didn't match check the spec for tags that were added when this effect was created
+ 			if (!Effect.Spec.CapturedSourceTags.GetSpecTags().MatchesAny(*OwningTagContainer, false))
+ 			{
+ 				return false;
+ 			}
+ 		}
+ 	}	
+ 	
+ 	// if we are just looking for Tags on the Effect then look at the Gameplay Effect Tags
+ 	if (EffectTagContainer)
+ 	{
+ 		if (!Effect.Spec.Def->InheritableGameplayEffectTags.CombinedTags.MatchesAny(*EffectTagContainer, true) &&
+ 			!Effect.Spec.DynamicAssetTags.MatchesAny(*EffectTagContainer, false))
+ 		{
+ 			// this doesn't match our Tags so bail
+ 			return false;
+ 		}
+ 	}
+
 	// if we are just looking for Tags on the Effect then look at the Gameplay Effect Tags
 	if (EffectTagContainer_Rejection)
 	{
diff --git a/Engine/Source/Runtime/GameplayAbilities/Public/AbilitySystemComponent.h b/Engine/Source/Runtime/GameplayAbilities/Public/AbilitySystemComponent.h
--- a/Engine/Source/Runtime/GameplayAbilities/Public/AbilitySystemComponent.h
+++ b/Engine/Source/Runtime/GameplayAbilities/Public/AbilitySystemComponent.h
@@ -472,10 +472,13 @@ class GAMEPLAYABILITIES_API UAbilitySystemComponent : public UGameplayTasksCompo
 	//		-Get magnitude + duration of a movespeed buff
 
 	TArray<float> GetActiveEffectsTimeRemaining(const FActiveGameplayEffectQuery Query) const;
+	TArray<float> GetActiveEffectsTimeRemaining(const FGameplayEffectQuery Query) const;
 
 	TArray<float> GetActiveEffectsDuration(const FActiveGameplayEffectQuery Query) const;
+	TArray<float> GetActiveEffectsDuration(const FGameplayEffectQuery Query) const;
 
 	TArray<FActiveGameplayEffectHandle> GetActiveEffects(const FActiveGameplayEffectQuery Query) const;
+	TArray<FActiveGameplayEffectHandle> GetActiveEffects(const FGameplayEffectQuery Query) const;
 
 	void ModifyActiveEffectStartTime(FActiveGameplayEffectHandle Handle, float StartTimeDiff);
 
diff --git a/Engine/Source/Runtime/GameplayAbilities/Public/GameplayEffect.h b/Engine/Source/Runtime/GameplayAbilities/Public/GameplayEffect.h
--- a/Engine/Source/Runtime/GameplayAbilities/Public/GameplayEffect.h
+++ b/Engine/Source/Runtime/GameplayAbilities/Public/GameplayEffect.h
@@ -1187,7 +1187,7 @@ struct GAMEPLAYABILITIES_API FGameplayEffectSpecForRPC
  *  -Replication callbacks
  *
  */
-USTRUCT()
+USTRUCT(BlueprintType)
 struct GAMEPLAYABILITIES_API FActiveGameplayEffect : public FFastArraySerializerItem
 {
 	GENERATED_USTRUCT_BODY()
@@ -1279,6 +1279,63 @@ struct GAMEPLAYABILITIES_API FActiveGameplayEffect : public FFastArraySerializer
 	FActiveGameplayEffect* PendingNext;
 };
 
+DECLARE_DYNAMIC_DELEGATE_TwoParams(FActiveGameplayEffectQueryCustomMatch_Dynamic, FActiveGameplayEffect, Effect, bool&, bMatches);
+
+/** Every set condition within this query must match in order for the query to match. i.e. individual query elements are ANDed together. */
+USTRUCT(BlueprintType)
+struct GAMEPLAYABILITIES_API FGameplayEffectQuery
+{
+	GENERATED_USTRUCT_BODY()
+
+public:
+	// ctors and operators
+	FGameplayEffectQuery();
+	FGameplayEffectQuery(FActiveGameplayEffectQueryCustomMatch InCustomMatchDelegate);
+	FGameplayEffectQuery(FGameplayEffectQuery&& Other);
+	FGameplayEffectQuery& operator=(FGameplayEffectQuery&& Other);
+
+	/** Native delegate for providing custom matching conditions. */
+	FActiveGameplayEffectQueryCustomMatch CustomMatchDelegate;
+
+	/** BP-exposed delegate for providing custom matching conditions. */
+	UPROPERTY(EditAnywhere, BlueprintReadWrite, Category = Query)
+	FActiveGameplayEffectQueryCustomMatch_Dynamic CustomMatchDelegate_BP;
+
+	/** Query that is matched against InheritableOwnedTagsContainer */
+	UPROPERTY(EditAnywhere, BlueprintReadWrite, Category = Query)
+	FGameplayTagQuery OwningTagQuery;
+
+	/** Query that is matched against InheritableGameplayEffectTags */
+	UPROPERTY(EditAnywhere, BlueprintReadWrite, Category = Query)
+	FGameplayTagQuery EffectTagQuery;
+
+	/** Matches on GameplayEffects which modify given attribute. */
+	UPROPERTY(EditAnywhere, BlueprintReadWrite, Category = Query)
+	FGameplayAttribute ModifyingAttribute;
+
+	/** Matches on GameplayEffects which come from this source */
+	UPROPERTY(EditAnywhere, BlueprintReadWrite, Category = Query)
+	const UObject* EffectSource;
+
+	/** Matches on GameplayEffects with this definition */
+	UPROPERTY(EditAnywhere, BlueprintReadWrite, Category = Query)
+	const UGameplayEffect* EffectDefinition;
+
+	/** Handles to ignore as matches, even if other criteria is met */
+	TArray<FActiveGameplayEffectHandle> IgnoreHandles;
+
+	/** Returns true if Effect matches all specified criteria of this query, including CustomMatch delegates if bound. Returns false otherwise. */
+	bool Matches(const FActiveGameplayEffect& Effect) const;
+
+	/** 
+	 * Shortcuts for easily creating common query types 
+	 * @todo: add more as dictated by use cases
+	 */
+
+	/** Creates an effect query that will match if there are any common tags between the given tags and an ActiveGameplayEffect's owning tags */
+	static FGameplayEffectQuery MakeQuery_MatchAnyOwningTags(const FGameplayTagContainer& InOwningTagContainer);
+};
+
 /**
  *	Generic querying data structure for active GameplayEffects. Lets us ask things like:
  *		Give me duration/magnitude of active gameplay effects with these tags
@@ -1483,14 +1540,18 @@ struct GAMEPLAYABILITIES_API FActiveGameplayEffectsContainer : public FFastArray
 	bool CanApplyAttributeModifiers(const UGameplayEffect *GameplayEffect, float Level, const FGameplayEffectContextHandle& EffectContext);
 	
 	TArray<float> GetActiveEffectsTimeRemaining(const FActiveGameplayEffectQuery Query) const;
+	TArray<float> GetActiveEffectsTimeRemaining(const FGameplayEffectQuery Query) const;
 
 	TArray<float> GetActiveEffectsDuration(const FActiveGameplayEffectQuery Query) const;
+	TArray<float> GetActiveEffectsDuration(const FGameplayEffectQuery Query) const;
 
 	TArray<FActiveGameplayEffectHandle> GetActiveEffects(const FActiveGameplayEffectQuery Query) const;
+	TArray<FActiveGameplayEffectHandle> GetActiveEffects(const FGameplayEffectQuery Query) const;
 
 	void ModifyActiveEffectStartTime(FActiveGameplayEffectHandle Handle, float StartTimeDiff);
 
 	void RemoveActiveEffects(const FActiveGameplayEffectQuery Query, int32 StacksToRemove);
+	void RemoveActiveEffects(const FGameplayEffectQuery Query, int32 StacksToRemove);
 
 	/**
 	 * Get the count of the effects matching the specified query (including stack count)
@@ -1498,6 +1559,7 @@ struct GAMEPLAYABILITIES_API FActiveGameplayEffectsContainer : public FFastArray
 	 * @return Count of the effects matching the specified query
 	 */
 	int32 GetActiveEffectCount(const FActiveGameplayEffectQuery Query) const;
+	int32 GetActiveEffectCount(const FGameplayEffectQuery& Query) const;
 
 	float GetServerWorldTime() const;
 
diff --git a/Engine/Source/Runtime/GameplayTags/Classes/GameplayTagContainer.h b/Engine/Source/Runtime/GameplayTags/Classes/GameplayTagContainer.h
--- a/Engine/Source/Runtime/GameplayTags/Classes/GameplayTagContainer.h
+++ b/Engine/Source/Runtime/GameplayTags/Classes/GameplayTagContainer.h
@@ -128,7 +128,7 @@ struct GAMEPLAYTAGS_API FGameplayTagContainer
 
 	/** Assignment/Equality operators */
 	FGameplayTagContainer& operator=(FGameplayTagContainer const& Other);
-	FGameplayTagContainer & operator=(FGameplayTagContainer&& Other);
+	FGameplayTagContainer& operator=(FGameplayTagContainer&& Other);
 	bool operator==(FGameplayTagContainer const& Other) const;
 	bool operator!=(FGameplayTagContainer const& Other) const;
 
@@ -394,6 +394,7 @@ namespace EGameplayTagQueryStreamVersion
  *		);
  * 
  * Queries are internally represented as a byte stream that is memory-efficient and can be evaluated quickly at runtime.
+ * Note: these have an extensive details and graph pin customization for editing, so there is no need to expose the internals to Blueprints.
  */
 USTRUCT(BlueprintType)
 struct GAMEPLAYTAGS_API FGameplayTagQuery
@@ -403,6 +404,13 @@ struct GAMEPLAYTAGS_API FGameplayTagQuery
 public:
 	FGameplayTagQuery();
 
+	FGameplayTagQuery(FGameplayTagQuery const& Other);
+	FGameplayTagQuery(FGameplayTagQuery&& Other);
+
+	/** Assignment/Equality operators */
+	FGameplayTagQuery& operator=(FGameplayTagQuery const& Other);
+	FGameplayTagQuery& operator=(FGameplayTagQuery&& Other);
+
 private:
 	/** Versioning for future token stream protocol changes. See EGameplayTagQueryStreamVersion. */
 	UPROPERTY()
@@ -416,9 +424,13 @@ private:
 	UPROPERTY()
 	TArray<uint8> QueryTokenStream;
 
-	/** Human-readable string describing the query */
+	/** User-provided string describing the query */
+	UPROPERTY()
+	FString UserDescription;
+
+	/** Auto-generated string describing the query */
 	UPROPERTY()
-	FString Description;
+	FString AutoDescription;
 
 	/** Returns a gameplay tag from the tag dictionary */
 	FGameplayTag GetTagFromIndex(int32 TagIdx) const
@@ -427,7 +439,6 @@ private:
 		return TagDictionary[TagIdx];
 	}
 
-
 public:
 	/** Returns true if the given tags match this query, or false otherwise. */
 	bool Matches(FGameplayTagContainer const& Tags) const;
@@ -439,14 +450,17 @@ public:
 	void Clear();
 
 	/** Creates this query with the given root expression. */
-	void BuildQuery(struct FGameplayTagQueryExpression& RootQueryExpr);
+	void Build(struct FGameplayTagQueryExpression& RootQueryExpr, FString InUserDescription = FString());
+
+	/** Static function to assemble and return a query. */
+	static FGameplayTagQuery BuildQuery(struct FGameplayTagQueryExpression& RootQueryExpr, FString InDescription = FString());
 
 	/** Builds a FGameplayTagQueryExpression from this query. */
 	void GetQueryExpr(struct FGameplayTagQueryExpression& OutExpr) const;
 
 	/** Returns description string. */
-	FString GetDescription() const { return Description; };
-	
+	FString GetDescription() const { return UserDescription.IsEmpty() ? AutoDescription : UserDescription; };
+
 #if WITH_EDITOR
 	/** Creates this query based on the given EditableQuery object */
 	void BuildFromEditableQuery(class UEditableGameplayTagQuery& EditableQuery); 
@@ -457,11 +471,23 @@ public:
 
 	static const FGameplayTagQuery EmptyQuery;
 
+	/**
+	* Shortcuts for easily creating common query types
+	* @todo: add more as dictated by use cases
+	*/
+
+	/** Creates a tag query that will match if there are any common tags between the given tags and the tags being queries against. */
+	static FGameplayTagQuery MakeQuery_MatchAnyTag(FGameplayTagContainer const& InTags);
+
 	friend class FQueryEvaluator;
 };
 
 struct GAMEPLAYTAGS_API FGameplayTagQueryExpression
 {
+	/** 
+	 * Fluid syntax approach for setting the type of this expression. 
+	 */
+
 	FGameplayTagQueryExpression& AnyTagsMatch()
 	{
 		ExprType = EGameplayTagQueryExprType::AnyTagsMatch;
@@ -498,13 +524,18 @@ struct GAMEPLAYTAGS_API FGameplayTagQueryExpression
 		return *this;
 	}
 
+	FGameplayTagQueryExpression& AddTag(TCHAR const* TagString)
+	{
+		return AddTag(FName(TagString));
+	}
+	FGameplayTagQueryExpression& AddTag(FName TagName);
 	FGameplayTagQueryExpression& AddTag(FGameplayTag Tag)
 	{
 		ensure(UsesTagSet());
 		TagSet.Add(Tag);
 		return *this;
 	}
-
+	
 	FGameplayTagQueryExpression& AddTags(FGameplayTagContainer const& Tags)
 	{
 		ensure(UsesTagSet());
@@ -544,6 +575,17 @@ struct GAMEPLAYTAGS_API FGameplayTagQueryExpression
 	}
 };
 
+template<>
+struct TStructOpsTypeTraits<FGameplayTagQuery> : public TStructOpsTypeTraitsBase
+{
+	enum
+	{
+		WithCopy = true
+	};
+};
+
+
+
 /** 
  * This is an editor-only representation of a query, designed to be editable with a typical property window. 
  * To edit a query in the editor, an FGameplayTagQuery is converted to a set of UObjects and edited,  When finished,
@@ -556,10 +598,12 @@ class GAMEPLAYTAGS_API UEditableGameplayTagQuery : public UObject
 	GENERATED_BODY()
 
 public:
-
 	/** User-supplied description, shown in property details. Auto-generated description is shown if not supplied. */
 	UPROPERTY(EditDefaultsOnly, Category = Query)
-	FString Description;
+	FString UserDescription;
+
+	/** Automatically-generated description */
+	FString AutoDescription;
 
 	/** The base expression of this query. */
 	UPROPERTY(EditDefaultsOnly, Instanced, Category = Query)
@@ -568,7 +612,15 @@ public:
 #if WITH_EDITOR
 	/** Converts this editor query construct into the runtime-usable token stream. */
 	void EmitTokens(TArray<uint8>& TokenStream, TArray<FGameplayTag>& TagDictionary, FString* DebugString=nullptr) const;
+
+	/** Generates and returns the export text for this query. */
+	FString GetTagQueryExportText(FGameplayTagQuery const& TagQuery);
 #endif  // WITH_EDITOR
+
+private:
+	/** Property to hold a gameplay tag query so we can use the exporttext path to get a string representation. */
+	UPROPERTY()
+	FGameplayTagQuery TagQueryExportText_Helper;
 };
 
 UCLASS(abstract, editinlinenew, collapseCategories, Transient)
diff --git a/Engine/Source/Runtime/GameplayTags/Private/GameplayTagContainer.cpp b/Engine/Source/Runtime/GameplayTags/Private/GameplayTagContainer.cpp
--- a/Engine/Source/Runtime/GameplayTags/Private/GameplayTagContainer.cpp
+++ b/Engine/Source/Runtime/GameplayTags/Private/GameplayTagContainer.cpp
@@ -5,7 +5,6 @@
 const FGameplayTagContainer FGameplayTagContainer::EmptyContainer;
 const FGameplayTagQuery FGameplayTagQuery::EmptyQuery;
 
-
 /** Helper class to parse/eval query token streams. */
 class FQueryEvaluator
 {
@@ -74,8 +73,16 @@ bool FQueryEvaluator::Eval(FGameplayTagContainer const& Tags)
 	{
 		return false;
 	}
+	
+	bool bRet = false;
+
+	uint8 const bHasRootExpression = GetToken();
+	if (!bReadError && bHasRootExpression)
+	{
+		bRet = EvalExpr(Tags);
+
+	}
 
-	bool const bRet = EvalExpr(Tags);
 	ensure(CurStreamIdx == Query.QueryTokenStream.Num());
 	return bRet;
 }
@@ -825,6 +832,40 @@ FGameplayTagQuery::FGameplayTagQuery()
 {
 }
 
+FGameplayTagQuery::FGameplayTagQuery(FGameplayTagQuery const& Other)
+{
+	*this = Other;
+}
+
+FGameplayTagQuery::FGameplayTagQuery(FGameplayTagQuery&& Other)
+{
+	*this = Other;
+}
+
+/** Assignment/Equality operators */
+FGameplayTagQuery& FGameplayTagQuery::operator=(FGameplayTagQuery const& Other)
+{
+	if (this != &Other)
+	{
+		TokenStreamVersion = Other.TokenStreamVersion;
+		TagDictionary = Other.TagDictionary;
+		QueryTokenStream = Other.QueryTokenStream;
+		UserDescription = Other.UserDescription;
+		AutoDescription = Other.AutoDescription;
+	}
+	return *this;
+}
+
+FGameplayTagQuery& FGameplayTagQuery::operator=(FGameplayTagQuery&& Other)
+{
+	TokenStreamVersion = Other.TokenStreamVersion;
+	TagDictionary = MoveTemp(Other.TagDictionary);
+	QueryTokenStream = MoveTemp(Other.QueryTokenStream);
+	UserDescription = MoveTemp(Other.UserDescription);
+	AutoDescription = MoveTemp(Other.AutoDescription);
+	return *this;
+}
+
 bool FGameplayTagQuery::Matches(FGameplayTagContainer const& Tags) const
 {
 	FQueryEvaluator QE(*this);
@@ -848,8 +889,11 @@ void FGameplayTagQuery::GetQueryExpr(FGameplayTagQueryExpression& OutExpr) const
 	QE.Read(OutExpr);
 }
 
-void FGameplayTagQuery::BuildQuery(FGameplayTagQueryExpression& QueryExpr)
+void FGameplayTagQuery::Build(FGameplayTagQueryExpression& RootQueryExpr, FString InUserDescription)
 {
+	TokenStreamVersion = EGameplayTagQueryStreamVersion::LatestVersion;
+	UserDescription = InUserDescription;
+
 	// Reserve size here is arbitrary, goal is to minimizing reallocs while being respectful of mem usage
 	QueryTokenStream.Reset(128);
 	TagDictionary.Reset();
@@ -859,7 +903,26 @@ void FGameplayTagQuery::BuildQuery(FGameplayTagQueryExpression& QueryExpr)
 
 	// emit the query
 	QueryTokenStream.Add(1);		// true to indicate is has a root expression
-	QueryExpr.EmitTokens(QueryTokenStream, TagDictionary);
+	RootQueryExpr.EmitTokens(QueryTokenStream, TagDictionary);
+}
+
+// static 
+FGameplayTagQuery FGameplayTagQuery::BuildQuery(FGameplayTagQueryExpression& RootQueryExpr, FString InDescription)
+{
+	FGameplayTagQuery Q;
+	Q.Build(RootQueryExpr, InDescription);
+	return Q;
+}
+
+//static 
+FGameplayTagQuery FGameplayTagQuery::MakeQuery_MatchAnyTag(FGameplayTagContainer const& InTags)
+{
+	return FGameplayTagQuery::BuildQuery
+	(
+		FGameplayTagQueryExpression()
+		.AnyTagsMatch()
+		.AddTags(InTags)
+	);
 }
 
 
@@ -883,6 +946,8 @@ UEditableGameplayTagQuery* FQueryEvaluator::CreateEditableQuery()
 	}
 	ensure(CurStreamIdx == Query.QueryTokenStream.Num());
 
+	EditableQuery->UserDescription = Query.UserDescription;
+
 	return EditableQuery;
 }
 
@@ -1029,12 +1094,21 @@ void FGameplayTagQuery::BuildFromEditableQuery(UEditableGameplayTagQuery& Editab
 	QueryTokenStream.Reset();
 	TagDictionary.Reset();
 
-	Description = EditableQuery.Description;
-	FString* const StringRepresentation = Description.IsEmpty() ? &Description : nullptr;
+	UserDescription = EditableQuery.UserDescription;
 
 	// add stream version first
 	QueryTokenStream.Add(EGameplayTagQueryStreamVersion::LatestVersion);
-	EditableQuery.EmitTokens(QueryTokenStream, TagDictionary, StringRepresentation);
+	EditableQuery.EmitTokens(QueryTokenStream, TagDictionary, &AutoDescription);
+}
+
+FString UEditableGameplayTagQuery::GetTagQueryExportText(FGameplayTagQuery const& TagQuery)
+{
+	TagQueryExportText_Helper = TagQuery;
+	UProperty* const TQProperty = FindField<UProperty>(GetClass(), TEXT("TagQueryExportText_Helper"));
+
+	FString OutString;
+	TQProperty->ExportTextItem(OutString, (void*)&TagQueryExportText_Helper, (void*)&TagQueryExportText_Helper, this, 0);
+	return OutString;
 }
 
 void UEditableGameplayTagQuery::EmitTokens(TArray<uint8>& TokenStream, TArray<FGameplayTag>& TagDictionary, FString* DebugString) const
@@ -1053,7 +1127,10 @@ void UEditableGameplayTagQuery::EmitTokens(TArray<uint8>& TokenStream, TArray<FG
 	else
 	{
 		TokenStream.Add(0);		// false if no root expression
-		DebugString->Append(TEXT("undefined"));
+		if (DebugString)
+		{
+			DebugString->Append(TEXT("undefined"));
+		}
 	}
 }
 
@@ -1112,7 +1189,10 @@ void UEditableGameplayTagQueryExpression::EmitExprListTokens(TArray<UEditableGam
 		{
 			// null expression
 			TokenStream.Add(EGameplayTagQueryExprType::Undefined);
-			DebugString->Append(TEXT("undefined"));
+			if (DebugString)
+			{
+				DebugString->Append(TEXT("undefined"));
+			}
 		}
 
 		bFirstExpr = false;
@@ -1217,6 +1297,11 @@ void UEditableGameplayTagQueryExpression_NoExprMatch::EmitTokens(TArray<uint8>&
 #endif	// WITH_EDITOR
 
 
+FGameplayTagQueryExpression& FGameplayTagQueryExpression::AddTag(FName TagName)
+{
+	FGameplayTag const Tag = IGameplayTagsModule::RequestGameplayTag(TagName);
+	return AddTag(Tag);
+}
 
 void FGameplayTagQueryExpression::EmitTokens(TArray<uint8>& TokenStream, TArray<FGameplayTag>& TagDictionary) const
 {
@@ -1237,7 +1322,7 @@ void FGameplayTagQueryExpression::EmitTokens(TArray<uint8>& TokenStream, TArray<
 		for (auto Tag : TagSet)
 		{
 			int32 TagIdx = TagDictionary.AddUnique(Tag);
-			check(TagIdx <= 255);
+			check(TagIdx <= 254);		// we reserve token 255 for internal use, so 254 is max unique tags
 			TokenStream.Add((uint8)TagIdx);
 		}
 	}
