Added copy ctor to FGameplayEffectQuery (fixes clang)

diff --git a/Engine/Source/Runtime/GameplayAbilities/Private/GameplayEffect.cpp b/Engine/Source/Runtime/GameplayAbilities/Private/GameplayEffect.cpp
--- a/Engine/Source/Runtime/GameplayAbilities/Private/GameplayEffect.cpp
+++ b/Engine/Source/Runtime/GameplayAbilities/Private/GameplayEffect.cpp
@@ -3146,6 +3146,12 @@ FGameplayEffectQuery::FGameplayEffectQuery()
 {
 }
 
+FGameplayEffectQuery::FGameplayEffectQuery(const FGameplayEffectQuery& Other)
+{
+	*this = Other;
+}
+
+
 FGameplayEffectQuery::FGameplayEffectQuery(FActiveGameplayEffectQueryCustomMatch InCustomMatchDelegate)
 	: CustomMatchDelegate(InCustomMatchDelegate),
 	EffectSource(nullptr),
diff --git a/Engine/Source/Runtime/GameplayAbilities/Public/GameplayEffect.h b/Engine/Source/Runtime/GameplayAbilities/Public/GameplayEffect.h
--- a/Engine/Source/Runtime/GameplayAbilities/Public/GameplayEffect.h
+++ b/Engine/Source/Runtime/GameplayAbilities/Public/GameplayEffect.h
@@ -1290,6 +1290,7 @@ struct GAMEPLAYABILITIES_API FGameplayEffectQuery
 public:
 	// ctors and operators
 	FGameplayEffectQuery();
+	FGameplayEffectQuery(const FGameplayEffectQuery& Other);
 	FGameplayEffectQuery(FActiveGameplayEffectQueryCustomMatch InCustomMatchDelegate);
 	FGameplayEffectQuery(FGameplayEffectQuery&& Other);
 	FGameplayEffectQuery& operator=(FGameplayEffectQuery&& Other);
