add explicit assignment operator, fix clang compile

diff --git a/Engine/Source/Runtime/GameplayAbilities/Private/GameplayEffect.cpp b/Engine/Source/Runtime/GameplayAbilities/Private/GameplayEffect.cpp
--- a/Engine/Source/Runtime/GameplayAbilities/Private/GameplayEffect.cpp
+++ b/Engine/Source/Runtime/GameplayAbilities/Private/GameplayEffect.cpp
@@ -3178,6 +3178,20 @@ FGameplayEffectQuery& FGameplayEffectQuery::operator=(FGameplayEffectQuery&& Oth
 	return *this;
 }
 
+FGameplayEffectQuery& FGameplayEffectQuery::operator=(const FGameplayEffectQuery& Other)
+{
+	CustomMatchDelegate = Other.CustomMatchDelegate;
+	CustomMatchDelegate_BP = Other.CustomMatchDelegate_BP;
+	OwningTagQuery = Other.OwningTagQuery;
+	EffectTagQuery = Other.EffectTagQuery;
+	ModifyingAttribute = Other.ModifyingAttribute;
+	EffectSource = Other.EffectSource;
+	EffectDefinition = Other.EffectDefinition;
+	IgnoreHandles = Other.IgnoreHandles;
+	return *this;
+}
+
+
 bool FGameplayEffectQuery::Matches(const FActiveGameplayEffect& Effect) const
 {
 	// since all of these query conditions must be met to be considered a match, failing
diff --git a/Engine/Source/Runtime/GameplayAbilities/Public/GameplayEffect.h b/Engine/Source/Runtime/GameplayAbilities/Public/GameplayEffect.h
--- a/Engine/Source/Runtime/GameplayAbilities/Public/GameplayEffect.h
+++ b/Engine/Source/Runtime/GameplayAbilities/Public/GameplayEffect.h
@@ -1294,6 +1294,7 @@ public:
 	FGameplayEffectQuery(FActiveGameplayEffectQueryCustomMatch InCustomMatchDelegate);
 	FGameplayEffectQuery(FGameplayEffectQuery&& Other);
 	FGameplayEffectQuery& operator=(FGameplayEffectQuery&& Other);
+	FGameplayEffectQuery& operator=(const FGameplayEffectQuery& Other);
 
 	/** Native delegate for providing custom matching conditions. */
 	FActiveGameplayEffectQueryCustomMatch CustomMatchDelegate;
