Deprecated old effect queries.  Use FGameplayEffectQuery going forward.

diff --git a/Engine/Source/Runtime/GameplayAbilities/Private/Abilities/GameplayAbility.cpp b/Engine/Source/Runtime/GameplayAbilities/Private/Abilities/GameplayAbility.cpp
--- a/Engine/Source/Runtime/GameplayAbilities/Private/Abilities/GameplayAbility.cpp
+++ b/Engine/Source/Runtime/GameplayAbilities/Private/Abilities/GameplayAbility.cpp
@@ -750,7 +750,8 @@ float UGameplayAbility::GetCooldownTimeRemaining(const FGameplayAbilityActorInfo
 	const FGameplayTagContainer* CooldownTags = GetCooldownTags();
 	if (CooldownTags && CooldownTags->Num() > 0)
 	{
-		TArray< float > Durations = ActorInfo->AbilitySystemComponent->GetActiveEffectsTimeRemaining(FActiveGameplayEffectQuery(CooldownTags));
+		FGameplayEffectQuery const Query = FGameplayEffectQuery::MakeQuery_MatchAnyOwningTags(*CooldownTags);
+		TArray< float > Durations = ActorInfo->AbilitySystemComponent->GetActiveEffectsTimeRemaining(Query);
 		if (Durations.Num() > 0)
 		{
 			Durations.Sort();
@@ -773,10 +774,11 @@ void UGameplayAbility::GetCooldownTimeRemainingAndDuration(FGameplayAbilitySpecH
 	const FGameplayTagContainer* CooldownTags = GetCooldownTags();
 	if (CooldownTags && CooldownTags->Num() > 0)
 	{
-		TArray< float > DurationRemaining = ActorInfo->AbilitySystemComponent->GetActiveEffectsTimeRemaining(FActiveGameplayEffectQuery(CooldownTags));
+		FGameplayEffectQuery const Query = FGameplayEffectQuery::MakeQuery_MatchAnyOwningTags(*CooldownTags);
+		TArray< float > DurationRemaining = ActorInfo->AbilitySystemComponent->GetActiveEffectsTimeRemaining(Query);
 		if (DurationRemaining.Num() > 0)
 		{
-			TArray< float > Durations = ActorInfo->AbilitySystemComponent->GetActiveEffectsDuration(FActiveGameplayEffectQuery(CooldownTags));
+			TArray< float > Durations = ActorInfo->AbilitySystemComponent->GetActiveEffectsDuration(Query);
 			check(Durations.Num() == DurationRemaining.Num());
 			int32 BestIdx = 0;
 			float LongestTime = DurationRemaining[0];
@@ -1424,8 +1426,7 @@ void UGameplayAbility::BP_RemoveGameplayEffectFromOwnerWithAssetTags(FGameplayTa
 		return;
 	}
 
-	FActiveGameplayEffectQuery Query;
-	Query.EffectTagContainer = &WithTags;
+	FGameplayEffectQuery const Query = FGameplayEffectQuery::MakeQuery_MatchAnyEffectTags(WithTags);
 	CurrentActorInfo->AbilitySystemComponent->RemoveActiveEffects(Query, StacksToRemove);
 }
 
@@ -1436,8 +1437,7 @@ void UGameplayAbility::BP_RemoveGameplayEffectFromOwnerWithGrantedTags(FGameplay
 		return;
 	}
 
-	FActiveGameplayEffectQuery Query;
-	Query.OwningTagContainer = &WithGrantedTags;
+	FGameplayEffectQuery const Query = FGameplayEffectQuery::MakeQuery_MatchAnyOwningTags(WithGrantedTags);
 	CurrentActorInfo->AbilitySystemComponent->RemoveActiveEffects(Query, StacksToRemove);
 }
 
diff --git a/Engine/Source/Runtime/GameplayAbilities/Private/AbilitySystemComponent.cpp b/Engine/Source/Runtime/GameplayAbilities/Private/AbilitySystemComponent.cpp
--- a/Engine/Source/Runtime/GameplayAbilities/Private/AbilitySystemComponent.cpp
+++ b/Engine/Source/Runtime/GameplayAbilities/Private/AbilitySystemComponent.cpp
@@ -249,8 +249,8 @@ int32 UAbilitySystemComponent::GetGameplayEffectCount(TSubclassOf<UGameplayEffec
 
 	if (SourceGameplayEffect)
 	{
-		FActiveGameplayEffectQuery Query;
-		Query.CustomMatch.BindLambda([&](const FActiveGameplayEffect& CurEffect)
+		FGameplayEffectQuery Query;
+		Query.CustomMatchDelegate.BindLambda([&](const FActiveGameplayEffect& CurEffect)
 		{
 			bool bMatches = false;
 
@@ -727,7 +727,7 @@ FActiveGameplayEffectHandle UAbilitySystemComponent::ApplyGameplayEffectSpecToSe
 	if (bIsNetAuthority && Spec.Def->RemoveGameplayEffectsWithTags.CombinedTags.Num() > 0)
 	{
 		// Clear tags is always removing all stacks.
-		FActiveGameplayEffectQuery ClearQuery(&Spec.Def->RemoveGameplayEffectsWithTags.CombinedTags);
+		FGameplayEffectQuery ClearQuery = FGameplayEffectQuery::MakeQuery_MatchAnyOwningTags(Spec.Def->RemoveGameplayEffectsWithTags.CombinedTags);
 		if (MyHandle.IsValid())
 		{
 			ClearQuery.IgnoreHandles.Add(MyHandle);
@@ -857,8 +857,8 @@ void UAbilitySystemComponent::RemoveActiveGameplayEffectBySourceEffect(TSubclass
 {
 	if (GameplayEffect)
 	{
-		FActiveGameplayEffectQuery Query;
-		Query.CustomMatch.BindLambda([&](const FActiveGameplayEffect& CurEffect)
+		FGameplayEffectQuery Query;
+		Query.CustomMatchDelegate.BindLambda([&](const FActiveGameplayEffect& CurEffect)
 		{
 			bool bMatches = false;
 
@@ -1118,28 +1118,34 @@ bool UAbilitySystemComponent::CanApplyAttributeModifiers(const UGameplayEffect *
 // #deprecated, use FGameplayEffectQuery version
 TArray<float> UAbilitySystemComponent::GetActiveEffectsTimeRemaining(const FActiveGameplayEffectQuery Query) const
 {
+	PRAGMA_DISABLE_DEPRECATION_WARNINGS
 	return ActiveGameplayEffects.GetActiveEffectsTimeRemaining(Query);
+	PRAGMA_ENABLE_DEPRECATION_WARNINGS
 }
 TArray<float> UAbilitySystemComponent::GetActiveEffectsTimeRemaining(const FGameplayEffectQuery Query) const
 {
 	return ActiveGameplayEffects.GetActiveEffectsTimeRemaining(Query);
 }
 
-// #deprecated
+// #deprecated, use FGameplayEffectQuery version
 TArray<float> UAbilitySystemComponent::GetActiveEffectsDuration(const FActiveGameplayEffectQuery Query) const
 {
+	PRAGMA_DISABLE_DEPRECATION_WARNINGS
 	return ActiveGameplayEffects.GetActiveEffectsDuration(Query);
+	PRAGMA_ENABLE_DEPRECATION_WARNINGS
 }
 TArray<float> UAbilitySystemComponent::GetActiveEffectsDuration(const FGameplayEffectQuery Query) const
 {
 	return ActiveGameplayEffects.GetActiveEffectsDuration(Query);
 }
 
+// #deprecated, use FGameplayEffectQuery version
 TArray<FActiveGameplayEffectHandle> UAbilitySystemComponent::GetActiveEffects(const FActiveGameplayEffectQuery Query) const
 {
+	PRAGMA_DISABLE_DEPRECATION_WARNINGS
 	return ActiveGameplayEffects.GetActiveEffects(Query);
+	PRAGMA_ENABLE_DEPRECATION_WARNINGS
 }
-
 TArray<FActiveGameplayEffectHandle> UAbilitySystemComponent::GetActiveEffects(const FGameplayEffectQuery Query) const
 {
 	return ActiveGameplayEffects.GetActiveEffects(Query);
@@ -1154,11 +1160,21 @@ void UAbilitySystemComponent::RemoveActiveEffectsWithTags(const FGameplayTagCont
 {
 	if (IsOwnerActorAuthoritative())
 	{
-		RemoveActiveEffects(FActiveGameplayEffectQuery(&Tags));
+		RemoveActiveEffects(FGameplayEffectQuery::MakeQuery_MatchAnyOwningTags(Tags));
 	}
 }
 
+// #deprecated, use FGameplayEffectQuery version
 void UAbilitySystemComponent::RemoveActiveEffects(const FActiveGameplayEffectQuery Query, int32 StacksToRemove)
+{
+	if (IsOwnerActorAuthoritative())
+	{
+		PRAGMA_DISABLE_DEPRECATION_WARNINGS
+		ActiveGameplayEffects.RemoveActiveEffects(Query, StacksToRemove);
+		PRAGMA_ENABLE_DEPRECATION_WARNINGS
+	}
+}
+void UAbilitySystemComponent::RemoveActiveEffects(const FGameplayEffectQuery Query, int32 StacksToRemove)
 {
 	if (IsOwnerActorAuthoritative())
 	{
diff --git a/Engine/Source/Runtime/GameplayAbilities/Private/GameplayEffect.cpp b/Engine/Source/Runtime/GameplayAbilities/Private/GameplayEffect.cpp
--- a/Engine/Source/Runtime/GameplayAbilities/Private/GameplayEffect.cpp
+++ b/Engine/Source/Runtime/GameplayAbilities/Private/GameplayEffect.cpp
@@ -2935,6 +2935,7 @@ void FActiveGameplayEffectsContainer::ModifyActiveEffectStartTime(FActiveGamepla
 	}
 }
 
+// #deprecated, use FGameplayEffectQuery version
 void FActiveGameplayEffectsContainer::RemoveActiveEffects(const FActiveGameplayEffectQuery Query, int32 StacksToRemove)
 {
 	// Force a lock because the removals could cause other removals earlier in the array, so iterating backwards is not safe all by itself
@@ -2950,6 +2951,22 @@ void FActiveGameplayEffectsContainer::RemoveActiveEffects(const FActiveGameplayE
 		}
 	}
 }
+void FActiveGameplayEffectsContainer::RemoveActiveEffects(const FGameplayEffectQuery Query, int32 StacksToRemove)
+{
+	// Force a lock because the removals could cause other removals earlier in the array, so iterating backwards is not safe all by itself
+	GAMEPLAYEFFECT_SCOPE_LOCK();
+
+	// Manually iterating through in reverse because this is a removal operation
+	for (int32 idx = GetNumGameplayEffects() - 1; idx >= 0; --idx)
+	{
+		const FActiveGameplayEffect& Effect = *GetActiveGameplayEffect(idx);
+		if (Effect.IsPendingRemove == false && Query.Matches(Effect))
+		{
+			InternalRemoveActiveGameplayEffect(idx, StacksToRemove, true);
+		}
+	}
+}
+
 
 int32 FActiveGameplayEffectsContainer::GetActiveEffectCount(const FActiveGameplayEffectQuery Query) const
 {
@@ -3159,7 +3176,6 @@ FGameplayEffectQuery::FGameplayEffectQuery(FActiveGameplayEffectQueryCustomMatch
 {
 }
 
-
 FGameplayEffectQuery::FGameplayEffectQuery(FGameplayEffectQuery&& Other)
 {
 	*this = Other;
@@ -3296,10 +3312,18 @@ bool FGameplayEffectQuery::Matches(const FActiveGameplayEffect& Effect) const
 }
 
 // static
-FGameplayEffectQuery FGameplayEffectQuery::MakeQuery_MatchAnyOwningTags(const FGameplayTagContainer& InOwningTagContainer)
+FGameplayEffectQuery FGameplayEffectQuery::MakeQuery_MatchAnyOwningTags(const FGameplayTagContainer& InTags)
+{
+	FGameplayEffectQuery OutQuery;
+	OutQuery.OwningTagQuery = FGameplayTagQuery::MakeQuery_MatchAnyTag(InTags);
+	return OutQuery;
+}
+
+// static
+FGameplayEffectQuery FGameplayEffectQuery::MakeQuery_MatchAnyEffectTags(const FGameplayTagContainer& InTags)
 {
 	FGameplayEffectQuery OutQuery;
-	OutQuery.OwningTagQuery = FGameplayTagQuery::MakeQuery_MatchAnyTag(InOwningTagContainer);
+	OutQuery.EffectTagQuery = FGameplayTagQuery::MakeQuery_MatchAnyTag(InTags);
 	return OutQuery;
 }
 
diff --git a/Engine/Source/Runtime/GameplayAbilities/Public/Abilities/GameplayAbility.h b/Engine/Source/Runtime/GameplayAbilities/Public/Abilities/GameplayAbility.h
--- a/Engine/Source/Runtime/GameplayAbilities/Public/Abilities/GameplayAbility.h
+++ b/Engine/Source/Runtime/GameplayAbilities/Public/Abilities/GameplayAbility.h
@@ -671,7 +671,13 @@ protected:
 	//	Ability exclusion / canceling
 	//
 	// ----------------------------------------------------------------------------------------------------------------
-	
+
+	UPROPERTY(EditDefaultsOnly, Category = TagQueries)
+	FGameplayTagQuery CancelAbilitiesMatchingTagQuery;
+
+	UPROPERTY(VisibleDefaultsOnly, BlueprintReadOnly, Category = TagQueries)
+	FGameplayTagQuery ConstTagQuery;
+
 	/** Abilities with these tags are cancelled when this ability is executed */
 	UPROPERTY(EditDefaultsOnly, Category = Tags)
 	FGameplayTagContainer CancelAbilitiesWithTag;
diff --git a/Engine/Source/Runtime/GameplayAbilities/Public/AbilitySystemComponent.h b/Engine/Source/Runtime/GameplayAbilities/Public/AbilitySystemComponent.h
--- a/Engine/Source/Runtime/GameplayAbilities/Public/AbilitySystemComponent.h
+++ b/Engine/Source/Runtime/GameplayAbilities/Public/AbilitySystemComponent.h
@@ -471,12 +471,15 @@ class GAMEPLAYABILITIES_API UAbilitySystemComponent : public UGameplayTasksCompo
 	//		-Get duration of CD
 	//		-Get magnitude + duration of a movespeed buff
 
+	DEPRECATED(4.9, "FActiveGameplayEffectQuery is deprecated, use version that takes FGameplayEffectQuery")
 	TArray<float> GetActiveEffectsTimeRemaining(const FActiveGameplayEffectQuery Query) const;
 	TArray<float> GetActiveEffectsTimeRemaining(const FGameplayEffectQuery Query) const;
 
+	DEPRECATED(4.9, "FActiveGameplayEffectQuery is deprecated, use version that takes FGameplayEffectQuery")
 	TArray<float> GetActiveEffectsDuration(const FActiveGameplayEffectQuery Query) const;
 	TArray<float> GetActiveEffectsDuration(const FGameplayEffectQuery Query) const;
 
+	DEPRECATED(4.9, "FActiveGameplayEffectQuery is deprecated, use version that takes FGameplayEffectQuery")
 	TArray<FActiveGameplayEffectHandle> GetActiveEffects(const FActiveGameplayEffectQuery Query) const;
 	TArray<FActiveGameplayEffectHandle> GetActiveEffects(const FGameplayEffectQuery Query) const;
 
@@ -486,7 +489,9 @@ class GAMEPLAYABILITIES_API UAbilitySystemComponent : public UGameplayTasksCompo
 	void RemoveActiveEffectsWithTags(FGameplayTagContainer Tags);
 
 	/** Removes all active effects that match given query. StacksToRemove=-1 will remove all stacks. */
+	DEPRECATED(4.9, "FActiveGameplayEffectQuery is deprecated, use version that takes FGameplayEffectQuery")
 	void RemoveActiveEffects(const FActiveGameplayEffectQuery Query, int32 StacksToRemove=-1);
+	void RemoveActiveEffects(const FGameplayEffectQuery Query, int32 StacksToRemove = -1);
 
 	void OnRestackGameplayEffects();	
 	
diff --git a/Engine/Source/Runtime/GameplayAbilities/Public/GameplayEffect.h b/Engine/Source/Runtime/GameplayAbilities/Public/GameplayEffect.h
--- a/Engine/Source/Runtime/GameplayAbilities/Public/GameplayEffect.h
+++ b/Engine/Source/Runtime/GameplayAbilities/Public/GameplayEffect.h
@@ -21,8 +21,6 @@ class UAbilitySystemComponent;
 class UGameplayModMagnitudeCalculation;
 class UGameplayEffectExecutionCalculation;
 
-DECLARE_DELEGATE_RetVal_OneParam(bool, FActiveGameplayEffectQueryCustomMatch, const FActiveGameplayEffect&);
-
 /** Enumeration outlining the possible gameplay effect magnitude calculation policies. */
 UENUM()
 enum class EGameplayEffectMagnitudeCalculation : uint8
@@ -1279,6 +1277,7 @@ struct GAMEPLAYABILITIES_API FActiveGameplayEffect : public FFastArraySerializer
 	FActiveGameplayEffect* PendingNext;
 };
 
+DECLARE_DELEGATE_RetVal_OneParam(bool, FActiveGameplayEffectQueryCustomMatch, const FActiveGameplayEffect&);
 DECLARE_DYNAMIC_DELEGATE_TwoParams(FActiveGameplayEffectQueryCustomMatch_Dynamic, FActiveGameplayEffect, Effect, bool&, bMatches);
 
 /** Every set condition within this query must match in order for the query to match. i.e. individual query elements are ANDed together. */
@@ -1335,7 +1334,8 @@ public:
 	 */
 
 	/** Creates an effect query that will match if there are any common tags between the given tags and an ActiveGameplayEffect's owning tags */
-	static FGameplayEffectQuery MakeQuery_MatchAnyOwningTags(const FGameplayTagContainer& InOwningTagContainer);
+	static FGameplayEffectQuery MakeQuery_MatchAnyOwningTags(const FGameplayTagContainer& InTags);
+	static FGameplayEffectQuery MakeQuery_MatchAnyEffectTags(const FGameplayTagContainer& InTags);
 };
 
 /**
@@ -1541,17 +1541,21 @@ struct GAMEPLAYABILITIES_API FActiveGameplayEffectsContainer : public FFastArray
 
 	bool CanApplyAttributeModifiers(const UGameplayEffect *GameplayEffect, float Level, const FGameplayEffectContextHandle& EffectContext);
 	
+	DEPRECATED(4.9, "FActiveGameplayEffectQuery is deprecated, use FGameplayEffectQuery instead")
 	TArray<float> GetActiveEffectsTimeRemaining(const FActiveGameplayEffectQuery Query) const;
 	TArray<float> GetActiveEffectsTimeRemaining(const FGameplayEffectQuery Query) const;
 
+	DEPRECATED(4.9, "FActiveGameplayEffectQuery is deprecated, use FGameplayEffectQuery instead")
 	TArray<float> GetActiveEffectsDuration(const FActiveGameplayEffectQuery Query) const;
 	TArray<float> GetActiveEffectsDuration(const FGameplayEffectQuery Query) const;
 
+	DEPRECATED(4.9, "FActiveGameplayEffectQuery is deprecated, use FGameplayEffectQuery instead")
 	TArray<FActiveGameplayEffectHandle> GetActiveEffects(const FActiveGameplayEffectQuery Query) const;
 	TArray<FActiveGameplayEffectHandle> GetActiveEffects(const FGameplayEffectQuery Query) const;
 
 	void ModifyActiveEffectStartTime(FActiveGameplayEffectHandle Handle, float StartTimeDiff);
 
+	DEPRECATED(4.9, "FActiveGameplayEffectQuery is deprecated, use FGameplayEffectQuery instead")
 	void RemoveActiveEffects(const FActiveGameplayEffectQuery Query, int32 StacksToRemove);
 	void RemoveActiveEffects(const FGameplayEffectQuery Query, int32 StacksToRemove);
 
@@ -1560,6 +1564,7 @@ struct GAMEPLAYABILITIES_API FActiveGameplayEffectsContainer : public FFastArray
 	 * 
 	 * @return Count of the effects matching the specified query
 	 */
+	DEPRECATED(4.9, "FActiveGameplayEffectQuery is deprecated, use FGameplayEffectQuery instead")
 	int32 GetActiveEffectCount(const FActiveGameplayEffectQuery Query) const;
 	int32 GetActiveEffectCount(const FGameplayEffectQuery& Query) const;
 
