Fixed up a couple move ctors in gameplay queries

diff --git a/Engine/Source/Runtime/GameplayAbilities/Private/GameplayEffect.cpp b/Engine/Source/Runtime/GameplayAbilities/Private/GameplayEffect.cpp
--- a/Engine/Source/Runtime/GameplayAbilities/Private/GameplayEffect.cpp
+++ b/Engine/Source/Runtime/GameplayAbilities/Private/GameplayEffect.cpp
@@ -404,7 +404,6 @@ FGameplayEffectSpec::FGameplayEffectSpec()
 	, bDurationLocked(false)
 	, Level(UGameplayEffect::INVALID_LEVEL)
 {
-
 }
 
 FGameplayEffectSpec::FGameplayEffectSpec(const UGameplayEffect* InDef, const FGameplayEffectContextHandle& InEffectContext, float InLevel)
@@ -964,7 +963,6 @@ const FGameplayEffectAttributeCaptureDefinition& FGameplayEffectAttributeCapture
 FGameplayEffectAttributeCaptureSpecContainer::FGameplayEffectAttributeCaptureSpecContainer()
 	: bHasNonSnapshottedAttributes(false)
 {
-
 }
 
 FGameplayEffectAttributeCaptureSpecContainer::FGameplayEffectAttributeCaptureSpecContainer(FGameplayEffectAttributeCaptureSpecContainer&& Other)
@@ -981,7 +979,6 @@ FGameplayEffectAttributeCaptureSpecContainer::FGameplayEffectAttributeCaptureSpe
 {
 }
 
-
 FGameplayEffectAttributeCaptureSpecContainer& FGameplayEffectAttributeCaptureSpecContainer::operator=(FGameplayEffectAttributeCaptureSpecContainer&& Other)
 {
 	SourceAttributes = MoveTemp(Other.SourceAttributes);
@@ -3178,7 +3175,7 @@ FGameplayEffectQuery::FGameplayEffectQuery(FActiveGameplayEffectQueryCustomMatch
 
 FGameplayEffectQuery::FGameplayEffectQuery(FGameplayEffectQuery&& Other)
 {
-	*this = Other;
+	*this = MoveTemp(Other);
 }
 
 FGameplayEffectQuery& FGameplayEffectQuery::operator=(FGameplayEffectQuery&& Other)
diff --git a/Engine/Source/Runtime/GameplayTags/Private/GameplayTagContainer.cpp b/Engine/Source/Runtime/GameplayTags/Private/GameplayTagContainer.cpp
--- a/Engine/Source/Runtime/GameplayTags/Private/GameplayTagContainer.cpp
+++ b/Engine/Source/Runtime/GameplayTags/Private/GameplayTagContainer.cpp
@@ -839,7 +839,7 @@ FGameplayTagQuery::FGameplayTagQuery(FGameplayTagQuery const& Other)
 
 FGameplayTagQuery::FGameplayTagQuery(FGameplayTagQuery&& Other)
 {
-	*this = Other;
+	*this = MoveTemp(Other);
 }
 
 /** Assignment/Equality operators */
