Implement GetWorld() on GameplayCueManager and GameplayCueNotify_Static. Allows non instanced GameplayCues to access a valid world without passing World* around in blueprint functions.

--------
Integrated using branch UE4-To-UE4-Orion (reversed) of change#2605763 by [REDACTED_PERSON] on 2015/06/30 11:13:06.

diff --git a/Engine/Source/Runtime/GameplayAbilities/Private/GameplayCueManager.cpp b/Engine/Source/Runtime/GameplayAbilities/Private/GameplayCueManager.cpp
--- a/Engine/Source/Runtime/GameplayAbilities/Private/GameplayCueManager.cpp
+++ b/Engine/Source/Runtime/GameplayAbilities/Private/GameplayCueManager.cpp
@@ -34,6 +34,7 @@ UGameplayCueManager::UGameplayCueManager(const FObjectInitializer& PCIP)
 #endif
 
 	GlobalCueSet = NewObject<UGameplayCueSet>(this, TEXT("GlobalCueSet"));
+	CurrentWorld = nullptr;
 }
 
 bool IsDedicatedServerForGameplayCue()
@@ -92,6 +93,9 @@ void UGameplayCueManager::HandleGameplayCue(AActor* TargetActor, FGameplayTag Ga
 		FColor DebugColor = FColor::Green;
 		DrawDebugString(TargetActor->GetWorld(), FVector(0.f, 0.f, 100.f), DebugStr, TargetActor, DebugColor, DisplayGameplayCueDuration);
 	}
+
+	CurrentWorld = TargetActor->GetWorld();
+
 	// Give the global set a chance
 	check(GlobalCueSet);
 	if (bAcceptsCue)
@@ -104,6 +108,8 @@ void UGameplayCueManager::HandleGameplayCue(AActor* TargetActor, FGameplayTag Ga
 	{
 		GameplayCueInterface->HandleGameplayCue(TargetActor, GameplayCueTag, EventType, Parameters);
 	}
+
+	CurrentWorld = nullptr;
 }
 
 void UGameplayCueManager::EndGameplayCuesFor(AActor* TargetActor)
@@ -380,6 +386,12 @@ void UGameplayCueManager::HandleAssetDeleted(UObject *Object)
 
 #endif
 
+
+UWorld* UGameplayCueManager::GetWorld() const
+{
+	return CurrentWorld;
+}
+
 void UGameplayCueManager::PrintGameplayCueNotifyMap()
 {
 	check(GlobalCueSet);
diff --git a/Engine/Source/Runtime/GameplayAbilities/Private/GameplayCueNotify_Static.cpp b/Engine/Source/Runtime/GameplayAbilities/Private/GameplayCueNotify_Static.cpp
--- a/Engine/Source/Runtime/GameplayAbilities/Private/GameplayCueNotify_Static.cpp
+++ b/Engine/Source/Runtime/GameplayAbilities/Private/GameplayCueNotify_Static.cpp
@@ -114,3 +114,9 @@ bool UGameplayCueNotify_Static::OnRemove_Implementation(AActor* MyTarget, FGamep
 {
 	return false;
 }
+
+UWorld* UGameplayCueNotify_Static::GetWorld() const
+{
+	static TWeakObjectPtr<UGameplayCueManager> CueManager = UAbilitySystemGlobals::Get().GetGameplayCueManager();
+	return CueManager.IsValid() ? CueManager->GetWorld() : nullptr;
+}
\ No newline at end of file
diff --git a/Engine/Source/Runtime/GameplayAbilities/Public/GameplayCueManager.h b/Engine/Source/Runtime/GameplayAbilities/Public/GameplayCueManager.h
--- a/Engine/Source/Runtime/GameplayAbilities/Public/GameplayCueManager.h
+++ b/Engine/Source/Runtime/GameplayAbilities/Public/GameplayCueManager.h
@@ -209,6 +209,8 @@ class GAMEPLAYABILITIES_API UGameplayCueManager : public UDataAsset
 
 	void PrintGameplayCueNotifyMap();
 
+	virtual class UWorld* GetWorld() const override;
+
 #if WITH_EDITOR
 	bool IsAssetInLoadedPaths(UObject *Object) const;
 
@@ -245,4 +247,7 @@ protected:
 	/** Number of active gameplay cue send contexts, when it goes to 0 cues are flushed */
 	UPROPERTY()
 	int32 GameplayCueSendContextCount;
+
+	/** Cached world we are currently handling cues for. Used for non instanced GC Notifies that need world. */
+	UWorld* CurrentWorld;
 };
\ No newline at end of file
diff --git a/Engine/Source/Runtime/GameplayAbilities/Public/GameplayCueNotify_Static.h b/Engine/Source/Runtime/GameplayAbilities/Public/GameplayCueNotify_Static.h
--- a/Engine/Source/Runtime/GameplayAbilities/Public/GameplayCueNotify_Static.h
+++ b/Engine/Source/Runtime/GameplayAbilities/Public/GameplayCueNotify_Static.h
@@ -35,6 +35,8 @@ class GAMEPLAYABILITIES_API UGameplayCueNotify_Static : public UObject
 
 	virtual void HandleGameplayCue(AActor* MyTarget, EGameplayCueEvent::Type EventType, FGameplayCueParameters Parameters);
 
+	UWorld* GetWorld() const override;
+
 #if WITH_EDITOR
 	virtual void PostEditChangeProperty(FPropertyChangedEvent& PropertyChangedEvent) override;
 #endif // WITH_EDITOR
