only do reimport delegate when running in editor (dont do it if running with -game).

--------
Integrated using branch UE4-To-UE4-Orion (reversed) of change#2609934 by [REDACTED_PERSON] on 2015/07/02 19:01:02.

diff --git a/Engine/Source/Runtime/GameplayAbilities/Private/AttributeSet.cpp b/Engine/Source/Runtime/GameplayAbilities/Private/AttributeSet.cpp
--- a/Engine/Source/Runtime/GameplayAbilities/Private/AttributeSet.cpp
+++ b/Engine/Source/Runtime/GameplayAbilities/Private/AttributeSet.cpp
@@ -195,7 +195,7 @@ void FScalableFloat::SetScalingValue(float InCoeffecient, FName InRowName, UCurv
 void FScalableFloat::RegisterOnCurveTablePostReimport() const
 {
 #if WITH_EDITOR
-	if (!OnCurveTablePostReimportHandle.IsValid())
+	if (GIsEditor && !OnCurveTablePostReimportHandle.IsValid())
 	{
 		// Register our interest in knowing when our referenced curve table is changed, so that we can update FinalCurve appropriately
 		OnCurveTablePostReimportHandle = FReimportManager::Instance()->OnPostReimport().AddRaw(this, &FScalableFloat::OnCurveTablePostReimport);
@@ -206,7 +206,7 @@ void FScalableFloat::RegisterOnCurveTablePostReimport() const
 void FScalableFloat::UnRegisterOnCurveTablePostReimport() const
 {
 #if WITH_EDITOR
-	if (OnCurveTablePostReimportHandle.IsValid())
+	if (GIsEditor && OnCurveTablePostReimportHandle.IsValid())
 	{
 		FReimportManager::Instance()->OnPostReimport().Remove(OnCurveTablePostReimportHandle);
 	}
