Merge from Orion to Main

- adding in SetBlockAbilityBindingsArray() so the server can init the array correctly
- adding in the CanActivateAbility() to the code path that calls ServerTryActivateAbility() and if it cant call NotifyAbilityFailed()

diff --git a/Engine/Source/Runtime/GameplayAbilities/Private/AbilitySystemComponent_Abilities.cpp b/Engine/Source/Runtime/GameplayAbilities/Private/AbilitySystemComponent_Abilities.cpp
--- a/Engine/Source/Runtime/GameplayAbilities/Private/AbilitySystemComponent_Abilities.cpp
+++ b/Engine/Source/Runtime/GameplayAbilities/Private/AbilitySystemComponent_Abilities.cpp
@@ -903,9 +903,18 @@ bool UAbilitySystemComponent::TryActivateAbility(FGameplayAbilitySpecHandle Abil
 	{
 		if (bAllowRemoteActivation)
 		{
-			// No prediction key, server will assign a server-generated key
-			ServerTryActivateAbility(AbilityToActivate, Spec->InputPressed, FPredictionKey());
-			return true;
+			FGameplayTagContainer FailureTags;
+			if (Ability->CanActivateAbility(AbilityToActivate, ActorInfo, nullptr, nullptr, &FailureTags))
+			{
+				// No prediction key, server will assign a server-generated key
+				ServerTryActivateAbility(AbilityToActivate, Spec->InputPressed, FPredictionKey());
+				return true;
+			}
+			else
+			{
+				NotifyAbilityFailed(AbilityToActivate, Ability, FailureTags);
+				return false;
+			}
 		}
 
 		ABILITY_LOG(Log, TEXT("Can't activate ServerOnly or ServerInitiated ability %s when not the server."), *Ability->GetName());
@@ -1774,7 +1783,8 @@ void UAbilitySystemComponent::BindToInputComponent(UInputComponent* InputCompone
 void UAbilitySystemComponent::BindAbilityActivationToInputComponent(UInputComponent* InputComponent, FGameplayAbiliyInputBinds BindInfo)
 {
 	UEnum* EnumBinds = BindInfo.GetBindEnum();
-	BlockedAbilityBindings.SetNumZeroed(EnumBinds->NumEnums());
+
+	SetBlockAbilityBindingsArray(BindInfo);
 
 	for(int32 idx=0; idx < EnumBinds->NumEnums(); ++idx)
 	{
@@ -1823,6 +1833,12 @@ void UAbilitySystemComponent::BindAbilityActivationToInputComponent(UInputCompon
 	}
 }
 
+void UAbilitySystemComponent::SetBlockAbilityBindingsArray(FGameplayAbiliyInputBinds BindInfo)
+{
+	UEnum* EnumBinds = BindInfo.GetBindEnum();
+	BlockedAbilityBindings.SetNumZeroed(EnumBinds->NumEnums());
+}
+
 void UAbilitySystemComponent::AbilityLocalInputPressed(int32 InputID)
 {
 	// Consume the input if this InputID is overloaded with GenericConfirm/Cancel and the GenericConfim/Cancel callback is bound
diff --git a/Engine/Source/Runtime/GameplayAbilities/Public/AbilitySystemComponent.h b/Engine/Source/Runtime/GameplayAbilities/Public/AbilitySystemComponent.h
--- a/Engine/Source/Runtime/GameplayAbilities/Public/AbilitySystemComponent.h
+++ b/Engine/Source/Runtime/GameplayAbilities/Public/AbilitySystemComponent.h
@@ -782,7 +782,8 @@ public:
 	// ----------------------------------------------------------------------------------------------------------------
 
 	virtual void BindToInputComponent(UInputComponent* InputComponent);
-	
+
+	void SetBlockAbilityBindingsArray(FGameplayAbiliyInputBinds BindInfo);
 	virtual void BindAbilityActivationToInputComponent(UInputComponent* InputComponent, FGameplayAbiliyInputBinds BindInfo);
 
 	virtual void AbilityLocalInputPressed(int32 InputID);
