Guard against updating attribute aggregators for periodic effects. This was highlighted by a periodic effect with stacking rules, when stack count was refreshed, we were pushing periodic mods to aggregators (when they should only ever be applied as instant mods to the attributes base value). Fixes another potential bug with clients getting onrep from a periodic effect that has attribute mods (not executions).

--------
Integrated using branch UE4-To-UE4-Orion (reversed) of change#2613791 by [REDACTED_PERSON] on 2015/07/08 13:35:21.

diff --git a/Engine/Source/Runtime/GameplayAbilities/Private/GameplayEffect.cpp b/Engine/Source/Runtime/GameplayAbilities/Private/GameplayEffect.cpp
--- a/Engine/Source/Runtime/GameplayAbilities/Private/GameplayEffect.cpp
+++ b/Engine/Source/Runtime/GameplayAbilities/Private/GameplayEffect.cpp
@@ -1638,6 +1638,12 @@ void FActiveGameplayEffectsContainer::OnStackCountChange(FActiveGameplayEffect&
 
 void FActiveGameplayEffectsContainer::UpdateAllAggregatorModMagnitudes(FActiveGameplayEffect& ActiveEffect)
 {
+	// We should never be doing this for periodic effects since their mods are not persistent on attribute aggregators
+	if (ActiveEffect.Spec.GetPeriod() > UGameplayEffect::NO_PERIOD)
+	{
+		return;
+	}
+
 	const FGameplayEffectSpec& Spec = ActiveEffect.Spec;
 	TSet<FGameplayAttribute> AttributesToUpdate;
 
