Added some state to gameplaycue manager that can be used for previewing gameplay cues in the editor. WIP

--------
Integrated using branch UE4-To-UE4-Orion (reversed) of change#2617236 by [REDACTED_PERSON] on 2015/07/10 16:35:41.

diff --git a/Engine/Source/Runtime/GameplayAbilities/Private/GameplayCueManager.cpp b/Engine/Source/Runtime/GameplayAbilities/Private/GameplayCueManager.cpp
--- a/Engine/Source/Runtime/GameplayAbilities/Private/GameplayCueManager.cpp
+++ b/Engine/Source/Runtime/GameplayAbilities/Private/GameplayCueManager.cpp
@@ -25,6 +25,11 @@ static FAutoConsoleVariableRef CVarDurationeGameplayCues(TEXT("AbilitySystem.Gam
 int32 GameplayCueRunOnDedicatedServer = 0;
 static FAutoConsoleVariableRef CVarDedicatedServerGameplayCues(TEXT("AbilitySystem.GameplayCue.RunOnDedicatedServer"), GameplayCueRunOnDedicatedServer, TEXT("Run gameplay cue events on dedicated server"), ECVF_Default );
 
+#if WITH_EDITOR
+USceneComponent* UGameplayCueManager::PreviewComponent = nullptr;
+UWorld* UGameplayCueManager::PreviewWorld = nullptr;
+#endif
+
 UGameplayCueManager::UGameplayCueManager(const FObjectInitializer& PCIP)
 : Super(PCIP)
 {
@@ -74,6 +79,13 @@ void UGameplayCueManager::HandleGameplayCue(AActor* TargetActor, FGameplayTag Ga
 		return;
 	}
 
+#if WITH_EDITOR
+	if (GIsEditor && TargetActor == nullptr && UGameplayCueManager::PreviewComponent)
+	{
+		TargetActor = Cast<AActor>(AActor::StaticClass()->GetDefaultObject());
+	}
+#endif
+
 	if (TargetActor == nullptr)
 	{
 		ABILITY_LOG(Warning, TEXT("UGameplayCueManager::HandleGameplayCue called on null TargetActor. GameplayCueTag: %s."), *GameplayCueTag.ToString());
@@ -153,8 +165,14 @@ AGameplayCueNotify_Actor* UGameplayCueManager::GetInstancedCueActor(AActor* Targ
 	if (ensure(TargetActor) && ensure(CueClass))
 	{
 		FActorSpawnParameters SpawnParams;
+#if WITH_EDITOR
+		// Don't set owner if we are using fake CDO actor to do anim previewing
+		SpawnParams.Owner = (TargetActor && TargetActor->HasAnyFlags(RF_ClassDefaultObject) == false ? TargetActor : nullptr);
+#else
 		SpawnParams.Owner = TargetActor;
-		SpawnedCue = TargetActor->GetWorld()->SpawnActor<AGameplayCueNotify_Actor>(CueClass, TargetActor->GetActorLocation(), TargetActor->GetActorRotation(), SpawnParams);
+#endif
+		
+		SpawnedCue = GetWorld()->SpawnActor<AGameplayCueNotify_Actor>(CueClass, TargetActor->GetActorLocation(), TargetActor->GetActorRotation(), SpawnParams);
 		if (ensure(SpawnedCue))
 		{
 			NotifyMapActor.Add(NotifyKey, SpawnedCue);
@@ -389,6 +407,11 @@ void UGameplayCueManager::HandleAssetDeleted(UObject *Object)
 
 UWorld* UGameplayCueManager::GetWorld() const
 {
+#if WITH_EDITOR
+	if (PreviewWorld)
+		return PreviewWorld;
+#endif
+
 	return CurrentWorld;
 }
 
diff --git a/Engine/Source/Runtime/GameplayAbilities/Public/GameplayCueManager.h b/Engine/Source/Runtime/GameplayAbilities/Public/GameplayCueManager.h
--- a/Engine/Source/Runtime/GameplayAbilities/Public/GameplayCueManager.h
+++ b/Engine/Source/Runtime/GameplayAbilities/Public/GameplayCueManager.h
@@ -225,6 +225,10 @@ class GAMEPLAYABILITIES_API UGameplayCueManager : public UDataAsset
 	bool bAccelerationMapOutdated;
 
 	FOnGameplayCueNotifyChange	OnGameplayCueNotifyAddOrRemove;
+
+	static class USceneComponent* PreviewComponent;
+
+	static UWorld* PreviewWorld;
 #endif
 
 protected:
