TagResponseTable first pass

--------
Integrated using branch UE4-To-UE4-Orion (reversed) of change#2620606 by [REDACTED_PERSON] on 2015/07/14 14:59:26.

diff --git a/Engine/Source/Runtime/GameplayAbilities/Private/AbilitySystemComponent.cpp b/Engine/Source/Runtime/GameplayAbilities/Private/AbilitySystemComponent.cpp
--- a/Engine/Source/Runtime/GameplayAbilities/Private/AbilitySystemComponent.cpp
+++ b/Engine/Source/Runtime/GameplayAbilities/Private/AbilitySystemComponent.cpp
@@ -277,6 +277,11 @@ int32 UAbilitySystemComponent::GetGameplayEffectCount(TSubclassOf<UGameplayEffec
 	return Count;
 }
 
+int32 UAbilitySystemComponent::GetAggregatedStackCount(const FActiveGameplayEffectQuery& Query)
+{
+	return ActiveGameplayEffects.GetActiveEffectCount(Query);
+}
+
 FActiveGameplayEffectHandle UAbilitySystemComponent::BP_ApplyGameplayEffectToTarget(TSubclassOf<UGameplayEffect> GameplayEffectClass, UAbilitySystemComponent* Target, float Level, FGameplayEffectContextHandle Context)
 {
 	if (Target == nullptr)
@@ -404,9 +409,9 @@ void UAbilitySystemComponent::CaptureAttributeForGameplayEffect(OUT FGameplayEff
 	}
 }
 
-FOnGameplayEffectTagCountChanged& UAbilitySystemComponent::RegisterGameplayTagEvent(FGameplayTag Tag)
+FOnGameplayEffectTagCountChanged& UAbilitySystemComponent::RegisterGameplayTagEvent(FGameplayTag Tag, EGameplayTagEventType::Type EventType)
 {
-	return GameplayTagCountContainer.RegisterGameplayTagEvent(Tag);
+	return GameplayTagCountContainer.RegisterGameplayTagEvent(Tag, EventType);
 }
 
 FOnGameplayEffectTagCountChanged& UAbilitySystemComponent::RegisterGenericGameplayTagEvent()
@@ -514,6 +519,15 @@ void UAbilitySystemComponent::UpdateTagMap(const FGameplayTagContainer& Containe
 	}
 }
 
+void UAbilitySystemComponent::NotifyTagMap_StackCountChange(const FGameplayTagContainer& Container)
+{
+	for (auto TagIt = Container.CreateConstIterator(); TagIt; ++TagIt)
+	{
+		const FGameplayTag& Tag = *TagIt;
+		GameplayTagCountContainer.Notify_StackCountChange(Tag);
+	}
+}
+
 // ------------------------------------------------------------------------
 
 FActiveGameplayEffectHandle UAbilitySystemComponent::ApplyGameplayEffectSpecToTarget(OUT FGameplayEffectSpec &Spec, UAbilitySystemComponent *Target, FPredictionKey PredictionKey)
@@ -893,6 +907,11 @@ float UAbilitySystemComponent::GetGameplayEffectMagnitude(FActiveGameplayEffectH
 	return ActiveGameplayEffects.GetGameplayEffectMagnitude(Handle, Attribute);
 }
 
+void UAbilitySystemComponent::SetActiveGameplayEffectLevel(FActiveGameplayEffectHandle ActiveHandle, int32 NewLevel)
+{
+	ActiveGameplayEffects.SetActiveGameplayEffectLevel(ActiveHandle, NewLevel);
+}
+
 int32 UAbilitySystemComponent::GetCurrentStackCount(FActiveGameplayEffectHandle Handle) const
 {
 	if (const FActiveGameplayEffect* ActiveGE = ActiveGameplayEffects.GetActiveGameplayEffect(Handle))
diff --git a/Engine/Source/Runtime/GameplayAbilities/Private/AbilitySystemComponent_Abilities.cpp b/Engine/Source/Runtime/GameplayAbilities/Private/AbilitySystemComponent_Abilities.cpp
--- a/Engine/Source/Runtime/GameplayAbilities/Private/AbilitySystemComponent_Abilities.cpp
+++ b/Engine/Source/Runtime/GameplayAbilities/Private/AbilitySystemComponent_Abilities.cpp
@@ -8,7 +8,7 @@
 #include "Abilities/Tasks/AbilityTask.h"
 #include "TickableAttributeSetInterface.h"
 #include "GameplayPrediction.h"
-
+#include "GameplayTagResponseTable.h"
 #include "Net/UnrealNetwork.h"
 #include "MessageLog.h"
 #include "UObjectToken.h"
@@ -121,6 +121,11 @@ void UAbilitySystemComponent::InitAbilityActorInfo(AActor* InOwnerActor, AActor*
 			}
 		}
 	}
+
+	if (UGameplayTagReponseTable* TagTable = UAbilitySystemGlobals::Get().GetGameplayTagResponseTable())
+	{
+		TagTable->RegisterResponseForEvents(this);
+	}
 }
 
 bool UAbilitySystemComponent::GetShouldTick() const 
diff --git a/Engine/Source/Runtime/GameplayAbilities/Private/AbilitySystemGlobals.cpp b/Engine/Source/Runtime/GameplayAbilities/Private/AbilitySystemGlobals.cpp
--- a/Engine/Source/Runtime/GameplayAbilities/Private/AbilitySystemGlobals.cpp
+++ b/Engine/Source/Runtime/GameplayAbilities/Private/AbilitySystemGlobals.cpp
@@ -36,6 +36,7 @@ void UAbilitySystemGlobals::InitGlobalData()
 	InitAttributeDefaults();
 
 	GetGameplayCueManager();
+	GetGameplayTagResponseTable();
 	InitGlobalTags();
 }
 
@@ -229,7 +230,7 @@ UGameplayCueManager* UAbilitySystemGlobals::GetGameplayCueManager()
 {
 	if (GlobalGameplayCueManager == nullptr)
 	{
-		GlobalGameplayCueManager = LoadObject<UGameplayCueManager>(NULL, *GlobalGameplayCueManagerName.ToString(), NULL, LOAD_None, NULL);
+		GlobalGameplayCueManager = LoadObject<UGameplayCueManager>(nullptr, *GlobalGameplayCueManagerName.ToString(), nullptr, LOAD_None, nullptr);
 		if (GameplayCueNotifyPaths.Num() > 0)
 		{
 			GlobalGameplayCueManager->LoadObjectLibraryFromPaths(GameplayCueNotifyPaths);
@@ -239,6 +240,16 @@ UGameplayCueManager* UAbilitySystemGlobals::GetGameplayCueManager()
 	return GlobalGameplayCueManager;
 }
 
+UGameplayTagReponseTable* UAbilitySystemGlobals::GetGameplayTagResponseTable()
+{
+	if (GameplayTagResponseTable == nullptr && GameplayTagResponseTableName.IsValid())
+	{
+		GameplayTagResponseTable = LoadObject<UGameplayTagReponseTable>(nullptr, *GameplayTagResponseTableName.ToString(), nullptr, LOAD_None, nullptr);
+	}
+
+	return GameplayTagResponseTable;
+}
+
 void UAbilitySystemGlobals::GlobalPreGameplayEffectSpecApply(FGameplayEffectSpec& Spec, UAbilitySystemComponent* AbilitySystemComponent)
 {
 
diff --git a/Engine/Source/Runtime/GameplayAbilities/Private/GameplayEffect.cpp b/Engine/Source/Runtime/GameplayAbilities/Private/GameplayEffect.cpp
--- a/Engine/Source/Runtime/GameplayAbilities/Private/GameplayEffect.cpp
+++ b/Engine/Source/Runtime/GameplayAbilities/Private/GameplayEffect.cpp
@@ -1634,6 +1634,8 @@ void FActiveGameplayEffectsContainer::OnStackCountChange(FActiveGameplayEffect&
 {
 	MarkItemDirty(ActiveEffect);
 	UpdateAllAggregatorModMagnitudes(ActiveEffect);
+	Owner->NotifyTagMap_StackCountChange(ActiveEffect.Spec.Def->InheritableOwnedTagsContainer.CombinedTags);
+	Owner->NotifyTagMap_StackCountChange(ActiveEffect.Spec.DynamicGrantedTags);
 }
 
 void FActiveGameplayEffectsContainer::UpdateAllAggregatorModMagnitudes(FActiveGameplayEffect& ActiveEffect)
@@ -1817,6 +1819,21 @@ float FActiveGameplayEffectsContainer::GetGameplayEffectMagnitude(FActiveGamepla
 	return -1.f;
 }
 
+void FActiveGameplayEffectsContainer::SetActiveGameplayEffectLevel(FActiveGameplayEffectHandle ActiveHandle, int32 NewLevel)
+{
+	for (FActiveGameplayEffect& Effect : this)
+	{
+		if (Effect.Handle == ActiveHandle)
+		{
+			Effect.Spec.SetLevel(NewLevel);
+			MarkItemDirty(Effect);
+			Effect.Spec.CalculateModifierMagnitudes();
+			UpdateAllAggregatorModMagnitudes(Effect);
+			break;
+		}
+	}
+}
+
 const FGameplayTagContainer* FActiveGameplayEffectsContainer::GetGameplayEffectSourceTagsFromHandle(FActiveGameplayEffectHandle Handle) const
 {
 	// @todo: Need to consider this with tag changes
diff --git a/Engine/Source/Runtime/GameplayAbilities/Private/GameplayEffectTypes.cpp b/Engine/Source/Runtime/GameplayAbilities/Private/GameplayEffectTypes.cpp
--- a/Engine/Source/Runtime/GameplayAbilities/Private/GameplayEffectTypes.cpp
+++ b/Engine/Source/Runtime/GameplayAbilities/Private/GameplayEffectTypes.cpp
@@ -321,9 +321,34 @@ void FGameplayTagCountContainer::UpdateTagCount(const FGameplayTagContainer& Con
 	}
 }
 
-FOnGameplayEffectTagCountChanged& FGameplayTagCountContainer::RegisterGameplayTagEvent(const FGameplayTag& Tag)
+void FGameplayTagCountContainer::Notify_StackCountChange(const FGameplayTag& Tag)
+{	
+	// The purpose of this function is to let anyone listening on the EGameplayTagEventType::AnyCountChange event know that the 
+	// stack count of a GE that was backing this GE has changed. We do not update our internal map/count with this info, since that
+	// map only counts the number of GE/sources that are giving that tag.
+	FGameplayTagContainer TagAndParentsContainer = IGameplayTagsModule::Get().GetGameplayTagsManager().RequestGameplayTagParents(Tag);
+	for (auto CompleteTagIt = TagAndParentsContainer.CreateConstIterator(); CompleteTagIt; ++CompleteTagIt)
+	{
+		const FGameplayTag& CurTag = *CompleteTagIt;
+		FDelegateInfo* DelegateInfo = GameplayTagEventMap.Find(CurTag);
+		if (DelegateInfo)
+		{
+			int32 TagCount = GameplayTagCountMap.FindOrAdd(CurTag);
+			DelegateInfo->OnAnyChange.Broadcast(CurTag, TagCount);
+		}
+	}
+}
+
+FOnGameplayEffectTagCountChanged& FGameplayTagCountContainer::RegisterGameplayTagEvent(const FGameplayTag& Tag, EGameplayTagEventType::Type EventType)
 {
-	return GameplayTagEventMap.FindOrAdd(Tag);
+	FDelegateInfo& Info = GameplayTagEventMap.FindOrAdd(Tag);
+
+	if (EventType == EGameplayTagEventType::NewOrRemoved)
+	{
+		return Info.OnNewOrRemove;
+	}
+
+	return Info.OnAnyChange;
 }
 
 FOnGameplayEffectTagCountChanged& FGameplayTagCountContainer::RegisterGenericGameplayEvent()
@@ -388,14 +413,19 @@ void FGameplayTagCountContainer::UpdateTagMap_Internal(const FGameplayTag& Tag,
 		TagCount = FMath::Max(TagCount + CountDelta, 0);
 
 		// If a significant change (new addition or total removal) occurred, trigger related delegates
-		if (OldCount == 0 || TagCount == 0)
+		bool SignificantChange = (OldCount == 0 || TagCount == 0);
+		if (SignificantChange)
 		{
 			OnAnyTagChangeDelegate.Broadcast(CurTag, TagCount);
+		}
 
-			FOnGameplayEffectTagCountChanged* CountChangeDelegate = GameplayTagEventMap.Find(CurTag);
-			if (CountChangeDelegate)
+		FDelegateInfo* DelegateInfo = GameplayTagEventMap.Find(CurTag);
+		if (DelegateInfo)
+		{
+			DelegateInfo->OnAnyChange.Broadcast(CurTag, TagCount);
+			if (SignificantChange)
 			{
-				CountChangeDelegate->Broadcast(CurTag, TagCount);
+				DelegateInfo->OnNewOrRemove.Broadcast(CurTag, TagCount);
 			}
 		}
 	}
diff --git a/Engine/Source/Runtime/GameplayAbilities/Private/GameplayTagResponseTable.cpp b/Engine/Source/Runtime/GameplayAbilities/Private/GameplayTagResponseTable.cpp
new file mode 100644
--- /dev/null
+++ b/Engine/Source/Runtime/GameplayAbilities/Private/GameplayTagResponseTable.cpp
@@ -0,0 +1,76 @@
+// Copyright 1998-2015 Epic Games, Inc. All Rights Reserved.
+
+#include "AbilitySystemPrivatePCH.h"
+#include "GameplayTagResponseTable.h"
+
+// --------------------------------------------------------------------------------------------------------------------------------------------------------
+//
+//	UGameplayTagReponseTable
+//
+// --------------------------------------------------------------------------------------------------------------------------------------------------------
+
+
+UGameplayTagReponseTable::UGameplayTagReponseTable(const FObjectInitializer& ObjectInitializer)
+	: Super(ObjectInitializer)
+{
+	Query.OwningTagContainer = &QueryContainer;
+}
+
+void UGameplayTagReponseTable::RegisterResponseForEvents(UAbilitySystemComponent* ASC)
+{
+	if (RegisteredASCs.Contains(ASC))
+	{
+		return;
+	}
+
+	RegisteredASCs.Add(ASC, FGameplayTagResponseAppliedInfo() );
+
+	for (int32 idx=0; idx < Entries.Num(); ++idx)
+	{
+		const FGameplayTagResponseTableEntry& Entry = Entries[idx];
+		if (Entry.Positive.Tag.IsValid())
+		{
+			ASC->RegisterGameplayTagEvent( Entry.Positive.Tag, EGameplayTagEventType::AnyCountChange ).AddUObject(this, &UGameplayTagReponseTable::TagResponseEvent, ASC, idx);
+		}
+		if (Entry.Negative.Tag.IsValid())
+		{
+			ASC->RegisterGameplayTagEvent( Entry.Negative.Tag, EGameplayTagEventType::AnyCountChange ).AddUObject(this, &UGameplayTagReponseTable::TagResponseEvent, ASC, idx);
+		}
+	}
+}
+
+void UGameplayTagReponseTable::TagResponseEvent(const FGameplayTag Tag, int32 NewCount, UAbilitySystemComponent* ASC, int32 idx)
+{
+	if (!ensure(Entries.IsValidIndex(idx)))
+	{
+		return;
+	}
+
+	const FGameplayTagResponseTableEntry& Entry = Entries[idx];
+	int32 TotalCount = 0;
+
+	{
+		QUICK_SCOPE_CYCLE_COUNTER(ABILITY_TRT_CALC_COUNT);
+
+		TotalCount += ASC->GetAggregatedStackCount(MakeQuery(Entry.Positive.Tag));
+		TotalCount -= ASC->GetAggregatedStackCount(MakeQuery(Entry.Negative.Tag));
+	}
+
+	FGameplayTagResponseAppliedInfo& Info = RegisteredASCs.FindChecked(ASC);
+
+	if (TotalCount < 0)
+	{
+		Remove(ASC, Info.PositiveHandle);
+		AddOrUpdate(ASC, Entry.Negative.ResponseGameplayEffect, TotalCount, Info.NegativeHandle);
+	}
+	else if (TotalCount > 0)
+	{
+		Remove(ASC, Info.NegativeHandle);
+		AddOrUpdate(ASC, Entry.Positive.ResponseGameplayEffect, TotalCount, Info.PositiveHandle);
+	}
+	else if (TotalCount == 0)
+	{
+		Remove(ASC, Info.PositiveHandle);
+		Remove(ASC, Info.NegativeHandle);
+	}
+}
\ No newline at end of file
diff --git a/Engine/Source/Runtime/GameplayAbilities/Public/AbilitySystemComponent.h b/Engine/Source/Runtime/GameplayAbilities/Public/AbilitySystemComponent.h
--- a/Engine/Source/Runtime/GameplayAbilities/Public/AbilitySystemComponent.h
+++ b/Engine/Source/Runtime/GameplayAbilities/Public/AbilitySystemComponent.h
@@ -304,6 +304,9 @@ class GAMEPLAYABILITIES_API UAbilitySystemComponent : public UGameplayTasksCompo
 	UFUNCTION(BlueprintCallable, BlueprintPure, Category=GameplayEffects)
 	int32 GetGameplayEffectCount(TSubclassOf<UGameplayEffect> SourceGameplayEffect, UAbilitySystemComponent* OptionalInstigatorFilterComponent);
 
+	/** Returns the sum of StackCount of all gameplay effects that pass query */
+	int32 GetAggregatedStackCount(const FActiveGameplayEffectQuery& Query);
+
 	/** This only exists so it can be hooked up to a multicast delegate */
 	void RemoveActiveGameplayEffect_NoReturn(FActiveGameplayEffectHandle Handle, int32 StacksToRemove=-1)
 	{
@@ -320,6 +323,10 @@ class GAMEPLAYABILITIES_API UAbilitySystemComponent : public UGameplayTasksCompo
 
 	float GetGameplayEffectDuration() const;
 
+	/** Updates the level of an already applied gameplay effect. The intention is that this is 'seemless' and doesnt behave like removing/reapplying */
+	UFUNCTION(BlueprintCallable, BlueprintAuthorityOnly, Category = GameplayEffects)
+	void SetActiveGameplayEffectLevel(FActiveGameplayEffectHandle ActiveHandle, int32 NewLevel);
+
 	// Not happy with this interface but don't see a better way yet. How should outside code (UI, etc) ask things like 'how much is this gameplay effect modifying my damage by'
 	// (most likely we want to catch this on the backend - when damage is applied we can get a full dump/history of how the number got to where it is. But still we may need polling methods like below (how much would my damage be)
 	UFUNCTION(BlueprintCallable, Category = GameplayEffects)
@@ -405,9 +412,9 @@ class GAMEPLAYABILITIES_API UAbilitySystemComponent : public UGameplayTasksCompo
 	void RemoveLooseGameplayTag(const FGameplayTag& GameplayTag, int32 Count = 1);
 
 	void RemoveLooseGameplayTags(const FGameplayTagContainer& GameplayTag, int32 Count = 1);
-
+	
 	/** Allow events to be registered for specific gameplay tags being added or removed */
-	FOnGameplayEffectTagCountChanged& RegisterGameplayTagEvent(FGameplayTag Tag);
+	FOnGameplayEffectTagCountChanged& RegisterGameplayTagEvent(FGameplayTag Tag, EGameplayTagEventType::Type EventType=EGameplayTagEventType::NewOrRemoved);
 
 	/** Returns multicast delegate that is invoked whenever a tag is added or removed (but not if just count is increased. Only for 'new' and 'removed' events) */
 	FOnGameplayEffectTagCountChanged& RegisterGenericGameplayTagEvent();
@@ -1092,6 +1099,8 @@ protected:
 	
 	void UpdateTagMap(const FGameplayTagContainer& Container, int32 CountDelta);
 
+	void NotifyTagMap_StackCountChange(const FGameplayTagContainer& Container);
+
 	virtual void OnTagUpdated(const FGameplayTag& Tag, bool TagExists) {};
 	
 	// ---------------------------------------------
diff --git a/Engine/Source/Runtime/GameplayAbilities/Public/AbilitySystemGlobals.h b/Engine/Source/Runtime/GameplayAbilities/Public/AbilitySystemGlobals.h
--- a/Engine/Source/Runtime/GameplayAbilities/Public/AbilitySystemGlobals.h
+++ b/Engine/Source/Runtime/GameplayAbilities/Public/AbilitySystemGlobals.h
@@ -10,6 +10,7 @@ class UAbilitySystemComponent;
 class UCurveTable;
 class UDataTable;
 class UGameplayCueManager;
+class UGameplayTagReponseTable;
 
 struct FGameplayAbilityActorInfo;
 struct FGameplayEffectContext;
@@ -68,6 +69,9 @@ class GAMEPLAYABILITIES_API UAbilitySystemGlobals : public UObject
 	/** Returns the gameplay cue manager singleton object, creating if necessary */
 	virtual UGameplayCueManager* GetGameplayCueManager();
 
+	/** Returns the gameplay tag response object, creating if necessary */
+	UGameplayTagReponseTable* GetGameplayTagResponseTable();
+
 	/** Sets a default gameplay cue tag using the asset's name */
 	static void DeriveGameplayCueTagFromAssetName(FString AssetName, FGameplayTag& GameplayCueTag, FName& GameplayCueName);
 
@@ -194,6 +198,13 @@ private:
 	UPROPERTY(config)
 	TArray<FString>	GameplayCueNotifyPaths;
 
+	/** The class to instantiate as the GameplayTagResponseTable. */
+	UPROPERTY(config)
+	FStringAssetReference GameplayTagResponseTableName;
+
+	UPROPERTY()
+	UGameplayTagReponseTable* GameplayTagResponseTable;
+
 	/** Set to true if you want clients to try to predict gameplay effects done to targets. If false it will only predict self effects */
 	UPROPERTY(config)
 	bool	PredictTargetGameplayEffects;
diff --git a/Engine/Source/Runtime/GameplayAbilities/Public/GameplayEffect.h b/Engine/Source/Runtime/GameplayAbilities/Public/GameplayEffect.h
--- a/Engine/Source/Runtime/GameplayAbilities/Public/GameplayEffect.h
+++ b/Engine/Source/Runtime/GameplayAbilities/Public/GameplayEffect.h
@@ -1485,6 +1485,8 @@ struct GAMEPLAYABILITIES_API FActiveGameplayEffectsContainer : public FFastArray
 
 	float GetGameplayEffectMagnitude(FActiveGameplayEffectHandle Handle, FGameplayAttribute Attribute) const;
 
+	void SetActiveGameplayEffectLevel(FActiveGameplayEffectHandle ActiveHandle, int32 NewLevel);
+
 	void SetAttributeBaseValue(FGameplayAttribute Attribute, float NewBaseValue);
 
 	/** Actually applies given mod to the attribute */
diff --git a/Engine/Source/Runtime/GameplayAbilities/Public/GameplayEffectTypes.h b/Engine/Source/Runtime/GameplayAbilities/Public/GameplayEffectTypes.h
--- a/Engine/Source/Runtime/GameplayAbilities/Public/GameplayEffectTypes.h
+++ b/Engine/Source/Runtime/GameplayAbilities/Public/GameplayEffectTypes.h
@@ -151,6 +151,11 @@ struct GAMEPLAYABILITIES_API FActiveGameplayEffectHandle
 		return FString::Printf(TEXT("%d"), Handle);
 	}
 
+	void Invalidate()
+	{
+		Handle = INDEX_NONE;
+	}
+
 private:
 
 	UPROPERTY()
@@ -760,6 +765,20 @@ DECLARE_DELEGATE_RetVal_OneParam(FOnGameplayEffectTagCountChanged&, FRegisterGam
 
 // -----------------------------------------------------------
 
+
+UENUM(BlueprintType)
+namespace EGameplayTagEventType
+{
+	enum Type
+	{		
+		/** Event only happens when tag is branch new or removed */
+		NewOrRemoved,
+
+		/** Event happens any time tag "count" changes */
+		AnyCountChange		
+	};
+}
+
 /**
  * Struct that tracks the number/count of tag applications within it. Explicitly tracks the tags added or removed,
  * while simultaneously tracking the count of parent tags as well. Events/delegates are fired whenever the tag counts
@@ -816,6 +835,12 @@ struct GAMEPLAYABILITIES_API FGameplayTagCountContainer
 	 */
 	void UpdateTagCount(const FGameplayTag& Tag, int32 CountDelta);
 
+	/**
+	 *	Broadcasts the AnyChange event for this tag. This is called when the stack count of the backing gameplay effect change.
+	 *	It is up to the receiver of the broadcasted delegate to decide what to do with this.
+	 */
+	void Notify_StackCountChange(const FGameplayTag& Tag);
+
 	/**
 	 * Return delegate that can be bound to for when the specific tag's count changes to or off of zero
 	 *
@@ -823,7 +848,7 @@ struct GAMEPLAYABILITIES_API FGameplayTagCountContainer
 	 * 
 	 * @return Delegate for when the specified tag's count changes to or off of zero
 	 */
-	FOnGameplayEffectTagCountChanged& RegisterGameplayTagEvent(const FGameplayTag& Tag);
+	FOnGameplayEffectTagCountChanged& RegisterGameplayTagEvent(const FGameplayTag& Tag, EGameplayTagEventType::Type EventType=EGameplayTagEventType::NewOrRemoved);
 	
 	/**
 	 * Return delegate that can be bound to for when the any tag's count changes to or off of zero
@@ -837,8 +862,14 @@ struct GAMEPLAYABILITIES_API FGameplayTagCountContainer
 
 private:
 
+	struct FDelegateInfo
+	{
+		FOnGameplayEffectTagCountChanged	OnNewOrRemove;
+		FOnGameplayEffectTagCountChanged	OnAnyChange;
+	};
+
 	/** Map of tag to delegate that will be fired when the count for the key tag changes to or away from zero */
-	TMap<FGameplayTag, FOnGameplayEffectTagCountChanged> GameplayTagEventMap;
+	TMap<FGameplayTag, FDelegateInfo> GameplayTagEventMap;
 
 	/** Map of tag to active count of that tag */
 	TMap<FGameplayTag, int32> GameplayTagCountMap;
diff --git a/Engine/Source/Runtime/GameplayAbilities/Public/GameplayTagResponseTable.h b/Engine/Source/Runtime/GameplayAbilities/Public/GameplayTagResponseTable.h
new file mode 100644
--- /dev/null
+++ b/Engine/Source/Runtime/GameplayAbilities/Public/GameplayTagResponseTable.h
@@ -0,0 +1,107 @@
+// Copyright 1998-2015 Epic Games, Inc. All Rights Reserved.
+
+#pragma once
+
+#include "GameplayTags.h"
+#include "GameplayTagResponseTable.generated.h"
+
+class UGameplayEffect;
+
+USTRUCT()
+struct FGameplayTagReponsePair
+{
+	GENERATED_USTRUCT_BODY()
+
+	/** Tag that triggers this response */
+	UPROPERTY(EditAnywhere, Category="Response")
+	FGameplayTag	Tag;
+
+	/** The GameplayEffect to apply in reponse to the tag */
+	UPROPERTY(EditAnywhere, Category="Response")
+	TSubclassOf<UGameplayEffect> ResponseGameplayEffect;
+};
+
+USTRUCT()
+struct FGameplayTagResponseTableEntry
+{
+	GENERATED_USTRUCT_BODY()
+
+	/** Tags that count as "positive" toward to final response count. If the overall count is positive, this ResponseGameplayEffect is applied. */
+	UPROPERTY(EditAnywhere, Category="Response")
+	FGameplayTagReponsePair		Positive;
+
+	/** Tags that count as "negative" toward to final response count. If the overall count is negative, this ResponseGameplayEffect is applied. */
+	UPROPERTY(EditAnywhere, Category="Response")
+	FGameplayTagReponsePair		Negative;
+};
+
+/**
+ *	A data driven table for applying gameplay effects based on tag count. This allows designers to map a 
+ *	"tag count" -> "response Gameplay Effect" relationship.
+ *	
+ *	For example, "for every stack of "Status.Haste" I get 1 level of GE_Response_Haste. This class facilitates
+ *	building this table and automatically registering and responding to tag events on the ability system component.
+ */
+UCLASS()
+class GAMEPLAYABILITIES_API UGameplayTagReponseTable : public UDataAsset
+{
+	GENERATED_UCLASS_BODY()
+
+	UPROPERTY(EditAnywhere, Category="Response")
+	TArray<FGameplayTagResponseTableEntry>	Entries;
+
+	/** Registers for tag events for the given ability system component. Note this will happen to every spawned ASC, we may want to allow games
+	 *	to limit what classe this is called on, or potentially build into the table class restrictions for each response entry.
+	 */
+	void RegisterResponseForEvents(UAbilitySystemComponent* ASC);
+
+protected:
+
+	UFUNCTION()
+	void TagResponseEvent(const FGameplayTag Tag, int32 NewCount, UAbilitySystemComponent* ASC, int32 idx);
+	
+	/** Temporary structs to avoid extra heap allocations every time we recalculate tag count */
+	FActiveGameplayEffectQuery Query;
+	FGameplayTagContainer QueryContainer;
+
+	FActiveGameplayEffectQuery& MakeQuery(const FGameplayTag& Tag)
+	{
+		QueryContainer.RemoveAllTags(QueryContainer.Num());
+		QueryContainer.AddTagFast(Tag);
+		return Query;
+	}
+
+	// ----------------------------------------------------
+
+	struct FGameplayTagResponseAppliedInfo
+	{
+		FActiveGameplayEffectHandle PositiveHandle;
+		FActiveGameplayEffectHandle NegativeHandle;
+	};
+
+	TMap< TWeakObjectPtr<UAbilitySystemComponent>, FGameplayTagResponseAppliedInfo > RegisteredASCs;
+
+	void Remove(UAbilitySystemComponent* ASC, FActiveGameplayEffectHandle& Handle)
+	{
+		if (Handle.IsValid())
+		{
+			ASC->RemoveActiveGameplayEffect(Handle);
+			Handle.Invalidate();
+		}
+	}
+
+	void AddOrUpdate(UAbilitySystemComponent* ASC, const TSubclassOf<UGameplayEffect>& ResponseGameplayEffect, int32 TotalCount, FActiveGameplayEffectHandle& Handle)
+	{
+		if (ResponseGameplayEffect)
+		{
+			if (Handle.IsValid())
+			{
+				ASC->SetActiveGameplayEffectLevel(Handle, TotalCount);
+			}
+			else
+			{
+				Handle = ASC->ApplyGameplayEffectToSelf(Cast<UGameplayEffect>(ResponseGameplayEffect->ClassDefaultObject), TotalCount, ASC->GetEffectContext());
+			}
+		}
+	}
+};
\ No newline at end of file
