Fixed several non unity compile errors

diff --git a/Engine/Source/Editor/AnimGraph/Private/AnimGraphNode_ObserveBone.cpp b/Engine/Source/Editor/AnimGraph/Private/AnimGraphNode_ObserveBone.cpp
--- a/Engine/Source/Editor/AnimGraph/Private/AnimGraphNode_ObserveBone.cpp
+++ b/Engine/Source/Editor/AnimGraph/Private/AnimGraphNode_ObserveBone.cpp
@@ -6,6 +6,8 @@
 #include "SGraphNode.h"
 #include "KismetNodeInfoContext.h"
 #include "KismetDebugUtilities.h"
+#include "DetailLayoutBuilder.h"
+#include "PropertyHandle.h"
 
 #define LOCTEXT_NAMESPACE "ObserveBone"
 
diff --git a/Engine/Source/Runtime/GameplayAbilities/Private/AbilitySystemGlobals.cpp b/Engine/Source/Runtime/GameplayAbilities/Private/AbilitySystemGlobals.cpp
--- a/Engine/Source/Runtime/GameplayAbilities/Private/AbilitySystemGlobals.cpp
+++ b/Engine/Source/Runtime/GameplayAbilities/Private/AbilitySystemGlobals.cpp
@@ -8,6 +8,7 @@
 #include "GameplayTagsModule.h"
 #include "GameplayCueInterface.h"
 #include "GameplayCueManager.h"
+#include "GameplayTagResponseTable.h"
 
 
 UAbilitySystemGlobals::UAbilitySystemGlobals(const FObjectInitializer& ObjectInitializer)
diff --git a/Engine/Source/Runtime/GameplayAbilities/Public/GameplayTagResponseTable.h b/Engine/Source/Runtime/GameplayAbilities/Public/GameplayTagResponseTable.h
--- a/Engine/Source/Runtime/GameplayAbilities/Public/GameplayTagResponseTable.h
+++ b/Engine/Source/Runtime/GameplayAbilities/Public/GameplayTagResponseTable.h
@@ -3,8 +3,10 @@
 #pragma once
 
 #include "GameplayTags.h"
+#include "AbilitySystemComponent.h"
 #include "GameplayTagResponseTable.generated.h"
 
+
 class UGameplayEffect;
 
 USTRUCT()
