Fixed deprecation warnings in AbilitySystemComponent.

diff --git a/Engine/Source/Runtime/GameplayAbilities/Private/AbilitySystemComponent.cpp b/Engine/Source/Runtime/GameplayAbilities/Private/AbilitySystemComponent.cpp
--- a/Engine/Source/Runtime/GameplayAbilities/Private/AbilitySystemComponent.cpp
+++ b/Engine/Source/Runtime/GameplayAbilities/Private/AbilitySystemComponent.cpp
@@ -278,6 +278,13 @@ int32 UAbilitySystemComponent::GetGameplayEffectCount(TSubclassOf<UGameplayEffec
 }
 
 int32 UAbilitySystemComponent::GetAggregatedStackCount(const FActiveGameplayEffectQuery& Query)
+{
+	PRAGMA_DISABLE_DEPRECATION_WARNINGS
+	return ActiveGameplayEffects.GetActiveEffectCount(Query);
+	PRAGMA_ENABLE_DEPRECATION_WARNINGS
+}
+
+int32 UAbilitySystemComponent::GetAggregatedStackCount(const FGameplayEffectQuery& Query)
 {
 	return ActiveGameplayEffects.GetActiveEffectCount(Query);
 }
diff --git a/Engine/Source/Runtime/GameplayAbilities/Private/GameplayTagResponseTable.cpp b/Engine/Source/Runtime/GameplayAbilities/Private/GameplayTagResponseTable.cpp
--- a/Engine/Source/Runtime/GameplayAbilities/Private/GameplayTagResponseTable.cpp
+++ b/Engine/Source/Runtime/GameplayAbilities/Private/GameplayTagResponseTable.cpp
@@ -13,7 +13,6 @@
 UGameplayTagReponseTable::UGameplayTagReponseTable(const FObjectInitializer& ObjectInitializer)
 	: Super(ObjectInitializer)
 {
-	Query.OwningTagContainer = &QueryContainer;
 }
 
 void UGameplayTagReponseTable::RegisterResponseForEvents(UAbilitySystemComponent* ASC)
diff --git a/Engine/Source/Runtime/GameplayAbilities/Public/AbilitySystemComponent.h b/Engine/Source/Runtime/GameplayAbilities/Public/AbilitySystemComponent.h
--- a/Engine/Source/Runtime/GameplayAbilities/Public/AbilitySystemComponent.h
+++ b/Engine/Source/Runtime/GameplayAbilities/Public/AbilitySystemComponent.h
@@ -305,7 +305,9 @@ class GAMEPLAYABILITIES_API UAbilitySystemComponent : public UGameplayTasksCompo
 	int32 GetGameplayEffectCount(TSubclassOf<UGameplayEffect> SourceGameplayEffect, UAbilitySystemComponent* OptionalInstigatorFilterComponent);
 
 	/** Returns the sum of StackCount of all gameplay effects that pass query */
+	DEPRECATED(4.9, "FActiveGameplayEffectQuery is deprecated, use version that takes FGameplayEffectQuery")
 	int32 GetAggregatedStackCount(const FActiveGameplayEffectQuery& Query);
+	int32 GetAggregatedStackCount(const FGameplayEffectQuery& Query);
 
 	/** This only exists so it can be hooked up to a multicast delegate */
 	void RemoveActiveGameplayEffect_NoReturn(FActiveGameplayEffectHandle Handle, int32 StacksToRemove=-1)
diff --git a/Engine/Source/Runtime/GameplayAbilities/Public/GameplayTagResponseTable.h b/Engine/Source/Runtime/GameplayAbilities/Public/GameplayTagResponseTable.h
--- a/Engine/Source/Runtime/GameplayAbilities/Public/GameplayTagResponseTable.h
+++ b/Engine/Source/Runtime/GameplayAbilities/Public/GameplayTagResponseTable.h
@@ -61,18 +61,13 @@ protected:
 
 	UFUNCTION()
 	void TagResponseEvent(const FGameplayTag Tag, int32 NewCount, UAbilitySystemComponent* ASC, int32 idx);
-	
-	/** Temporary structs to avoid extra heap allocations every time we recalculate tag count */
-	FActiveGameplayEffectQuery Query;
-	FGameplayTagContainer QueryContainer;
 
-	FActiveGameplayEffectQuery& MakeQuery(const FGameplayTag& Tag)
+	FGameplayEffectQuery MakeQuery(const FGameplayTag& Tag)
 	{
-		QueryContainer.RemoveAllTags(QueryContainer.Num());
-		QueryContainer.AddTagFast(Tag);
-		return Query;
+		return FGameplayEffectQuery::MakeQuery_MatchAnyOwningTags(FGameplayTagContainer(Tag));
 	}
 
+
 	// ----------------------------------------------------
 
 	struct FGameplayTagResponseAppliedInfo
