GameplayEffect matching now handles empty tag container matches more naturally/correctly.

diff --git a/Engine/Source/Runtime/GameplayAbilities/Private/GameplayEffect.cpp b/Engine/Source/Runtime/GameplayAbilities/Private/GameplayEffect.cpp
--- a/Engine/Source/Runtime/GameplayAbilities/Private/GameplayEffect.cpp
+++ b/Engine/Source/Runtime/GameplayAbilities/Private/GameplayEffect.cpp
@@ -3261,7 +3261,7 @@ bool FGameplayEffectQuery::Matches(const FActiveGameplayEffect& Effect) const
 	if (OwningTagQuery.IsEmpty() == false)
 	{
 		FGameplayTagContainer const& CombinedTags = Effect.Spec.Def->InheritableOwnedTagsContainer.CombinedTags;
-		if ((CombinedTags.Num() > 0) && (OwningTagQuery.Matches(CombinedTags) == false))		// note: matches on empty
+		if (OwningTagQuery.Matches(CombinedTags) == false)
 		{
 			if (OwningTagQuery.Matches(Effect.Spec.DynamicGrantedTags) == false)
 			{
@@ -3277,7 +3277,7 @@ bool FGameplayEffectQuery::Matches(const FActiveGameplayEffect& Effect) const
 	if (EffectTagQuery.IsEmpty() == false)
 	{
 		FGameplayTagContainer const& CombinedTags = Effect.Spec.Def->InheritableOwnedTagsContainer.CombinedTags;
-		if ((CombinedTags.Num() > 0) && (OwningTagQuery.Matches(CombinedTags) == false))		// note: matches on empty
+		if (OwningTagQuery.Matches(CombinedTags) == false)
 		{
 			if (OwningTagQuery.Matches(Effect.Spec.DynamicAssetTags) == false)
 			{
