Changes to Ability logs/vlogs, to limit overhead when  logging is disabled.

diff --git a/Engine/Source/Runtime/GameplayAbilities/Private/AbilitySystemLog.cpp b/Engine/Source/Runtime/GameplayAbilities/Private/AbilitySystemLog.cpp
--- a/Engine/Source/Runtime/GameplayAbilities/Private/AbilitySystemLog.cpp
+++ b/Engine/Source/Runtime/GameplayAbilities/Private/AbilitySystemLog.cpp
@@ -23,7 +23,7 @@ AbilitySystemLog * GetInstance()
 	return &Instance;
 }
 
-FString AbilitySystemLog::Log(ELogVerbosity::Type Verbosity, FString Log)
+FString AbilitySystemLog::Log(ELogVerbosity::Type Verbosity, const FString& LogStr)
 {
 #if !NO_LOGGING
 	if (!LogAbilitySystem.IsSuppressed(Verbosity))
@@ -49,11 +49,11 @@ FString AbilitySystemLog::Log(ELogVerbosity::Type Verbosity, FString Log)
 		}
 
 		FString IndentStr = FString::Printf(TEXT("%*s"), Instance->Indent, TEXT(""));
-		return IndentStr + Log;
+		return IndentStr + LogStr;
 	}
 #endif
 	
-	return Log;
+	return LogStr;
 }
 
 void AbilitySystemLog::PushScope(AbilitySystemLogScope * Scope)
diff --git a/Engine/Source/Runtime/GameplayAbilities/Public/AbilitySystemLog.h b/Engine/Source/Runtime/GameplayAbilities/Public/AbilitySystemLog.h
--- a/Engine/Source/Runtime/GameplayAbilities/Public/AbilitySystemLog.h
+++ b/Engine/Source/Runtime/GameplayAbilities/Public/AbilitySystemLog.h
@@ -34,15 +34,17 @@ GAMEPLAYABILITIES_API DECLARE_LOG_CATEGORY_EXTERN(VLogAbilitySystem, Warning, Al
 
 #define ABILITY_LOG(Verbosity, Format, ...) \
 { \
-	FString Str = AbilitySystemLog::Log(ELogVerbosity::Verbosity, FString::Printf(Format, ##__VA_ARGS__)); \
-	UE_LOG(LogAbilitySystem, Verbosity, TEXT("%s"), *Str); \
+	UE_LOG(LogAbilitySystem, Verbosity, TEXT("%s"), *AbilitySystemLog::Log(ELogVerbosity::Verbosity, FString::Printf(Format, ##__VA_ARGS__))); \
 }
 
 #define ABILITY_VLOG(Actor, Verbosity, Format, ...) \
 { \
-	FString Str = AbilitySystemLog::Log(ELogVerbosity::Verbosity, FString::Printf(Format, ##__VA_ARGS__)); \
-	UE_LOG(LogAbilitySystem, Verbosity, TEXT("%s"), *Str); \
-	UE_VLOG(Actor, VLogAbilitySystem, Verbosity, TEXT("%s"), *Str); \
+	if(FVisualLogger::IsRecording() || LogAbilitySystem.IsSuppressed(ELogVerbosity::Verbosity) == false ) \
+	{ \
+		const FString Str = AbilitySystemLog::Log(ELogVerbosity::Verbosity, FString::Printf(Format, ##__VA_ARGS__)); \
+		UE_LOG(LogAbilitySystem, Verbosity, TEXT("%s"), *Str); \
+		UE_VLOG(Actor, VLogAbilitySystem, Verbosity, TEXT("%s"), *Str); \
+	} \
 }
 
 #define ABILITY_LOG_SCOPE( Format, ... ) AbilitySystemLogScope PREPROCESSOR_JOIN(LogScope,__LINE__)( FString::Printf(Format, ##__VA_ARGS__));
@@ -53,13 +55,16 @@ GAMEPLAYABILITIES_API DECLARE_LOG_CATEGORY_EXTERN(VLogAbilitySystem, Warning, Al
 
 #define ABILITY_VLOG_ATTRIBUTE_GRAPH(Actor, Verbosity, AttributeName, OldValue, NewValue) \
 { \
-	const FName GraphName("Attribute Graph"); \
-	float CurrentTime = Actor->GetWorld() ? Actor->GetWorld()->GetTimeSeconds() : 0.f; \
-	FVector2D OldPt(CurrentTime, OldValue); \
-	FVector2D NewPt(CurrentTime, NewValue); \
-	FName LineName(*AttributeName); \
-	UE_VLOG_HISTOGRAM(Actor, VLogAbilitySystem, Log, GraphName, LineName, OldPt); \
-	UE_VLOG_HISTOGRAM(OwnerActor, VLogAbilitySystem, Log, GraphName, LineName, NewPt); \
+	if( FVisualLogger::IsRecording() ) \
+	{ \
+		static const FName GraphName("Attribute Graph"); \
+		const float CurrentTime = Actor->GetWorld() ? Actor->GetWorld()->GetTimeSeconds() : 0.f; \
+		const FVector2D OldPt(CurrentTime, OldValue); \
+		const FVector2D NewPt(CurrentTime, NewValue); \
+		const FName LineName(*AttributeName); \
+		UE_VLOG_HISTOGRAM(Actor, VLogAbilitySystem, Log, GraphName, LineName, OldPt); \
+		UE_VLOG_HISTOGRAM(OwnerActor, VLogAbilitySystem, Log, GraphName, LineName, NewPt); \
+	} \
 }
 
 #else
@@ -75,7 +80,7 @@ struct AbilitySystemLogScope
 		Init();
 	}
 
-	AbilitySystemLogScope(FString InScopeName) :
+	AbilitySystemLogScope(const FString& InScopeName) :
 		ScopeName(InScopeName)
 	{
 		Init();
@@ -101,7 +106,7 @@ public:
 		NeedNewLine = false;
 	}
 
-	static FString Log(ELogVerbosity::Type, FString Str);
+	static FString Log(ELogVerbosity::Type, const FString& Str);
 
 	static void PushScope(AbilitySystemLogScope* Scope);
 
