Fix for win32 monolitic build issue.

diff --git a/Engine/Source/Runtime/GameplayAbilities/Private/AbilitySystemPrivatePCH.h b/Engine/Source/Runtime/GameplayAbilities/Private/AbilitySystemPrivatePCH.h
--- a/Engine/Source/Runtime/GameplayAbilities/Private/AbilitySystemPrivatePCH.h
+++ b/Engine/Source/Runtime/GameplayAbilities/Private/AbilitySystemPrivatePCH.h
@@ -21,3 +21,4 @@
 #include "AbilitySystemLog.h"
 
 #include "GameplayTasksComponent.h"
+#include "VisualLogger.h"
