Fixed monolitic build issues in Ability related to visual logger.

diff --git a/Engine/Source/Runtime/GameplayAbilities/Private/AbilitySystemPrivatePCH.h b/Engine/Source/Runtime/GameplayAbilities/Private/AbilitySystemPrivatePCH.h
--- a/Engine/Source/Runtime/GameplayAbilities/Private/AbilitySystemPrivatePCH.h
+++ b/Engine/Source/Runtime/GameplayAbilities/Private/AbilitySystemPrivatePCH.h
@@ -21,4 +21,3 @@
 #include "AbilitySystemLog.h"
 
 #include "GameplayTasksComponent.h"
-#include "VisualLogger.h"
diff --git a/Engine/Source/Runtime/GameplayAbilities/Public/AbilitySystemLog.h b/Engine/Source/Runtime/GameplayAbilities/Public/AbilitySystemLog.h
--- a/Engine/Source/Runtime/GameplayAbilities/Public/AbilitySystemLog.h
+++ b/Engine/Source/Runtime/GameplayAbilities/Public/AbilitySystemLog.h
@@ -37,15 +37,22 @@ GAMEPLAYABILITIES_API DECLARE_LOG_CATEGORY_EXTERN(VLogAbilitySystem, Warning, Al
 	UE_LOG(LogAbilitySystem, Verbosity, TEXT("%s"), *AbilitySystemLog::Log(ELogVerbosity::Verbosity, FString::Printf(Format, ##__VA_ARGS__))); \
 }
 
-#define ABILITY_VLOG(Actor, Verbosity, Format, ...) \
-{ \
-	if(FVisualLogger::IsRecording() || LogAbilitySystem.IsSuppressed(ELogVerbosity::Verbosity) == false ) \
+#if ENABLE_VISUAL_LOG
+#	define ABILITY_VLOG(Actor, Verbosity, Format, ...) \
 	{ \
-		const FString Str = AbilitySystemLog::Log(ELogVerbosity::Verbosity, FString::Printf(Format, ##__VA_ARGS__)); \
-		UE_LOG(LogAbilitySystem, Verbosity, TEXT("%s"), *Str); \
-		UE_VLOG(Actor, VLogAbilitySystem, Verbosity, TEXT("%s"), *Str); \
-	} \
-}
+		if(FVisualLogger::IsRecording() || LogAbilitySystem.IsSuppressed(ELogVerbosity::Verbosity) == false ) \
+		{ \
+			const FString Str = AbilitySystemLog::Log(ELogVerbosity::Verbosity, FString::Printf(Format, ##__VA_ARGS__)); \
+			UE_LOG(LogAbilitySystem, Verbosity, TEXT("%s"), *Str); \
+			UE_VLOG(Actor, VLogAbilitySystem, Verbosity, TEXT("%s"), *Str); \
+		} \
+	}
+#else //ENABLE_VISUAL_LOG
+#	define ABILITY_VLOG(Actor, Verbosity, Format, ...) \
+	{ \
+		UE_LOG(LogAbilitySystem, Verbosity, TEXT("%s"), *AbilitySystemLog::Log(ELogVerbosity::Verbosity, FString::Printf(Format, ##__VA_ARGS__))); \
+	}
+#endif  //ENABLE_VISUAL_LOG
 
 #define ABILITY_LOG_SCOPE( Format, ... ) AbilitySystemLogScope PREPROCESSOR_JOIN(LogScope,__LINE__)( FString::Printf(Format, ##__VA_ARGS__));
 
