Fix simulated ability tasks not having AbilitySystemComponent member variable set during simulated initialization

diff --git a/Engine/Source/Runtime/GameplayAbilities/Private/Abilities/Tasks/AbilityTask.cpp b/Engine/Source/Runtime/GameplayAbilities/Private/Abilities/Tasks/AbilityTask.cpp
--- a/Engine/Source/Runtime/GameplayAbilities/Private/Abilities/Tasks/AbilityTask.cpp
+++ b/Engine/Source/Runtime/GameplayAbilities/Private/Abilities/Tasks/AbilityTask.cpp
@@ -20,6 +20,14 @@ void UAbilityTask::SetAbilitySystemComponent(UAbilitySystemComponent* InAbilityS
 	AbilitySystemComponent = InAbilitySystemComponent;
 }
 
+void UAbilityTask::InitSimulatedTask(UGameplayTasksComponent& InGameplayTasksComponent)
+{
+	UGameplayTask::InitSimulatedTask(InGameplayTasksComponent);
+
+	UAbilitySystemComponent* AbilitySystemComponent = Cast<UAbilitySystemComponent>(TasksComponent.Get());
+	SetAbilitySystemComponent(AbilitySystemComponent);
+}
+
 FPredictionKey UAbilityTask::GetActivationPredictionKey() const
 {
 	UGameplayAbility* MyAbility = Ability.Get();
diff --git a/Engine/Source/Runtime/GameplayAbilities/Public/Abilities/Tasks/AbilityTask.h b/Engine/Source/Runtime/GameplayAbilities/Public/Abilities/Tasks/AbilityTask.h
--- a/Engine/Source/Runtime/GameplayAbilities/Public/Abilities/Tasks/AbilityTask.h
+++ b/Engine/Source/Runtime/GameplayAbilities/Public/Abilities/Tasks/AbilityTask.h
@@ -86,6 +86,8 @@ class GAMEPLAYABILITIES_API UAbilityTask : public UGameplayTask
 	/** Returns ActivationPredictionKey of owning ability */
 	FPredictionKey GetActivationPredictionKey() const;
 
+	virtual void InitSimulatedTask(UGameplayTasksComponent& InGameplayTasksComponent) override;
+
 	/** Helper function for instantiating and initializing a new task */
 	template <class T>
 	static T* NewAbilityTask(UObject* WorldContextObject, FName InstanceName = FName())
