[Build watcher]
Fix shadowed variable warning.

diff --git a/Engine/Source/Runtime/GameplayAbilities/Private/Abilities/Tasks/AbilityTask.cpp b/Engine/Source/Runtime/GameplayAbilities/Private/Abilities/Tasks/AbilityTask.cpp
--- a/Engine/Source/Runtime/GameplayAbilities/Private/Abilities/Tasks/AbilityTask.cpp
+++ b/Engine/Source/Runtime/GameplayAbilities/Private/Abilities/Tasks/AbilityTask.cpp
@@ -24,8 +24,7 @@ void UAbilityTask::InitSimulatedTask(UGameplayTasksComponent& InGameplayTasksCom
 {
 	UGameplayTask::InitSimulatedTask(InGameplayTasksComponent);
 
-	UAbilitySystemComponent* AbilitySystemComponent = Cast<UAbilitySystemComponent>(TasksComponent.Get());
-	SetAbilitySystemComponent(AbilitySystemComponent);
+	SetAbilitySystemComponent(Cast<UAbilitySystemComponent>(TasksComponent.Get()));
 }
 
 FPredictionKey UAbilityTask::GetActivationPredictionKey() const
