Fix GameplayCueManager handling GCs when world is tearing down (especially during PIE exit) - safety check

diff --git a/Engine/Source/Runtime/GameplayAbilities/Private/GameplayCueManager.cpp b/Engine/Source/Runtime/GameplayAbilities/Private/GameplayCueManager.cpp
--- a/Engine/Source/Runtime/GameplayAbilities/Private/GameplayCueManager.cpp
+++ b/Engine/Source/Runtime/GameplayAbilities/Private/GameplayCueManager.cpp
@@ -108,6 +108,12 @@ void UGameplayCueManager::HandleGameplayCue(AActor* TargetActor, FGameplayTag Ga
 
 	CurrentWorld = TargetActor->GetWorld();
 
+	// Don't handle gameplay cues when world is tearing down
+	if (!GetWorld() || GetWorld()->bIsTearingDown)
+	{
+		return;
+	}
+
 	// Give the global set a chance
 	check(GlobalCueSet);
 	if (bAcceptsCue)
