Integration from development branch:

CL 2651151
Fix GameplayCues being invoked from ActiveGameplayEffectsContainer::NetDeltaSerialize when there are unmapped objects, which can cause GCs to be invoked with incomplete information in their EffectContext.
- Fixes Phase link GC not having SourceObject info to link from volume to caster

CL 2649496
Fix OptionalExternalTarget for WaitGameplayTag* tasks

CL 2644908
AbilityTask_WaitGameplayEvent. Allows abilities to explicitly listen for Gameplay Events (previously could only be activated from them).

CL 2644751
Tracking cases where AGameplayAbilityTargetActor fails to remove delegates from the ASC it applied them to - hoping to track down cause of #OR-5548 crash (that was fixed with nullptr checks)

CL 2643962
Fix occasional crash in AGameplayAbilityTargetActor - added nullptr checks, can't assume we have a valid OwningAbility pointer. #OR-5548

CL 2643797
GameplayEffect GrantedAbilities respect Effect inhibition (ongoing tag requirements) - abilities are applied/removed/cancelled based on whether the ongoing tags are met or not.
Previously, if application tags passed, target would have the granted ability for full duration of effect application, even if the ongoing tags weren't met
Fixes Rampage's E (which applies Juggernaut tag) not granting CC immunity after already being affected by the CC effects #OR-4019

CL 2643360
Fix ensures on exiting PIE - trying to spawn new instanced cue during level cleanup #OR-5521
Level being cleaned up, TravelMode triggering a GC removal that tries to spawn an instanced cue actor to handle the cue
ASC now are marked as unregistered during the unregistration process - previously they would DestroyActiveState() and then mark themselves as unregistered. There was nothing during DestroyActiveState() that allowed any system to know that they were in the process of being unregistered.
GameplayCues are now not handled by OrionChars that are pending kill or have ASCs that are pending kill or unregistered (or do not have ASCs)

CL 2642415
GetModfiedAttributeMagnitude - support for getting amount attribute changes from an applied GE spec

CL 2642414
Add support for periodic effects to WaitGameplayEffectApplied/Self

CL 2642413
Added delegate for periodic effects executing on target/self
Move PruneModifierAttributes to the constructor of FGameplayEffectSpecForRPC. This is to allow server code to look at the modified attribute list even if the GC doesn't need it.
Remove ApplyStackingLogicPostApplyAsSource function since it wasn't doing anything.

CL 2641147
Added DestroyActiveState to ASC so unregistering the component while in-flight abilities are active do not crash.
Added option for reusing player states when they become inactive instead of duplicated them.
Fixed crash for returning players during JIP.

CL 2637792
Temp fix for simulated ability tasks not having valid AbilitySystemComponents and crashing

CL 2637461
Orion game code fixes to go along with Integration from UE4-Stable @ CL 2634937. This includes all 4.9 revisions up to CL 2633546, which includes all UE4 (main) revisions up to CL 2613150. Note: many fixes were spot integrated after CL 2613150.

CL 2637431
Integration from UE4-Stable @ CL 2634937. This includes all 4.9 revisions up to CL 2633546, which includes all UE4 (main) revisions up to CL 2613150. Note: many fixes were spot integrated after CL 2613150.

CL 2635403
Add support for AbilitySystemComponent end debug display to flood to new columns to avoid cases where there's too much info to fit on one screen

CL 2635282
Added title text to AbilitySystemComponent end debug key cycling - contains category of debug information along with avatar and owner info and their respective roles. Saves brain cycles on figuring out what you're looking at (and if you're on local player vs simulated, etc.)

CL 2634858
Ability system debug component draw debug
Added view "Attributes" to only show attributes
Don't draw attributes in "aggregator view" if they have no active aggregators (inhibited still draw)
Show level on GEs with level > 1

diff --git a/Engine/Source/Runtime/GameplayAbilities/GameplayAbilities.Build.cs b/Engine/Source/Runtime/GameplayAbilities/GameplayAbilities.Build.cs
--- a/Engine/Source/Runtime/GameplayAbilities/GameplayAbilities.Build.cs
+++ b/Engine/Source/Runtime/GameplayAbilities/GameplayAbilities.Build.cs
@@ -35,7 +35,7 @@ namespace UnrealBuildTool.Rules
 			PrivateDependencyModuleNames.AddRange(
 				new string[]
 				{
-					// ... add private dependencies that you statically link with here ...
+
 				}
 				);
 
@@ -50,6 +50,7 @@ namespace UnrealBuildTool.Rules
 			{
 				PrivateDependencyModuleNames.Add("UnrealEd");
                 PrivateDependencyModuleNames.Add("GameplayTagsEditor");
+                PrivateDependencyModuleNames.Add("Slate");
 			}
 		}
 	}
diff --git a/Engine/Source/Runtime/GameplayAbilities/Private/Abilities/GameplayAbility.cpp b/Engine/Source/Runtime/GameplayAbilities/Private/Abilities/GameplayAbility.cpp
--- a/Engine/Source/Runtime/GameplayAbilities/Private/Abilities/GameplayAbility.cpp
+++ b/Engine/Source/Runtime/GameplayAbilities/Private/Abilities/GameplayAbility.cpp
@@ -148,7 +148,8 @@ bool UGameplayAbility::IsSupportedForNetworking() const
 
 bool UGameplayAbility::DoesAbilitySatisfyTagRequirements(const UAbilitySystemComponent& AbilitySystemComponent, const FGameplayTagContainer* SourceTags, const FGameplayTagContainer* TargetTags, OUT FGameplayTagContainer* OptionalRelevantTags) const
 {
-	bool bBlocked = false, bMissing = false;
+	bool bBlocked = false;
+	bool bMissing = false;
 
 	const FGameplayTag& BlockedTag = UAbilitySystemGlobals::Get().ActivateFailTagsBlockedTag;
 	const FGameplayTag& MissingTag = UAbilitySystemGlobals::Get().ActivateFailTagsMissingTag;
@@ -746,16 +747,18 @@ float UGameplayAbility::GetCooldownTimeRemaining(const FGameplayAbilityActorInfo
 {
 	SCOPE_CYCLE_COUNTER(STAT_GameplayAbilityGetCooldownTimeRemaining);
 
-	check(ActorInfo->AbilitySystemComponent.IsValid());
-	const FGameplayTagContainer* CooldownTags = GetCooldownTags();
-	if (CooldownTags && CooldownTags->Num() > 0)
+	if (ActorInfo->AbilitySystemComponent.IsValid())
 	{
-		FGameplayEffectQuery const Query = FGameplayEffectQuery::MakeQuery_MatchAnyOwningTags(*CooldownTags);
-		TArray< float > Durations = ActorInfo->AbilitySystemComponent->GetActiveEffectsTimeRemaining(Query);
-		if (Durations.Num() > 0)
+		const FGameplayTagContainer* CooldownTags = GetCooldownTags();
+		if (CooldownTags && CooldownTags->Num() > 0)
 		{
-			Durations.Sort();
-			return Durations[Durations.Num()-1];
+			FGameplayEffectQuery const Query = FGameplayEffectQuery::MakeQuery_MatchAnyOwningTags(*CooldownTags);
+			TArray< float > Durations = ActorInfo->AbilitySystemComponent->GetActiveEffectsTimeRemaining(Query);
+			if (Durations.Num() > 0)
+			{
+				Durations.Sort();
+				return Durations[Durations.Num() - 1];
+			}
 		}
 	}
 
@@ -1195,6 +1198,22 @@ FGameplayAbilitySpec* UGameplayAbility::GetCurrentAbilitySpec() const
 	return CurrentActorInfo->AbilitySystemComponent->FindAbilitySpecFromHandle(CurrentSpecHandle);
 }
 
+FGameplayEffectContextHandle UGameplayAbility::GetGrantedByEffectContext() const
+{
+	check(IsInstantiated()); // You should not call this on non instanced abilities.
+	check(CurrentActorInfo);
+	if (CurrentActorInfo)
+	{
+		FActiveGameplayEffectHandle ActiveHandle = CurrentActorInfo->AbilitySystemComponent->FindActiveGameplayEffectHandle(GetCurrentAbilitySpecHandle());
+		if (ActiveHandle.IsValid())
+		{
+			return CurrentActorInfo->AbilitySystemComponent->GetEffectContextFromActiveGEHandle(ActiveHandle);
+		}
+	}
+
+	return FGameplayEffectContextHandle();
+}
+
 UObject* UGameplayAbility::GetSourceObject(FGameplayAbilitySpecHandle Handle, const FGameplayAbilityActorInfo* ActorInfo) const
 {
 	if (ActorInfo != NULL)
diff --git a/Engine/Source/Runtime/GameplayAbilities/Private/Abilities/GameplayAbilityTargetActor.cpp b/Engine/Source/Runtime/GameplayAbilities/Private/Abilities/GameplayAbilityTargetActor.cpp
--- a/Engine/Source/Runtime/GameplayAbilities/Private/Abilities/GameplayAbilityTargetActor.cpp
+++ b/Engine/Source/Runtime/GameplayAbilities/Private/Abilities/GameplayAbilityTargetActor.cpp
@@ -22,12 +22,13 @@ AGameplayAbilityTargetActor::AGameplayAbilityTargetActor(const FObjectInitialize
 
 void AGameplayAbilityTargetActor::EndPlay(const EEndPlayReason::Type EndPlayReason)
 {
-	// We must remove ourselves from GenericLocalConfirmCallbacks/GenericLocalCancelCallbacks, since while these are bound they will inhibit any *other* abilities
-	// that are bound to the same key.
-
-	if (OwningAbility)
+	if (GenericDelegateBoundASC)
 	{
-		const FGameplayAbilityActorInfo* Info = OwningAbility->GetCurrentActorInfo();
+		// We must remove ourselves from GenericLocalConfirmCallbacks/GenericLocalCancelCallbacks, since while these are bound they will inhibit any *other* abilities
+		// that are bound to the same key.
+
+		UAbilitySystemComponent* UnboundASC = nullptr;
+		const FGameplayAbilityActorInfo* Info = (OwningAbility ? OwningAbility->GetCurrentActorInfo() : nullptr);
 		if (Info && Info->IsLocallyControlled())
 		{
 			UAbilitySystemComponent* ASC = Info->AbilitySystemComponent.Get();
@@ -35,8 +36,12 @@ void AGameplayAbilityTargetActor::EndPlay(const EEndPlayReason::Type EndPlayReas
 			{
 				ASC->GenericLocalConfirmCallbacks.RemoveDynamic(this, &AGameplayAbilityTargetActor::ConfirmTargeting);
 				ASC->GenericLocalCancelCallbacks.RemoveDynamic(this, &AGameplayAbilityTargetActor::CancelTargeting);
+
+				UnboundASC = ASC;
 			}
 		}
+
+		ensure(GenericDelegateBoundASC == UnboundASC); // Error checking that we have removed delegates from the same ASC we bound them to
 	}
 
 	Super::EndPlay(EndPlayReason);
@@ -72,8 +77,16 @@ void AGameplayAbilityTargetActor::ConfirmTargetingAndContinue()
 
 void AGameplayAbilityTargetActor::ConfirmTargeting()
 {
-	UAbilitySystemComponent* ASC = OwningAbility->GetCurrentActorInfo()->AbilitySystemComponent.Get();
-	ASC->AbilityReplicatedEventDelegate(EAbilityGenericReplicatedEvent::GenericConfirm, OwningAbility->GetCurrentAbilitySpecHandle(), OwningAbility->GetCurrentActivationInfo().GetActivationPredictionKey() ).Remove(GenericConfirmHandle);
+	const FGameplayAbilityActorInfo* ActorInfo = (OwningAbility ? OwningAbility->GetCurrentActorInfo() : nullptr);
+	UAbilitySystemComponent* ASC = (ActorInfo ? ActorInfo->AbilitySystemComponent.Get() : nullptr);
+	if (ASC)
+	{
+		ASC->AbilityReplicatedEventDelegate(EAbilityGenericReplicatedEvent::GenericConfirm, OwningAbility->GetCurrentAbilitySpecHandle(), OwningAbility->GetCurrentActivationInfo().GetActivationPredictionKey() ).Remove(GenericConfirmHandle);
+	}
+	else
+	{
+		ABILITY_LOG(Warning, TEXT("AGameplayAbilityTargetActor::ConfirmTargeting called with null Ability/ASC! Actor %s"), *GetName());
+	}
 
 	if (IsConfirmTargetingAllowed())
 	{
@@ -88,8 +101,16 @@ void AGameplayAbilityTargetActor::ConfirmTargeting()
 /** Outside code is saying 'stop everything and just forget about it' */
 void AGameplayAbilityTargetActor::CancelTargeting()
 {
-	UAbilitySystemComponent* ASC = OwningAbility->GetCurrentActorInfo()->AbilitySystemComponent.Get();
-	ASC->AbilityReplicatedEventDelegate(EAbilityGenericReplicatedEvent::GenericCancel, OwningAbility->GetCurrentAbilitySpecHandle(), OwningAbility->GetCurrentActivationInfo().GetActivationPredictionKey() ).Remove(GenericCancelHandle);
+	const FGameplayAbilityActorInfo* ActorInfo = (OwningAbility ? OwningAbility->GetCurrentActorInfo() : nullptr);
+	UAbilitySystemComponent* ASC = (ActorInfo ? ActorInfo->AbilitySystemComponent.Get() : nullptr);
+	if (ASC)
+	{
+		ASC->AbilityReplicatedEventDelegate(EAbilityGenericReplicatedEvent::GenericCancel, OwningAbility->GetCurrentAbilitySpecHandle(), OwningAbility->GetCurrentActivationInfo().GetActivationPredictionKey() ).Remove(GenericCancelHandle);
+	}
+	else
+	{
+		ABILITY_LOG(Warning, TEXT("AGameplayAbilityTargetActor::CancelTargeting called with null ASC! Actor %s"), *GetName());
+	}
 
 	CanceledDelegate.Broadcast(FGameplayAbilityTargetDataHandle());
 	Destroy();
@@ -144,6 +165,9 @@ void AGameplayAbilityTargetActor::BindToConfirmCancelInputs()
 			// We have to wait for the callback from the AbilitySystemComponent. Which will always be instigated locally
 			ASC->GenericLocalConfirmCallbacks.AddDynamic(this, &AGameplayAbilityTargetActor::ConfirmTargeting);	// Tell me if the confirm input is pressed
 			ASC->GenericLocalCancelCallbacks.AddDynamic(this, &AGameplayAbilityTargetActor::CancelTargeting);	// Tell me if the cancel input is pressed
+
+			// Save off which ASC we bound so that we can error check that we're removing them later
+			GenericDelegateBoundASC = ASC;
 		}
 		else
 		{	
diff --git a/Engine/Source/Runtime/GameplayAbilities/Private/Abilities/Tasks/AbilityTask_PlayMontageAndWait.cpp b/Engine/Source/Runtime/GameplayAbilities/Private/Abilities/Tasks/AbilityTask_PlayMontageAndWait.cpp
--- a/Engine/Source/Runtime/GameplayAbilities/Private/Abilities/Tasks/AbilityTask_PlayMontageAndWait.cpp
+++ b/Engine/Source/Runtime/GameplayAbilities/Private/Abilities/Tasks/AbilityTask_PlayMontageAndWait.cpp
@@ -42,7 +42,7 @@ void UAbilityTask_PlayMontageAndWait::OnMontageInterrupted()
 		if (AbilitySystemComponent->IsAnimatingAbility(Ability.Get())
 			&& AbilitySystemComponent->GetCurrentMontage() == MontageToPlay)
 		{
-			// Unbind the delegate in this case so OnMontageEnded does not get called as well
+			// Unbind delegates so they don't get called as well
 			BlendingOutDelegate.Unbind();
 
 			AbilitySystemComponent->CurrentMontageStop();
diff --git a/Engine/Source/Runtime/GameplayAbilities/Private/Abilities/Tasks/AbilityTask_SpawnActor.cpp b/Engine/Source/Runtime/GameplayAbilities/Private/Abilities/Tasks/AbilityTask_SpawnActor.cpp
--- a/Engine/Source/Runtime/GameplayAbilities/Private/Abilities/Tasks/AbilityTask_SpawnActor.cpp
+++ b/Engine/Source/Runtime/GameplayAbilities/Private/Abilities/Tasks/AbilityTask_SpawnActor.cpp
@@ -56,7 +56,7 @@ void UAbilityTask_SpawnActor::FinishSpawningActor(UObject* WorldContextObject, F
 			}
 			else if (LocationData->HasEndPoint())
 			{
-				SpawnTransform.SetLocation(LocationData->GetEndPoint());
+				SpawnTransform = LocationData->GetEndPointTransform();
 			}
 		}
 
diff --git a/Engine/Source/Runtime/GameplayAbilities/Private/Abilities/Tasks/AbilityTask_WaitGameplayEffectApplied_Self.cpp b/Engine/Source/Runtime/GameplayAbilities/Private/Abilities/Tasks/AbilityTask_WaitGameplayEffectApplied_Self.cpp
--- a/Engine/Source/Runtime/GameplayAbilities/Private/Abilities/Tasks/AbilityTask_WaitGameplayEffectApplied_Self.cpp
+++ b/Engine/Source/Runtime/GameplayAbilities/Private/Abilities/Tasks/AbilityTask_WaitGameplayEffectApplied_Self.cpp
@@ -13,7 +13,7 @@ UAbilityTask_WaitGameplayEffectApplied_Self::UAbilityTask_WaitGameplayEffectAppl
 {
 }
 
-UAbilityTask_WaitGameplayEffectApplied_Self* UAbilityTask_WaitGameplayEffectApplied_Self::WaitGameplayEffectAppliedToSelf(UObject* WorldContextObject, const FGameplayTargetDataFilterHandle InFilter, FGameplayTagRequirements InSourceTagRequirements, FGameplayTagRequirements InTargetTagRequirements, bool InTriggerOnce, AActor* OptionalExternalOwner)
+UAbilityTask_WaitGameplayEffectApplied_Self* UAbilityTask_WaitGameplayEffectApplied_Self::WaitGameplayEffectAppliedToSelf(UObject* WorldContextObject, const FGameplayTargetDataFilterHandle InFilter, FGameplayTagRequirements InSourceTagRequirements, FGameplayTagRequirements InTargetTagRequirements, bool InTriggerOnce, AActor* OptionalExternalOwner, bool InListenForPeriodicEffect)
 {
 	auto MyObj = NewAbilityTask<UAbilityTask_WaitGameplayEffectApplied_Self>(WorldContextObject);
 	MyObj->Filter = InFilter;
@@ -21,6 +21,7 @@ UAbilityTask_WaitGameplayEffectApplied_Self* UAbilityTask_WaitGameplayEffectAppl
 	MyObj->TargetTagRequirements = InTargetTagRequirements;
 	MyObj->TriggerOnce = InTriggerOnce;
 	MyObj->SetExternalActor(OptionalExternalOwner);
+	MyObj->ListenForPeriodicEffects = InListenForPeriodicEffect;
 	return MyObj;
 }
 
@@ -32,9 +33,17 @@ void UAbilityTask_WaitGameplayEffectApplied_Self::BroadcastDelegate(AActor* Avat
 void UAbilityTask_WaitGameplayEffectApplied_Self::RegisterDelegate()
 {
 	OnApplyGameplayEffectCallbackDelegateHandle = GetASC()->OnGameplayEffectAppliedDelegateToSelf.AddUObject(this, &UAbilityTask_WaitGameplayEffectApplied::OnApplyGameplayEffectCallback);
+	if (ListenForPeriodicEffects)
+	{
+		OnPeriodicGameplayEffectExecuteCallbackDelegateHandle = GetASC()->OnPeriodicGameplayEffectExecuteDelegateOnSelf.AddUObject(this, &UAbilityTask_WaitGameplayEffectApplied::OnApplyGameplayEffectCallback);
+	}
 }
 
 void UAbilityTask_WaitGameplayEffectApplied_Self::RemoveDelegate()
 {
 	GetASC()->OnGameplayEffectAppliedDelegateToSelf.Remove(OnApplyGameplayEffectCallbackDelegateHandle);
+	if (OnPeriodicGameplayEffectExecuteCallbackDelegateHandle.IsValid())
+	{
+		GetASC()->OnGameplayEffectAppliedDelegateToTarget.Remove(OnPeriodicGameplayEffectExecuteCallbackDelegateHandle);
+	}
 }
\ No newline at end of file
diff --git a/Engine/Source/Runtime/GameplayAbilities/Private/Abilities/Tasks/AbilityTask_WaitGameplayEffectApplied_Target.cpp b/Engine/Source/Runtime/GameplayAbilities/Private/Abilities/Tasks/AbilityTask_WaitGameplayEffectApplied_Target.cpp
--- a/Engine/Source/Runtime/GameplayAbilities/Private/Abilities/Tasks/AbilityTask_WaitGameplayEffectApplied_Target.cpp
+++ b/Engine/Source/Runtime/GameplayAbilities/Private/Abilities/Tasks/AbilityTask_WaitGameplayEffectApplied_Target.cpp
@@ -13,7 +13,7 @@ UAbilityTask_WaitGameplayEffectApplied_Target::UAbilityTask_WaitGameplayEffectAp
 {
 }
 
-UAbilityTask_WaitGameplayEffectApplied_Target* UAbilityTask_WaitGameplayEffectApplied_Target::WaitGameplayEffectAppliedToTarget(UObject* WorldContextObject, const FGameplayTargetDataFilterHandle InFilter, FGameplayTagRequirements InSourceTagRequirements, FGameplayTagRequirements InTargetTagRequirements, bool InTriggerOnce, AActor* OptionalExternalOwner)
+UAbilityTask_WaitGameplayEffectApplied_Target* UAbilityTask_WaitGameplayEffectApplied_Target::WaitGameplayEffectAppliedToTarget(UObject* WorldContextObject, const FGameplayTargetDataFilterHandle InFilter, FGameplayTagRequirements InSourceTagRequirements, FGameplayTagRequirements InTargetTagRequirements, bool InTriggerOnce, AActor* OptionalExternalOwner, bool InListenForPeriodicEffect)
 {
 	auto MyObj = NewAbilityTask<UAbilityTask_WaitGameplayEffectApplied_Target>(WorldContextObject);
 	MyObj->Filter = InFilter;
@@ -21,6 +21,7 @@ UAbilityTask_WaitGameplayEffectApplied_Target* UAbilityTask_WaitGameplayEffectAp
 	MyObj->TargetTagRequirements = InTargetTagRequirements;
 	MyObj->TriggerOnce = InTriggerOnce;
 	MyObj->SetExternalActor(OptionalExternalOwner);
+	MyObj->ListenForPeriodicEffects = InListenForPeriodicEffect;
 	return MyObj;
 }
 
@@ -32,9 +33,17 @@ void UAbilityTask_WaitGameplayEffectApplied_Target::BroadcastDelegate(AActor* Av
 void UAbilityTask_WaitGameplayEffectApplied_Target::RegisterDelegate()
 {
 	OnApplyGameplayEffectCallbackDelegateHandle = GetASC()->OnGameplayEffectAppliedDelegateToTarget.AddUObject(this, &UAbilityTask_WaitGameplayEffectApplied::OnApplyGameplayEffectCallback);
+	if (ListenForPeriodicEffects)
+	{
+		OnPeriodicGameplayEffectExecuteCallbackDelegateHandle = GetASC()->OnPeriodicGameplayEffectExecuteDelegateOnTarget.AddUObject(this, &UAbilityTask_WaitGameplayEffectApplied::OnApplyGameplayEffectCallback);
+	}
 }
 
 void UAbilityTask_WaitGameplayEffectApplied_Target::RemoveDelegate()
 {
 	GetASC()->OnGameplayEffectAppliedDelegateToTarget.Remove(OnApplyGameplayEffectCallbackDelegateHandle);
+	if (OnPeriodicGameplayEffectExecuteCallbackDelegateHandle.IsValid())
+	{
+		GetASC()->OnGameplayEffectAppliedDelegateToTarget.Remove(OnPeriodicGameplayEffectExecuteCallbackDelegateHandle);
+	}
 }
\ No newline at end of file
diff --git a/Engine/Source/Runtime/GameplayAbilities/Private/Abilities/Tasks/AbilityTask_WaitGameplayEvent.cpp b/Engine/Source/Runtime/GameplayAbilities/Private/Abilities/Tasks/AbilityTask_WaitGameplayEvent.cpp
new file mode 100644
--- /dev/null
+++ b/Engine/Source/Runtime/GameplayAbilities/Private/Abilities/Tasks/AbilityTask_WaitGameplayEvent.cpp
@@ -0,0 +1,74 @@
+// Copyright 1998-2015 Epic Games, Inc. All Rights Reserved.
+
+#include "AbilitySystemPrivatePCH.h"
+#include "Abilities/Tasks/AbilityTask_WaitGameplayEvent.h"
+#include "AbilitySystemComponent.h"
+#include "Abilities/GameplayAbility.h"
+
+// ----------------------------------------------------------------
+
+UAbilityTask_WaitGameplayEvent::UAbilityTask_WaitGameplayEvent(const FObjectInitializer& ObjectInitializer)
+: Super(ObjectInitializer)
+{
+
+}
+
+UAbilityTask_WaitGameplayEvent* UAbilityTask_WaitGameplayEvent::WaitGameplayEvent(UObject* WorldContextObject, FGameplayTag Tag, AActor* OptionalExternalTarget, bool OnlyTriggerOnce)
+{
+	auto MyObj = NewAbilityTask<UAbilityTask_WaitGameplayEvent>(WorldContextObject);
+	MyObj->Tag = Tag;
+	MyObj->SetExternalTarget(OptionalExternalTarget);
+	MyObj->OnlyTriggerOnce = OnlyTriggerOnce;
+
+	return MyObj;
+}
+
+void UAbilityTask_WaitGameplayEvent::Activate()
+{
+	UAbilitySystemComponent* ASC = GetTargetASC();
+	if (ASC)
+	{	
+		MyHandle = ASC->GenericGameplayEventCallbacks.FindOrAdd(Tag).AddUObject(this, &UAbilityTask_WaitGameplayEvent::GameplayEventCallback);
+	}
+
+	Super::Activate();
+}
+
+void UAbilityTask_WaitGameplayEvent::GameplayEventCallback(const FGameplayEventData* Payload)
+{
+	EventReceived.Broadcast(*Payload);
+	if(OnlyTriggerOnce)
+	{
+		EndTask();
+	}
+}
+
+void UAbilityTask_WaitGameplayEvent::SetExternalTarget(AActor* Actor)
+{
+	if (Actor)
+	{
+		UseExternalTarget = true;
+		OptionalExternalTarget = UAbilitySystemGlobals::GetAbilitySystemComponentFromActor(Actor);
+	}
+}
+
+UAbilitySystemComponent* UAbilityTask_WaitGameplayEvent::GetTargetASC()
+{
+	if (UseExternalTarget)
+	{
+		return OptionalExternalTarget.Get();
+	}
+
+	return AbilitySystemComponent.Get();
+}
+
+void UAbilityTask_WaitGameplayEvent::OnDestroy(bool AbilityEnding)
+{
+	UAbilitySystemComponent* ASC = GetTargetASC();
+	if (ASC && MyHandle.IsValid())
+	{	
+		ASC->GenericGameplayEventCallbacks.FindOrAdd(Tag).Remove(MyHandle);
+	}
+
+	Super::OnDestroy(AbilityEnding);
+}
diff --git a/Engine/Source/Runtime/GameplayAbilities/Private/Abilities/Tasks/AbilityTask_WaitGameplayTagBase.cpp b/Engine/Source/Runtime/GameplayAbilities/Private/Abilities/Tasks/AbilityTask_WaitGameplayTagBase.cpp
--- a/Engine/Source/Runtime/GameplayAbilities/Private/Abilities/Tasks/AbilityTask_WaitGameplayTagBase.cpp
+++ b/Engine/Source/Runtime/GameplayAbilities/Private/Abilities/Tasks/AbilityTask_WaitGameplayTagBase.cpp
@@ -15,9 +15,10 @@ UAbilityTask_WaitGameplayTag::UAbilityTask_WaitGameplayTag(const FObjectInitiali
 
 void UAbilityTask_WaitGameplayTag::Activate()
 {
-	if (AbilitySystemComponent.IsValid())
+	UAbilitySystemComponent* ASC = GetTargetASC();
+	if (ASC)
 	{
-		DelegateHandle = AbilitySystemComponent->RegisterGameplayTagEvent(Tag).AddUObject(this, &UAbilityTask_WaitGameplayTag::GameplayTagCallback);
+		DelegateHandle = ASC->RegisterGameplayTagEvent(Tag).AddUObject(this, &UAbilityTask_WaitGameplayTag::GameplayTagCallback);
 		RegisteredCallback = true;
 	}
 }
@@ -28,9 +29,10 @@ void UAbilityTask_WaitGameplayTag::GameplayTagCallback(const FGameplayTag InTag,
 
 void UAbilityTask_WaitGameplayTag::OnDestroy(bool AbilityIsEnding)
 {
-	if (RegisteredCallback && AbilitySystemComponent.IsValid())
+	UAbilitySystemComponent* ASC = GetTargetASC();
+	if (RegisteredCallback && ASC)
 	{
-		AbilitySystemComponent->RegisterGameplayTagEvent(Tag).Remove(DelegateHandle);
+		ASC->RegisterGameplayTagEvent(Tag).Remove(DelegateHandle);
 	}
 
 	Super::OnDestroy(AbilityIsEnding);
diff --git a/Engine/Source/Runtime/GameplayAbilities/Private/Abilities/Tasks/AbilityTask_WaitTargetData.cpp b/Engine/Source/Runtime/GameplayAbilities/Private/Abilities/Tasks/AbilityTask_WaitTargetData.cpp
--- a/Engine/Source/Runtime/GameplayAbilities/Private/Abilities/Tasks/AbilityTask_WaitTargetData.cpp
+++ b/Engine/Source/Runtime/GameplayAbilities/Private/Abilities/Tasks/AbilityTask_WaitTargetData.cpp
@@ -245,6 +245,10 @@ void UAbilityTask_WaitTargetData::OnTargetDataReplicatedCancelledCallback()
 void UAbilityTask_WaitTargetData::OnTargetDataReadyCallback(FGameplayAbilityTargetDataHandle Data)
 {
 	check(AbilitySystemComponent.IsValid());
+	if (Ability.IsValid() == false)
+	{
+		return;
+	}
 
 	FScopedPredictionWindow	ScopedPrediction(AbilitySystemComponent.Get(), ShouldReplicateDataToServer());
 	
@@ -327,8 +331,10 @@ void UAbilityTask_WaitTargetData::OnDestroy(bool AbilityEnded)
 
 bool UAbilityTask_WaitTargetData::ShouldReplicateDataToServer() const
 {
-	check(Ability.IsValid());
-	check(TargetActor.IsValid()); // Target actor should always be valid on client
+	if (Ability.IsValid() == false || TargetActor.IsValid() == false)
+	{
+		return false;
+	}
 
 	// Send TargetData to the server IFF we are the client and this isn't a GameplayTargetActor that can produce data on the server	
 	const FGameplayAbilityActorInfo* Info = Ability->GetCurrentActorInfo();
diff --git a/Engine/Source/Runtime/GameplayAbilities/Private/AbilitySystemBlueprintLibrary.cpp b/Engine/Source/Runtime/GameplayAbilities/Private/AbilitySystemBlueprintLibrary.cpp
--- a/Engine/Source/Runtime/GameplayAbilities/Private/AbilitySystemBlueprintLibrary.cpp
+++ b/Engine/Source/Runtime/GameplayAbilities/Private/AbilitySystemBlueprintLibrary.cpp
@@ -260,6 +260,11 @@ FHitResult UAbilitySystemBlueprintLibrary::GetHitResultFromTargetData(FGameplayA
 
 bool UAbilitySystemBlueprintLibrary::TargetDataHasOrigin(FGameplayAbilityTargetDataHandle TargetData, int32 Index)
 {
+	if (TargetData.Data.IsValidIndex(Index) == false)
+	{
+		return false;
+	}
+
 	FGameplayAbilityTargetData* Data = TargetData.Data[Index].Get();
 	if (Data)
 	{
@@ -270,6 +275,11 @@ bool UAbilitySystemBlueprintLibrary::TargetDataHasOrigin(FGameplayAbilityTargetD
 
 FTransform UAbilitySystemBlueprintLibrary::GetTargetDataOrigin(FGameplayAbilityTargetDataHandle TargetData, int32 Index)
 {
+	if (TargetData.Data.IsValidIndex(Index) == false)
+	{
+		return FTransform::Identity;
+	}
+
 	FGameplayAbilityTargetData* Data = TargetData.Data[Index].Get();
 	if (Data)
 	{
@@ -342,6 +352,10 @@ FTransform UAbilitySystemBlueprintLibrary::GetTargetDataEndPointTransform(FGamep
 
 // -------------------------------------------------------------------------------------
 
+bool UAbilitySystemBlueprintLibrary::EffectContextIsValid(FGameplayEffectContextHandle EffectContext)
+{
+	return EffectContext.IsValid();
+}
 
 bool UAbilitySystemBlueprintLibrary::EffectContextIsInstigatorLocallyControlled(FGameplayEffectContextHandle EffectContext)
 {
@@ -395,12 +409,12 @@ FVector UAbilitySystemBlueprintLibrary::EffectContextGetOrigin(FGameplayEffectCo
 
 bool UAbilitySystemBlueprintLibrary::IsInstigatorLocallyControlled(FGameplayCueParameters Parameters)
 {
-	return Parameters.EffectContext.IsLocallyControlled();
+	return Parameters.IsInstigatorLocallyControlled();
 }
 
 bool UAbilitySystemBlueprintLibrary::IsInstigatorLocallyControlledPlayer(FGameplayCueParameters Parameters)
 {
-	return Parameters.EffectContext.IsLocallyControlledPlayer();
+	return Parameters.IsInstigatorLocallyControlledPlayer();
 }
 
 int32 UAbilitySystemBlueprintLibrary::GetActorCount(FGameplayCueParameters Parameters)
@@ -443,7 +457,7 @@ void UAbilitySystemBlueprintLibrary::ForwardGameplayCueToTarget(TScriptInterface
 
 AActor*	UAbilitySystemBlueprintLibrary::GetInstigatorActor(FGameplayCueParameters Parameters)
 {
-	return Parameters.EffectContext.GetInstigator();
+	return Parameters.GetInstigator();
 }
 
 FTransform UAbilitySystemBlueprintLibrary::GetInstigatorTransform(FGameplayCueParameters Parameters)
@@ -465,15 +479,19 @@ FVector UAbilitySystemBlueprintLibrary::GetOrigin(FGameplayCueParameters Paramet
 		return Parameters.EffectContext.GetOrigin();
 	}
 
-	return FVector::ZeroVector;
+	return Parameters.Location;
 }
 
 bool UAbilitySystemBlueprintLibrary::GetGameplayCueEndLocationAndNormal(AActor* TargetActor, FGameplayCueParameters Parameters, FVector& Location, FVector& Normal)
 {
 	FGameplayEffectContext* Data = Parameters.EffectContext.Get();
-	if (Data && Data->GetHitResult())
+	if (Parameters.Location.IsNearlyZero() == false)
+	{
+		Location = Parameters.Location;
+		Normal = Parameters.Normal;
+	}
+	else if (Data && Data->GetHitResult())
 	{
-		
 		Location = Data->GetHitResult()->Location;
 		Normal = Data->GetHitResult()->Normal;
 		return true;
@@ -490,6 +508,12 @@ bool UAbilitySystemBlueprintLibrary::GetGameplayCueEndLocationAndNormal(AActor*
 
 bool UAbilitySystemBlueprintLibrary::GetGameplayCueDirection(AActor* TargetActor, FGameplayCueParameters Parameters, FVector& Direction)
 {
+	if (Parameters.Normal.IsNearlyZero() == false)
+	{
+		Direction = Parameters.Normal;
+		return true;
+	}
+
 	if (FGameplayEffectContext* Ctx = Parameters.EffectContext.Get())
 	{
 		if (Ctx->GetHitResult())
@@ -670,4 +694,21 @@ int32 UAbilitySystemBlueprintLibrary::GetActiveGameplayEffectStackCount(FActiveG
 		return ASC->GetCurrentStackCount(ActiveHandle);
 	}
 	return 0;
+}
+
+float UAbilitySystemBlueprintLibrary::GetModifiedAttributeMagnitude(FGameplayEffectSpecHandle SpecHandle, FGameplayAttribute Attribute)
+{
+	FGameplayEffectSpec* Spec = SpecHandle.Data.Get();
+	float Delta = 0.f;
+	if (Spec)
+	{
+		for (FGameplayEffectModifiedAttribute &Mod : Spec->ModifiedAttributes)
+		{
+			if (Mod.Attribute == Attribute)
+			{
+				Delta += Mod.TotalMagnitude;
+			}
+		}
+	}
+	return Delta;
 }
\ No newline at end of file
diff --git a/Engine/Source/Runtime/GameplayAbilities/Private/AbilitySystemComponent.cpp b/Engine/Source/Runtime/GameplayAbilities/Private/AbilitySystemComponent.cpp
--- a/Engine/Source/Runtime/GameplayAbilities/Private/AbilitySystemComponent.cpp
+++ b/Engine/Source/Runtime/GameplayAbilities/Private/AbilitySystemComponent.cpp
@@ -29,6 +29,9 @@ UAbilitySystemComponent::UAbilitySystemComponent(const FObjectInitializer& Objec
 	bWantsInitializeComponent = true;
 
 	PrimaryComponentTick.bStartWithTickEnabled = true; // FIXME! Just temp until timer manager figured out
+	bAutoActivate = true;	// Forcing AutoActivate since above we manually force tick enabled.
+							// if we don't have this, UpdateShouldTick() fails to have any effect
+							// because we'll be receiving ticks but bIsActive starts as false
 
 	ActiveGameplayCues.Owner = this;
 
@@ -127,6 +130,13 @@ void UAbilitySystemComponent::OnRegister()
 	AbilityActorInfo = TSharedPtr<FGameplayAbilityActorInfo>(UAbilitySystemGlobals::Get().AllocAbilityActorInfo());
 }
 
+void UAbilitySystemComponent::OnUnregister()
+{
+	Super::OnUnregister();
+
+	DestroyActiveState();
+}
+
 // ---------------------------------------------------------
 
 const FActiveGameplayEffect* UAbilitySystemComponent::GetActiveGameplayEffect(const FActiveGameplayEffectHandle Handle) const
@@ -756,29 +766,6 @@ FActiveGameplayEffectHandle UAbilitySystemComponent::ApplyGameplayEffectSpecToSe
 		ActiveGameplayEffects.RemoveActiveEffects(ClearQuery, -1);
 	}
 	
-
-	// ------------------------------------------------------
-	//	Apply Granted Abilities
-	//	
-	//	Note: Doing this before apply TargetEffectSpecs, but we could just as easily apply after.
-	//	Hedging bet for now that we are more likely to want to have 'passive ability that reacts to linked GE'
-	//	over 'need ability to activate after linked GE is applied'.
-	//	
-	//	Note2: This is allowing instant GEs to permanently grant abilities. This could be disallowed if needed.
-	// ------------------------------------------------------
-	if (bIsNetAuthority)
-	{
-		for (FGameplayAbilitySpecDef& AbilitySpecDef : OurCopyOfSpec->GrantedAbilitySpecs)
-		{
-			// Only do this is we havent assigned the ability yet! This prevents cases where stacking GEs
-			// would regrant the ability every time the stack was applied
-			if (AbilitySpecDef.AssignedHandle.IsValid() == false)
-			{
-				GiveAbility( FGameplayAbilitySpec(AbilitySpecDef, MyHandle) );
-			}
-		}	
-	}
-
 	// ------------------------------------------------------
 	// Apply Linked effects
 	// todo: this is ignoring the returned handles, should we put them into a TArray and return all of the handles?
@@ -964,10 +951,7 @@ void UAbilitySystemComponent::InvokeGameplayCueEvent(const FGameplayEffectSpecFo
 	
 	float ExecuteLevel = Spec.GetLevel();
 
-	FGameplayCueParameters CueParameters;
-	CueParameters.EffectContext = Spec.GetContext();
-	CueParameters.AggregatedSourceTags = Spec.AggregatedSourceTags;
-	CueParameters.AggregatedTargetTags = Spec.AggregatedTargetTags;
+	FGameplayCueParameters CueParameters(Spec);
 
 	for (FGameplayEffectCue CueInfo : Spec.Def->GameplayCues)
 	{
@@ -995,12 +979,7 @@ void UAbilitySystemComponent::InvokeGameplayCueEvent(const FGameplayEffectSpecFo
 
 void UAbilitySystemComponent::InvokeGameplayCueEvent(const FGameplayTag GameplayCueTag, EGameplayCueEvent::Type EventType, FGameplayEffectContextHandle EffectContext)
 {
-	FGameplayCueParameters CueParameters;
-
-	if (EffectContext.IsValid())
-	{
-		CueParameters.EffectContext = EffectContext;
-	}
+	FGameplayCueParameters CueParameters(EffectContext);
 
 	CueParameters.NormalizedMagnitude = 1.f;
 	CueParameters.RawMagnitude = 0.f;
@@ -1301,7 +1280,6 @@ bool UAbilitySystemComponent::HasAuthorityOrPredictionKey(const FGameplayAbility
 
 void UAbilitySystemComponent::PrintAllGameplayEffects() const
 {
-	ABILITY_LOG_SCOPE(TEXT("PrintAllGameplayEffects %s"), *GetName());
 	ABILITY_LOG(Log, TEXT("Owner: %s. Avatar: %s"), *GetOwner()->GetName(), *AbilityActorInfo->AvatarActor->GetName());
 	ActiveGameplayEffects.PrintAllGameplayEffects();
 }
@@ -1329,6 +1307,16 @@ void UAbilitySystemComponent::OnGameplayEffectAppliedToSelf(UAbilitySystemCompon
 	OnGameplayEffectAppliedDelegateToSelf.Broadcast(Source, SpecApplied, ActiveHandle);
 }
 
+void UAbilitySystemComponent::OnPeriodicGameplayEffectExecuteOnTarget(UAbilitySystemComponent* Target, const FGameplayEffectSpec& SpecExecuted, FActiveGameplayEffectHandle ActiveHandle)
+{
+	OnPeriodicGameplayEffectExecuteDelegateOnTarget.Broadcast(Target, SpecExecuted, ActiveHandle);
+}
+
+void UAbilitySystemComponent::OnPeriodicGameplayEffectExecuteOnSelf(UAbilitySystemComponent* Source, const FGameplayEffectSpec& SpecExecuted, FActiveGameplayEffectHandle ActiveHandle)
+{
+	OnPeriodicGameplayEffectExecuteDelegateOnSelf.Broadcast(Source, SpecExecuted, ActiveHandle);
+}
+
 TArray<TWeakObjectPtr<UGameplayTask> >&	UAbilitySystemComponent::GetAbilityActiveTasks(UGameplayAbility* Ability)
 {
 	return Ability->ActiveTasks;
@@ -1349,31 +1337,91 @@ FString ASC_CleanupName(FString Str)
 	return Str;
 }
 
-void UAbilitySystemComponent::DisplayDebug(class UCanvas* Canvas, const class FDebugDisplayInfo& DebugDisplay, float& YL, float& YPos)
+void AccumulateScreenPos(float& XPos, float& YPos, float AdditionalY, float OriginalY, float MaxY, const float NewColumnYPadding, UCanvas* Canvas)
 {
+	const float ColumnWidth = Canvas->ClipX * 0.4f;
 
+	float NewY = YPos + AdditionalY;
+	if (NewY > MaxY)
+	{
+		// Need new column, reset Y to original height
+		NewY = NewColumnYPadding;
+		XPos += ColumnWidth;
+	}
+	YPos = NewY;
+}
+
+void UAbilitySystemComponent::DisplayDebug(class UCanvas* Canvas, const class FDebugDisplayInfo& DebugDisplay, float& YL, float& YPos)
+{
 	bool bShowAttributes = true;
 	bool bShowGameplayEffects = true;
 	bool bShowAbilities = true;
 
+	if (DebugDisplay.IsDisplayOn(FName(TEXT("Attributes"))))
+	{
+		bShowAbilities = false;
+		bShowAttributes = true;
+		bShowGameplayEffects = false;
+	}
 	if (DebugDisplay.IsDisplayOn(FName(TEXT("Ability"))))
 	{
 		bShowAbilities = true;
 		bShowAttributes = false;
 		bShowGameplayEffects = false;
 	}
+	else if (DebugDisplay.IsDisplayOn(FName(TEXT("GameplayEffects"))))
+	{
+		bShowAbilities = false;
+		bShowAttributes = false;
+		bShowGameplayEffects = true;
+	}
+
+	float XPos = 0.f;
+	const float OriginalX = XPos;
+	const float OriginalY = YPos;
+	const float MaxY = Canvas->ClipY - 150.f; // Give some padding for any non-columnizing debug output following this output
+	float NewColumnYPadding = 30.f;
+
+	// Draw title at top of screen (default HUD debug text starts at 50 ypos, we can position this on top)*
+	//   *until someone changes it unknowingly
+	{
+		FString DebugTitle("");
+		// Category
+		if (bShowAbilities) DebugTitle += TEXT("ABILITIES ");
+		if (bShowAttributes) DebugTitle += TEXT("ATTRIBUTES ");
+		if (bShowGameplayEffects) DebugTitle += TEXT("GAMEPLAYEFFECTS ");
+		// Avatar info
+		if (AvatarActor)
+		{
+			DebugTitle += FString::Printf(TEXT("for avatar %s "), *AvatarActor->GetName());
+			if (AvatarActor->Role == ROLE_AutonomousProxy) DebugTitle += TEXT("(local player) ");
+			else if (AvatarActor->Role == ROLE_SimulatedProxy) DebugTitle += TEXT("(simulated) ");
+			else if (AvatarActor->Role == ROLE_Authority) DebugTitle += TEXT("(authority) ");
+		}
+		// Owner info
+		if (OwnerActor && OwnerActor != AvatarActor)
+		{
+			DebugTitle += FString::Printf(TEXT("for owner %s "), *OwnerActor->GetName());
+			if (OwnerActor->Role == ROLE_AutonomousProxy) DebugTitle += TEXT("(autonomous) ");
+			else if (OwnerActor->Role == ROLE_SimulatedProxy) DebugTitle += TEXT("(simulated) ");
+			else if (OwnerActor->Role == ROLE_Authority) DebugTitle += TEXT("(authority) ");
+		}
+
+		Canvas->SetDrawColor(FColor::White);
+		Canvas->DrawText(GEngine->GetLargeFont(), DebugTitle, XPos+4.f, 10.f, 1.5f, 1.5f);
+	}
 
 	FGameplayTagContainer OwnerTags;
 	GetOwnedGameplayTags(OwnerTags);
 
 	Canvas->SetDrawColor(FColor::White);
-	YL = Canvas->DrawText(GEngine->GetTinyFont(), FString::Printf(TEXT("Owned Tags: %s"), *OwnerTags.ToStringSimple()), 4.f, YPos);
-	YPos += YL;
+	YL = Canvas->DrawText(GEngine->GetTinyFont(), FString::Printf(TEXT("Owned Tags: %s"), *OwnerTags.ToStringSimple()), XPos+4.f, YPos);
+	AccumulateScreenPos(XPos, YPos, YL, OriginalY, MaxY, NewColumnYPadding, Canvas);
 
 	if (BlockedAbilityTags.GetExplicitGameplayTags().Num() > 0)
 	{
-		YL = Canvas->DrawText(GEngine->GetTinyFont(), FString::Printf(TEXT("BlockedAbilityTags: %s"), *BlockedAbilityTags.GetExplicitGameplayTags().ToStringSimple()), 4.f, YPos);
-		YPos += YL;
+		YL = Canvas->DrawText(GEngine->GetTinyFont(), FString::Printf(TEXT("BlockedAbilityTags: %s"), *BlockedAbilityTags.GetExplicitGameplayTags().ToStringSimple()), XPos+4.f, YPos);
+		AccumulateScreenPos(XPos, YPos, YL, OriginalY, MaxY, NewColumnYPadding, Canvas);
 	}
 
 	TSet<FGameplayAttribute> DrawAttributes;
@@ -1383,7 +1431,7 @@ void UAbilitySystemComponent::DisplayDebug(class UCanvas* Canvas, const class FD
 	// -------------------------------------------------------------
 
 
-	if (bShowGameplayEffects || bShowAttributes)
+	if (bShowAttributes)
 	{
 		// Draw the attribute aggregator map.
 		for (auto It = ActiveGameplayEffects.AttributeAggregatorMap.CreateConstIterator(); It; ++It)
@@ -1394,6 +1442,17 @@ void UAbilitySystemComponent::DisplayDebug(class UCanvas* Canvas, const class FD
 			{
 				FAggregator& Aggregator = *AggregatorRef.Get();
 
+				bool HasActiveMod = false;
+				for (int32 ModOpIdx = 0; ModOpIdx < ARRAY_COUNT(Aggregator.Mods); ++ModOpIdx)
+				{
+					if(Aggregator.Mods[ModOpIdx].Num() > 0)
+					{
+						HasActiveMod = true;
+					}
+				}
+				if (HasActiveMod == false)
+					continue;
+
 				float FinalValue = GetNumericAttribute(Attribute);
 				float BaseValue = Aggregator.GetBaseValue();
 
@@ -1404,8 +1463,8 @@ void UAbilitySystemComponent::DisplayDebug(class UCanvas* Canvas, const class FD
 				}
 
 				Canvas->SetDrawColor(FColor::White);
-				YL = Canvas->DrawText(GEngine->GetTinyFont(), AttributeString, 4.f, YPos);
-				YPos += YL;
+				YL = Canvas->DrawText(GEngine->GetTinyFont(), AttributeString, XPos+4.f, YPos);
+				AccumulateScreenPos(XPos, YPos, YL, OriginalY, MaxY, NewColumnYPadding, Canvas);
 
 				DrawAttributes.Add(Attribute);
 
@@ -1426,12 +1485,12 @@ void UAbilitySystemComponent::DisplayDebug(class UCanvas* Canvas, const class FD
 							if (Mod.TargetTagReqs) SrcName += FString::Printf(TEXT("TargetTags: [%s]"), *Mod.TargetTagReqs->ToString());
 						}
 
-						YL = Canvas->DrawText(GEngine->GetTinyFont(), FString::Printf(TEXT("   %s\t %.2f - %s"), *EGameplayModOpToString(ModOpIdx), Mod.EvaluatedMagnitude, *SrcName), 7.f, YPos);
-						YPos += YL;
-
+						YL = Canvas->DrawText(GEngine->GetTinyFont(), FString::Printf(TEXT("   %s\t %.2f - %s"), *EGameplayModOpToString(ModOpIdx), Mod.EvaluatedMagnitude, *SrcName), XPos+7.f, YPos);
+						AccumulateScreenPos(XPos, YPos, YL, OriginalY, MaxY, NewColumnYPadding, Canvas);
+						NewColumnYPadding = FMath::Max<float>(NewColumnYPadding, YPos+YL);
 					}
 				}
-				YPos += MaxCharHeight;
+				AccumulateScreenPos(XPos, YPos, MaxCharHeight, OriginalY, MaxY, NewColumnYPadding, Canvas);
 			}
 		}
 	}
@@ -1469,17 +1528,23 @@ void UAbilitySystemComponent::DisplayDebug(class UCanvas* Canvas, const class FD
 				}
 			}
 
+			FString LevelString;
+			if (ActiveGE.Spec.GetLevel() > 1.f)
+			{
+				LevelString = FString::Printf(TEXT("Level: %.2f"), ActiveGE.Spec.GetLevel() );
+			}
+
 			Canvas->SetDrawColor(ActiveGE.bIsInhibited ? FColor(128, 128, 128): FColor::White );
 
-			YL = Canvas->DrawText(GEngine->GetTinyFont(), FString::Printf(TEXT("%s %s %s"), *ASC_CleanupName(GetNameSafe(ActiveGE.Spec.Def)), *DurationStr, *StackString ), 4.f, YPos);
-			YPos += YL;
+			YL = Canvas->DrawText(GEngine->GetTinyFont(), FString::Printf(TEXT("%s %s %s %s"), *ASC_CleanupName(GetNameSafe(ActiveGE.Spec.Def)), *DurationStr, *StackString, *LevelString ), XPos+4.f, YPos);
+			AccumulateScreenPos(XPos, YPos, YL, OriginalY, MaxY, NewColumnYPadding, Canvas);
 
 			FGameplayTagContainer GrantedTags;
 			ActiveGE.Spec.GetAllGrantedTags(GrantedTags);
 			if (GrantedTags.Num() > 0)
 			{
-				YL = Canvas->DrawText(GEngine->GetTinyFont(), FString::Printf(TEXT("Granted Tags: %s"), *GrantedTags.ToStringSimple() ), 7.f, YPos);
-				YPos += YL;
+				YL = Canvas->DrawText(GEngine->GetTinyFont(), FString::Printf(TEXT("Granted Tags: %s"), *GrantedTags.ToStringSimple() ), XPos+7.f, YPos);
+				AccumulateScreenPos(XPos, YPos, YL, OriginalY, MaxY, NewColumnYPadding, Canvas);
 			}
 
 			for (int32 ModIdx=0; ModIdx < ActiveGE.Spec.Modifiers.Num(); ++ModIdx)
@@ -1501,13 +1566,13 @@ void UAbilitySystemComponent::DisplayDebug(class UCanvas* Canvas, const class FD
 					Canvas->SetDrawColor(FColor(128, 128, 128) );
 				}
 
-				YL = Canvas->DrawText(GEngine->GetTinyFont(), FString::Printf(TEXT("Mod: %s. %s. %.2f"), *ModInfo.Attribute.GetName(), *EGameplayModOpToString(ModInfo.ModifierOp), ModSpec.GetEvaluatedMagnitude() ), 7.f, YPos);
-				YPos += YL;
+				YL = Canvas->DrawText(GEngine->GetTinyFont(), FString::Printf(TEXT("Mod: %s. %s. %.2f"), *ModInfo.Attribute.GetName(), *EGameplayModOpToString(ModInfo.ModifierOp), ModSpec.GetEvaluatedMagnitude() ), XPos+7.f, YPos);
+				AccumulateScreenPos(XPos, YPos, YL, OriginalY, MaxY, NewColumnYPadding, Canvas);
 
 				Canvas->SetDrawColor(ActiveGE.bIsInhibited ? FColor(128, 128, 128): FColor::White );
 			}
 
-			YPos += MaxCharHeight;
+			AccumulateScreenPos(XPos, YPos, MaxCharHeight, OriginalY, MaxY, NewColumnYPadding, Canvas);
 		}
 	}
 
@@ -1528,12 +1593,12 @@ void UAbilitySystemComponent::DisplayDebug(class UCanvas* Canvas, const class FD
 				if (Attribute.IsValid())
 				{
 					float Value = GetNumericAttribute(Attribute);
-					YL = Canvas->DrawText(GEngine->GetTinyFont(), FString::Printf(TEXT("%s %.2f"), *Attribute.GetName(), Value ), 4.f, YPos);
-					YPos += YL;
+					YL = Canvas->DrawText(GEngine->GetTinyFont(), FString::Printf(TEXT("%s %.2f"), *Attribute.GetName(), Value ), XPos+4.f, YPos);
+					AccumulateScreenPos(XPos, YPos, YL, OriginalY, MaxY, NewColumnYPadding, Canvas);
 				}
 			}
 		}
-		YPos += MaxCharHeight;
+		AccumulateScreenPos(XPos, YPos, MaxCharHeight, OriginalY, MaxY, NewColumnYPadding, Canvas);
 	}
 
 	// -------------------------------------------------------------
@@ -1571,8 +1636,8 @@ void UAbilitySystemComponent::DisplayDebug(class UCanvas* Canvas, const class FD
 			FString InputPressedStr = AbilitySpec.InputPressed ? TEXT("(InputPressed)") : TEXT("");
 
 			Canvas->SetDrawColor(AbilityTextColor);
-			YL = Canvas->DrawText(GEngine->GetTinyFont(), FString::Printf(TEXT("%s %s %s"), *ASC_CleanupName(GetNameSafe(AbilitySpec.Ability)), *StatusText, *InputPressedStr), 4.f, YPos);
-			YPos += YL;
+			YL = Canvas->DrawText(GEngine->GetTinyFont(), FString::Printf(TEXT("%s %s %s"), *ASC_CleanupName(GetNameSafe(AbilitySpec.Ability)), *StatusText, *InputPressedStr), XPos+4.f, YPos);
+			AccumulateScreenPos(XPos, YPos, YL, OriginalY, MaxY, NewColumnYPadding, Canvas);
 	
 			if (AbilitySpec.IsActive())
 			{
@@ -1589,23 +1654,28 @@ void UAbilitySystemComponent::DisplayDebug(class UCanvas* Canvas, const class FD
 						UGameplayTask* Task = TaskPtr.Get();
 						if (Task)
 						{
-							YL = Canvas->DrawText(GEngine->GetTinyFont(), FString::Printf(TEXT("%s"), *Task->GetDebugString()), 7.f, YPos);
-							YPos += YL;
+							YL = Canvas->DrawText(GEngine->GetTinyFont(), FString::Printf(TEXT("%s"), *Task->GetDebugString()), XPos+7.f, YPos);
+							AccumulateScreenPos(XPos, YPos, YL, OriginalY, MaxY, NewColumnYPadding, Canvas);
 						}
 					}
 
 					if (InstanceIdx < Instances.Num() - 2)
 					{
 						Canvas->SetDrawColor(FColor(128, 128, 128));
-						YL = Canvas->DrawText(GEngine->GetTinyFont(), FString::Printf(TEXT("--------")), 7.f, YPos);
-						YPos += YL;
+						YL = Canvas->DrawText(GEngine->GetTinyFont(), FString::Printf(TEXT("--------")), XPos+7.f, YPos);
+						AccumulateScreenPos(XPos, YPos, YL, OriginalY, MaxY, NewColumnYPadding, Canvas);
 					}
 				}
 			}
 		}
-		YPos += MaxCharHeight;
+		AccumulateScreenPos(XPos, YPos, MaxCharHeight, OriginalY, MaxY, NewColumnYPadding, Canvas);
 	}
 
+	if (XPos > OriginalX)
+	{
+		// We flooded to new columns, returned YPos should be max Y (and some padding)
+		YPos = MaxY + MaxCharHeight*2.f;
+	}
 	YL = MaxCharHeight;
 }
 
diff --git a/Engine/Source/Runtime/GameplayAbilities/Private/AbilitySystemComponent_Abilities.cpp b/Engine/Source/Runtime/GameplayAbilities/Private/AbilitySystemComponent_Abilities.cpp
--- a/Engine/Source/Runtime/GameplayAbilities/Private/AbilitySystemComponent_Abilities.cpp
+++ b/Engine/Source/Runtime/GameplayAbilities/Private/AbilitySystemComponent_Abilities.cpp
@@ -54,26 +54,7 @@ void UAbilitySystemComponent::UninitializeComponent()
 
 void UAbilitySystemComponent::OnComponentDestroyed()
 {
-	// If we haven't already begun being destroyed
-	if ((GetFlags() & RF_BeginDestroyed) == 0)
-	{
-		// Cancel all abilities before we are destroyed.
-		CancelAbilities();
-
-		// Mark pending kill any remaining instanced abilities
-		// (CancelAbilities() will only MarkPending kill InstancePerExecution abilities).
-		for (const FGameplayAbilitySpec& Spec : GetActivatableAbilities())
-		{
-			TArray<UGameplayAbility*> AbilitiesToCancel = Spec.GetAbilityInstances();
-			for (UGameplayAbility* InstanceAbility : AbilitiesToCancel)
-			{
-				if (InstanceAbility)
-				{
-					InstanceAbility->MarkPendingKill();
-				}
-			}
-		}
-	}
+	DestroyActiveState();
 
 	// Call the super at the end, after we've done what we needed to do
 	Super::OnComponentDestroyed();
@@ -554,6 +535,17 @@ FGameplayAbilitySpec* UAbilitySystemComponent::FindAbilitySpecFromInputID(int32
 	return nullptr;
 }
 
+FGameplayEffectContextHandle UAbilitySystemComponent::GetEffectContextFromActiveGEHandle(FActiveGameplayEffectHandle Handle)
+{
+	FActiveGameplayEffect* ActiveGE = ActiveGameplayEffects.GetActiveGameplayEffect(Handle);
+	if (ActiveGE)
+	{
+		return ActiveGE->Spec.GetEffectContext();
+	}
+
+	return FGameplayEffectContextHandle();
+}
+
 UGameplayAbility* UAbilitySystemComponent::CreateNewInstanceOfAbility(FGameplayAbilitySpec& Spec, const UGameplayAbility* Ability)
 {
 	check(Ability);
@@ -716,6 +708,30 @@ void UAbilitySystemComponent::CancelAllAbilities(UGameplayAbility* Ignore)
 	}
 }
 
+void UAbilitySystemComponent::DestroyActiveState()
+{
+	// If we haven't already begun being destroyed
+	if ((GetFlags() & RF_BeginDestroyed) == 0)
+	{
+		// Cancel all abilities before we are destroyed.
+		CancelAbilities();
+
+		// Mark pending kill any remaining instanced abilities
+		// (CancelAbilities() will only MarkPending kill InstancePerExecution abilities).
+		for (const FGameplayAbilitySpec& Spec : GetActivatableAbilities())
+		{
+			TArray<UGameplayAbility*> AbilitiesToCancel = Spec.GetAbilityInstances();
+			for (UGameplayAbility* InstanceAbility : AbilitiesToCancel)
+			{
+				if (InstanceAbility)
+				{
+					InstanceAbility->MarkPendingKill();
+				}
+			}
+		}
+	}
+}
+
 void UAbilitySystemComponent::ApplyAbilityBlockAndCancelTags(const FGameplayTagContainer& AbilityTags, UGameplayAbility* RequestingAbility, bool bEnableBlockTags, const FGameplayTagContainer& BlockTags, bool bExecuteCancelTags, const FGameplayTagContainer& CancelTags)
 {
 	if (bEnableBlockTags)
@@ -983,8 +999,6 @@ bool UAbilitySystemComponent::InternalTryActivateAbility(FGameplayAbilitySpecHan
 		return false;
 	}
 
-	
-
 	// Check to see if this a local only or server only ability, if so don't execute
 	if (!bIsLocal)
 	{
@@ -1693,6 +1707,11 @@ int32 UAbilitySystemComponent::HandleGameplayEvent(FGameplayTag EventTag, const
 		}
 	}
 
+	if (FGameplayEventMulticastDelegate* Delegate = GenericGameplayEventCallbacks.Find(EventTag))
+	{
+		Delegate->Broadcast(Payload);
+	}
+
 	return TriggeredCount;
 }
 
diff --git a/Engine/Source/Runtime/GameplayAbilities/Private/AbilitySystemGlobals.cpp b/Engine/Source/Runtime/GameplayAbilities/Private/AbilitySystemGlobals.cpp
--- a/Engine/Source/Runtime/GameplayAbilities/Private/AbilitySystemGlobals.cpp
+++ b/Engine/Source/Runtime/GameplayAbilities/Private/AbilitySystemGlobals.cpp
@@ -10,7 +10,6 @@
 #include "GameplayCueManager.h"
 #include "GameplayTagResponseTable.h"
 
-
 UAbilitySystemGlobals::UAbilitySystemGlobals(const FObjectInitializer& ObjectInitializer)
 : Super(ObjectInitializer)
 {
@@ -199,6 +198,25 @@ UFunction* UAbilitySystemGlobals::GetGameplayCueFunction(const FGameplayTag& Chi
 
 // --------------------------------------------------------------------
 
+void UAbilitySystemGlobals::InitGameplayCueParameters(FGameplayCueParameters& CueParameters, const FGameplayEffectSpecForRPC &Spec)
+{
+	CueParameters.AggregatedSourceTags = Spec.AggregatedSourceTags;
+	CueParameters.AggregatedTargetTags = Spec.AggregatedTargetTags;
+
+	InitGameplayCueParameters(CueParameters, Spec.GetContext());
+}
+
+void UAbilitySystemGlobals::InitGameplayCueParameters(FGameplayCueParameters& CueParameters, const FGameplayEffectContextHandle& EffectContext)
+{
+	if (EffectContext.IsValid())
+	{
+		// Copy Context over wholesale. Projects may want to override this and not copy over all data
+		CueParameters.EffectContext = EffectContext;
+	}
+}
+
+// --------------------------------------------------------------------
+
 /** Initialize FAttributeSetInitter. This is virtual so projects can override what class they use */
 void UAbilitySystemGlobals::AllocAttributeSetInitter()
 {
diff --git a/Engine/Source/Runtime/GameplayAbilities/Private/AbilitySystemLog.cpp b/Engine/Source/Runtime/GameplayAbilities/Private/AbilitySystemLog.cpp
--- a/Engine/Source/Runtime/GameplayAbilities/Private/AbilitySystemLog.cpp
+++ b/Engine/Source/Runtime/GameplayAbilities/Private/AbilitySystemLog.cpp
@@ -4,82 +4,3 @@
 DEFINE_LOG_CATEGORY(LogAbilitySystem);
 DEFINE_LOG_CATEGORY(VLogAbilitySystem);
 
-AbilitySystemLogScope::~AbilitySystemLogScope()
-{
-	AbilitySystemLog::PopScope();
-}
-
-void AbilitySystemLogScope::Init()
-{
-	PrintedInThisScope = false;
-	AbilitySystemLog::PushScope(this);
-}
-
-// ----------------------------------------------
-
-AbilitySystemLog * GetInstance()
-{
-	static AbilitySystemLog Instance;
-	return &Instance;
-}
-
-FString AbilitySystemLog::Log(ELogVerbosity::Type Verbosity, const FString& LogStr)
-{
-#if !NO_LOGGING
-	if (!LogAbilitySystem.IsSuppressed(Verbosity))
-	{
-		AbilitySystemLog * Instance = GetInstance();
-
-		for (int32 idx=0; idx < Instance->ScopeStack.Num(); ++idx)
-		{
-			AbilitySystemLogScope *Scope = Instance->ScopeStack[idx];
-			if (!Scope->PrintedInThisScope && !Scope->ScopeName.IsEmpty())
-			{
-				if (Instance->NeedNewLine)
-				{
-					UE_LOG(LogAbilitySystem, Log, TEXT(""));
-					Instance->NeedNewLine = false;
-				}
-
-				Scope->PrintedInThisScope = true;
-				int32 ident = (2 * idx);
-				FString IndentStrX = FString::Printf(TEXT("%*s"), ident, TEXT(""));
-				UE_LOG(LogAbilitySystem, Log, TEXT("%s<%s>"), *IndentStrX, *Scope->ScopeName);
-			}
-		}
-
-		FString IndentStr = FString::Printf(TEXT("%*s"), Instance->Indent, TEXT(""));
-		return IndentStr + LogStr;
-	}
-#endif
-	
-	return LogStr;
-}
-
-void AbilitySystemLog::PushScope(AbilitySystemLogScope * Scope)
-{
-	AbilitySystemLog * Instance = GetInstance();
-
-	Instance->Indent += 2;
-	Instance->ScopeStack.Push(Scope);
-}
-
-void AbilitySystemLog::PopScope()
-{
-	AbilitySystemLog * Instance = GetInstance();
-
-	Instance->NeedNewLine = true;
-	Instance->Indent -= 2;
-	check(Instance->Indent >= 0);
-
-	if (Instance->ScopeStack.Top()->PrintedInThisScope)
-	{
-		if (!Instance->ScopeStack.Top()->ScopeName.IsEmpty())
-		{
-			FString IndentStrX = FString::Printf(TEXT("%*s"), Instance->Indent, TEXT(""));
-			UE_LOG(LogAbilitySystem, Log, TEXT("%s</%s>"), *IndentStrX, *Instance->ScopeStack.Top()->ScopeName);
-		}
-	}
-	
-	Instance->ScopeStack.Pop();
-}
\ No newline at end of file
diff --git a/Engine/Source/Runtime/GameplayAbilities/Private/AttributeSet.cpp b/Engine/Source/Runtime/GameplayAbilities/Private/AttributeSet.cpp
--- a/Engine/Source/Runtime/GameplayAbilities/Private/AttributeSet.cpp
+++ b/Engine/Source/Runtime/GameplayAbilities/Private/AttributeSet.cpp
@@ -209,6 +209,7 @@ void FScalableFloat::UnRegisterOnCurveTablePostReimport() const
 	if (GIsEditor && OnCurveTablePostReimportHandle.IsValid())
 	{
 		FReimportManager::Instance()->OnPostReimport().Remove(OnCurveTablePostReimportHandle);
+		OnCurveTablePostReimportHandle.Reset();
 	}
 #endif // WITH_EDITOR
 }
@@ -244,6 +245,21 @@ bool FScalableFloat::operator!=(const FScalableFloat& Other) const
 	return ((Other.Curve != Curve) || (Other.Value != Value));
 }
 
+void FScalableFloat::operator=(const FScalableFloat& Src)
+{
+	Value = Src.Value;
+	Curve = Src.Curve;
+
+#if WITH_EDITOR
+	if (Src.OnCurveTablePostReimportHandle.IsValid())
+	{
+		RegisterOnCurveTablePostReimport();
+	}
+#endif
+	FinalCurve = Src.FinalCurve;
+}
+
+
 // ------------------------------------------------------------------------------------
 //
 // ------------------------------------------------------------------------------------
@@ -317,7 +333,7 @@ void FAttributeSetInitter::PreloadAttributeSetData(UCurveTable* CurveData)
 
 		if (!ensure(!ClassName.IsEmpty() && !SetName.IsEmpty() && !AttributeName.IsEmpty()))
 		{
-			ABILITY_LOG(Warning, TEXT("FAttributeSetInitter::PreloadAttributeSetData Unable to parse row %s in %s"), *RowName, *CurveData->GetName());
+			ABILITY_LOG(Verbose, TEXT("FAttributeSetInitter::PreloadAttributeSetData Unable to parse row %s in %s"), *RowName, *CurveData->GetName());
 			continue;
 		}
 
@@ -336,7 +352,7 @@ void FAttributeSetInitter::PreloadAttributeSetData(UCurveTable* CurveData)
 		UNumericProperty* Property = FindField<UNumericProperty>(*Set, *AttributeName);
 		if (!Property)
 		{
-			ABILITY_LOG(Warning, TEXT("FAttributeSetInitter::PreloadAttributeSetData Unable to match Attribute from %s (row: %s)"), *AttributeName, *RowName);
+			ABILITY_LOG(Verbose, TEXT("FAttributeSetInitter::PreloadAttributeSetData Unable to match Attribute from %s (row: %s)"), *AttributeName, *RowName);
 			continue;
 		}
 
@@ -422,3 +438,50 @@ void FAttributeSetInitter::InitAttributeSetDefaults(UAbilitySystemComponent* Abi
 	
 	AbilitySystemComponent->ForceReplication();
 }
+
+void FAttributeSetInitter::ApplyAttributeDefault(UAbilitySystemComponent* AbilitySystemComponent, FGameplayAttribute& InAttribute, FName GroupName, int32 Level) const
+{
+	SCOPE_CYCLE_COUNTER(STAT_InitAttributeSetDefaults);
+
+	const FAttributeSetDefaulsCollection* Collection = Defaults.Find(GroupName);
+	if (!Collection)
+	{
+		ABILITY_LOG(Warning, TEXT("Unable to find DefaultAttributeSet Group %s. Failing back to Defaults"), *GroupName.ToString());
+		Collection = Defaults.Find(FName(TEXT("Default")));
+		if (!Collection)
+		{
+			ABILITY_LOG(Error, TEXT("FAttributeSetInitter::InitAttributeSetDefaults Default DefaultAttributeSet not found! Skipping Initialization"));
+			return;
+		}
+	}
+
+	if (!Collection->LevelData.IsValidIndex(Level - 1))
+	{
+		// We could eventually extrapolate values outside of the max defined levels
+		ABILITY_LOG(Warning, TEXT("Attribute defaults for Level %d are not defined! Skipping"), Level);
+		return;
+	}
+
+	const FAttributeSetDefaults& SetDefaults = Collection->LevelData[Level - 1];
+	for (const UAttributeSet* Set : AbilitySystemComponent->SpawnedAttributes)
+	{
+		const FAttributeDefaultValueList* DefaultDataList = SetDefaults.DataMap.Find(Set->GetClass());
+		if (DefaultDataList)
+		{
+			ABILITY_LOG(Log, TEXT("Initializing Set %s"), *Set->GetName());
+
+			for (auto& DataPair : DefaultDataList->List)
+			{
+				check(DataPair.Property);
+
+				if (DataPair.Property == InAttribute.GetUProperty())
+				{
+					FGameplayAttribute AttributeToModify(DataPair.Property);
+					AbilitySystemComponent->SetNumericAttributeBase(AttributeToModify, DataPair.Value);
+				}
+			}
+		}
+	}
+
+	AbilitySystemComponent->ForceReplication();
+}
\ No newline at end of file
diff --git a/Engine/Source/Runtime/GameplayAbilities/Private/GameplayAbilityTargetTypes.cpp b/Engine/Source/Runtime/GameplayAbilities/Private/GameplayAbilityTargetTypes.cpp
--- a/Engine/Source/Runtime/GameplayAbilities/Private/GameplayAbilityTargetTypes.cpp
+++ b/Engine/Source/Runtime/GameplayAbilities/Private/GameplayAbilityTargetTypes.cpp
@@ -66,6 +66,12 @@ void FGameplayAbilityTargetData::AddTargetDataToContext(FGameplayEffectContextHa
 	}
 }
 
+void FGameplayAbilityTargetData::AddTargetDataToGameplayCueParameters(FGameplayCueParameters& Parameters)
+{
+	
+	
+}
+
 FString FGameplayAbilityTargetData::ToString() const
 {
 	return TEXT("BASE CLASS");
diff --git a/Engine/Source/Runtime/GameplayAbilities/Private/GameplayCueInterface.cpp b/Engine/Source/Runtime/GameplayAbilities/Private/GameplayCueInterface.cpp
--- a/Engine/Source/Runtime/GameplayAbilities/Private/GameplayCueInterface.cpp
+++ b/Engine/Source/Runtime/GameplayAbilities/Private/GameplayCueInterface.cpp
@@ -192,4 +192,18 @@ void FActiveGameplayCueContainer::PredictiveAdd(const FGameplayTag& Tag, FPredic
 {
 	Owner->UpdateTagMap(Tag, 1);	
 	PredictionKey.NewRejectOrCaughtUpDelegate(FPredictionKeyEvent::CreateUObject(Owner, &UAbilitySystemComponent::RemoveOneTagCount_NoReturn, Tag));
+}
+
+bool FActiveGameplayCueContainer::HasCue(const FGameplayTag& Tag) const
+{
+	for (int32 idx=0; idx < GameplayCues.Num(); ++idx)
+	{
+		const FActiveGameplayCue& Cue = GameplayCues[idx];
+		if (Cue.GameplayCueTag == Tag)
+		{
+			return true;
+		}
+	}
+
+	return false;
 }
\ No newline at end of file
diff --git a/Engine/Source/Runtime/GameplayAbilities/Private/GameplayCueManager.cpp b/Engine/Source/Runtime/GameplayAbilities/Private/GameplayCueManager.cpp
--- a/Engine/Source/Runtime/GameplayAbilities/Private/GameplayCueManager.cpp
+++ b/Engine/Source/Runtime/GameplayAbilities/Private/GameplayCueManager.cpp
@@ -11,6 +11,9 @@
 
 #if WITH_EDITOR
 #include "UnrealEd.h"
+#include "SNotificationList.h"
+#include "NotificationManager.h"
+#define LOCTEXT_NAMESPACE "GameplayCueManager"
 #endif
 
 int32 DisplayGameplayCues = 0;
@@ -153,8 +156,8 @@ AGameplayCueNotify_Actor* UGameplayCueManager::GetInstancedCueActor(AActor* Targ
 
 	AGameplayCueNotify_Actor* CDO = Cast<AGameplayCueNotify_Actor>(CueClass->ClassDefaultObject);
 	FGCNotifyActorKey	NotifyKey(TargetActor, CueClass, 
-							CDO->bUniqueInstancePerInstigator ? Parameters.EffectContext.GetInstigator() : nullptr, 
-							CDO->bUniqueInstancePerSourceObject ? Parameters.EffectContext.GetSourceObject() : nullptr);
+							CDO->bUniqueInstancePerInstigator ? Parameters.GetInstigator() : nullptr, 
+							CDO->bUniqueInstancePerSourceObject ? Parameters.GetSourceObject() : nullptr);
 
 	AGameplayCueNotify_Actor* SpawnedCue = nullptr;
 	if (TWeakObjectPtr<AGameplayCueNotify_Actor>* WeakPtrPtr = NotifyMapActor.Find(NotifyKey))
@@ -211,6 +214,7 @@ void UGameplayCueManager::LoadObjectLibraryFromPaths(const TArray<FString>& InPa
 		FAssetRegistryModule& AssetRegistryModule = FModuleManager::LoadModuleChecked<FAssetRegistryModule>("AssetRegistry");
 		AssetRegistryModule.Get().OnInMemoryAssetCreated().AddUObject(this, &UGameplayCueManager::HandleAssetAdded);
 		AssetRegistryModule.Get().OnInMemoryAssetDeleted().AddUObject(this, &UGameplayCueManager::HandleAssetDeleted);
+		AssetRegistryModule.Get().OnAssetRenamed().AddUObject(this, &UGameplayCueManager::HandleAssetRenamed);
 		FWorldDelegates::OnPreWorldInitialization.AddUObject(this, &UGameplayCueManager::ReloadObjectLibrary);
 		RegisteredEditorCallbacks = true;
 	}
@@ -376,6 +380,12 @@ void UGameplayCueManager::HandleAssetAdded(UObject *Object)
 				GlobalCueSet->AddCues(CuesToAdd);
 
 				OnGameplayCueNotifyAddOrRemove.Broadcast();
+
+			
+			}
+			else
+			{
+				VerifyNotifyAssetIsInValidPath(Blueprint->GetOuter()->GetPathName());
 			}
 		}
 	}
@@ -408,6 +418,59 @@ void UGameplayCueManager::HandleAssetDeleted(UObject *Object)
 	}
 }
 
+/** Handles cleaning up an object library if it matches the passed in object */
+void UGameplayCueManager::HandleAssetRenamed(const FAssetData& Data, const FString& String)
+{
+	const FString* ParentClassNamePtr = Data.TagsAndValues.Find(TEXT("ParentClass"));
+	if (ParentClassNamePtr)
+	{
+		FString ParentClassName = *ParentClassNamePtr;
+		
+		UClass* DataClass = FindObject<UClass>(nullptr, *ParentClassName);
+		if (DataClass)
+		{
+			UGameplayCueNotify_Static* StaticCDO = Cast<UGameplayCueNotify_Static>(DataClass->ClassDefaultObject);
+			AGameplayCueNotify_Actor* ActorCDO = Cast<AGameplayCueNotify_Actor>(DataClass->ClassDefaultObject);
+			if (StaticCDO || ActorCDO)
+			{
+				VerifyNotifyAssetIsInValidPath(Data.PackagePath.ToString());
+			}
+		}
+	}
+}
+
+void UGameplayCueManager::VerifyNotifyAssetIsInValidPath(FString Path)
+{
+	bool ValidPath = false;
+	for (FString& str: LoadedPaths)
+	{
+		if (Path.Contains(str))
+		{
+			ValidPath = true;
+		}
+	}
+
+	if (!ValidPath)
+	{
+		FString MessageTry = FString::Printf(TEXT("Warning: Invalid GameplayCue Path %s"));
+		MessageTry += TEXT("\n\nGameplayCue Notifies should only be saved in the following folders:");
+
+		ABILITY_LOG(Warning, TEXT("Warning: Invalid GameplayCuePath: %s"), *Path);
+		ABILITY_LOG(Warning, TEXT("Valid Paths: "));
+		for (FString& str: LoadedPaths)
+		{
+			ABILITY_LOG(Warning, TEXT("  %s"), *str);
+			MessageTry += FString::Printf(TEXT("\n  %s"), *str);
+		}
+
+		MessageTry += FString::Printf(TEXT("\n\nThis asset must be moved to a valid location to work in game."));
+
+		const FText MessageText = FText::FromString(MessageTry);
+		const FText TitleText = NSLOCTEXT("GameplayCuePathWarning", "GameplayCuePathWarningTitle", "Invalid GameplayCue Path");
+		FMessageDialog::Open(EAppMsgType::Ok, MessageText, &TitleText);
+	}
+}
+
 #endif
 
 
@@ -470,10 +533,10 @@ void UGameplayCueManager::InvokeGameplayCueExecuted_FromSpec(UAbilitySystemCompo
 void UGameplayCueManager::InvokeGameplayCueExecuted(UAbilitySystemComponent* OwningComponent, const FGameplayTag GameplayCueTag, FPredictionKey PredictionKey, FGameplayEffectContextHandle EffectContext)
 {
 	FGameplayCuePendingExecute PendingCue;
-	PendingCue.PayloadType = EGameplayCuePayloadType::EffectContext;
+	PendingCue.PayloadType = EGameplayCuePayloadType::CueParameters;
 	PendingCue.GameplayCueTag = GameplayCueTag;
 	PendingCue.OwningComponent = OwningComponent;
-	PendingCue.CueParameters.EffectContext = EffectContext;
+	UAbilitySystemGlobals::Get().InitGameplayCueParameters(PendingCue.CueParameters, EffectContext);
 	PendingCue.PredictionKey = PredictionKey;
 
 	if (ProcessPendingCueExecute(PendingCue))
@@ -633,3 +696,6 @@ bool UGameplayCueManager::DoesPendingCueExecuteMatch(FGameplayCuePendingExecute&
 	return true;
 }
 
+#if WITH_EDITOR
+#undef LOCTEXT_NAMESPACE
+#endif
diff --git a/Engine/Source/Runtime/GameplayAbilities/Private/GameplayEffect.cpp b/Engine/Source/Runtime/GameplayAbilities/Private/GameplayEffect.cpp
--- a/Engine/Source/Runtime/GameplayAbilities/Private/GameplayEffect.cpp
+++ b/Engine/Source/Runtime/GameplayAbilities/Private/GameplayEffect.cpp
@@ -530,12 +530,24 @@ FGameplayEffectSpecForRPC::FGameplayEffectSpecForRPC()
 
 FGameplayEffectSpecForRPC::FGameplayEffectSpecForRPC(const FGameplayEffectSpec& InSpec)
 	: Def(InSpec.Def)
-	, ModifiedAttributes(MoveTemp(InSpec.ModifiedAttributes))
+	, ModifiedAttributes()
 	, EffectContext(InSpec.GetEffectContext())
 	, AggregatedSourceTags(*InSpec.CapturedSourceTags.GetAggregatedTags())
 	, AggregatedTargetTags(*InSpec.CapturedTargetTags.GetAggregatedTags())
 	, Level(InSpec.GetLevel())
 {
+	
+	// Only copy attributes that are in the gameplay cue info
+	for (int32 i = InSpec.ModifiedAttributes.Num() - 1; i >= 0; i--)
+	{
+		for (FGameplayEffectCue CueInfo : Def->GameplayCues)
+		{
+			if (CueInfo.MagnitudeAttribute == InSpec.ModifiedAttributes[i].Attribute)
+			{
+				ModifiedAttributes.Add(InSpec.ModifiedAttributes[i]);
+			}
+		}
+	}
 }
 
 void FGameplayEffectSpec::SetupAttributeCaptureDefinitions()
@@ -763,28 +775,6 @@ FGameplayEffectModifiedAttribute* FGameplayEffectSpec::AddModifiedAttribute(cons
 	return &ModifiedAttributes[ModifiedAttributes.Add(NewAttribute)];
 }
 
-void FGameplayEffectSpec::PruneModifiedAttributes()
-{
-	TSet<FGameplayAttribute> ImportantAttributes;
-
-	for (FGameplayEffectCue CueInfo : Def->GameplayCues)
-	{
-		if (CueInfo.MagnitudeAttribute.IsValid())
-		{
-			ImportantAttributes.Add(CueInfo.MagnitudeAttribute);
-		}
-	}
-
-	// Remove all unimportant attributes
-	for (int32 i = ModifiedAttributes.Num() - 1; i >= 0; i--)
-	{
-		if (!ImportantAttributes.Contains(ModifiedAttributes[i].Attribute))
-		{
-			ModifiedAttributes.RemoveAtSwap(i);
-		}
-	}
-}
-
 void FGameplayEffectSpec::SetContext(FGameplayEffectContextHandle NewEffectContext)
 {
 	bool bWasAlreadyInit = EffectContext.IsValid();
@@ -1423,9 +1413,6 @@ void FActiveGameplayEffectsContainer::ExecuteActiveEffectsFrom(FGameplayEffectSp
 	//	Invoke GameplayCue events
 	// ------------------------------------------------------
 	
-	// Prune the modified attributes before we replicate
-	SpecToUse.PruneModifiedAttributes();
-
 	// If there are no modifiers or we don't require modifier success to trigger, we apply the GameplayCue. 
 	bool InvokeGameplayCueExecute = (SpecToUse.Modifiers.Num() == 0) || !Spec.Def->bRequireModifierSuccessToTriggerCues;
 
@@ -1483,6 +1470,14 @@ void FActiveGameplayEffectsContainer::ExecutePeriodicGameplayEffect(FActiveGamep
 
 		// Execute
 		ExecuteActiveEffectsFrom(ActiveEffect->Spec);
+
+		// Invoke Delegates for periodic effects being executed
+		UAbilitySystemComponent* SourceASC = ActiveEffect->Spec.GetContext().GetInstigatorAbilitySystemComponent();
+		Owner->OnPeriodicGameplayEffectExecuteOnSelf(SourceASC, ActiveEffect->Spec, Handle);
+		if (SourceASC)
+		{
+			SourceASC->OnPeriodicGameplayEffectExecuteOnTarget(Owner, ActiveEffect->Spec, Handle);
+		}
 	}
 
 }
@@ -1535,7 +1530,6 @@ FAggregatorRef& FActiveGameplayEffectsContainer::FindOrCreateAttributeAggregator
 
 void FActiveGameplayEffectsContainer::OnAttributeAggregatorDirty(FAggregator* Aggregator, FGameplayAttribute Attribute)
 {
-	ABILITY_LOG_SCOPE(TEXT("FActiveGameplayEffectsContainer::OnAttributeAggregatorDirty"));
 	check(AttributeAggregatorMap.FindChecked(Attribute).Get() == Aggregator);
 
 	// Our Aggregator has changed, we need to reevaluate this aggregator and update the current value of the attribute.
@@ -1746,21 +1740,6 @@ bool FActiveGameplayEffectsContainer::HandleActiveGameplayEffectStackOverflow(co
 	return bAllowOverflowApplication;
 }
 
-void FActiveGameplayEffectsContainer::ApplyStackingLogicPostApplyAsSource(UAbilitySystemComponent* Target, const FGameplayEffectSpec& SpecApplied, FActiveGameplayEffectHandle ActiveHandle)
-{
-	if (SpecApplied.Def->StackingType == EGameplayEffectStackingType::AggregateBySource)
-	{
-		TArray<FActiveGameplayEffectHandle>& ActiveHandles = SourceStackingMap.FindOrAdd(SpecApplied.Def);
-
-		// This is probably wrong!
-		if (ActiveHandles.Num() == SpecApplied.Def->StackLimitCount)
-		{
-			// We are at the limit, so replace one based on policy
-			// For now, just always remove the oldest one applied
-		}
-	}
-}
-
 void FActiveGameplayEffectsContainer::SetBaseAttributeValueFromReplication(FGameplayAttribute Attribute, float ServerValue)
 {
 	FAggregatorRef* RefPtr = AttributeAggregatorMap.Find(Attribute);
@@ -1930,8 +1909,6 @@ bool FActiveGameplayEffectsContainer::InternalExecuteMod(FGameplayEffectSpec& Sp
 
 	if (AttributeSet)
 	{
-		ABILITY_LOG_SCOPE(TEXT("Executing Attribute Mod %s"), *ModEvalData.ToSimpleString());
-
 		FGameplayEffectModCallbackData ExecuteData(Spec, ModEvalData, *Owner);
 
 		/**
@@ -2191,9 +2168,6 @@ FActiveGameplayEffect* FActiveGameplayEffectsContainer::ApplyGameplayEffectSpec(
 				}
 				ModifiedAttribute->TotalMagnitude += Magnitude;
 			}
-
-			// Prune list since ModifiedAttributes have only been built to
-			AppliedEffectSpec.PruneModifiedAttributes();
 		}
 	}
 
@@ -2346,6 +2320,21 @@ void FActiveGameplayEffectsContainer::AddActiveGameplayEffectGrantedTagsAndModif
 	ApplicationImmunityGameplayTagCountContainer.UpdateTagCount(Effect.Spec.Def->GrantedApplicationImmunityTags.RequireTags, 1);
 	ApplicationImmunityGameplayTagCountContainer.UpdateTagCount(Effect.Spec.Def->GrantedApplicationImmunityTags.IgnoreTags, 1);
 
+	// Grant abilities
+	if (IsNetAuthority())
+	{
+		for (FGameplayAbilitySpecDef& AbilitySpecDef : Effect.Spec.GrantedAbilitySpecs)
+		{
+			// Only do this if we haven't assigned the ability yet! This prevents cases where stacking GEs
+			// would regrant the ability every time the stack was applied
+			if (AbilitySpecDef.AssignedHandle.IsValid() == false)
+			{
+				Owner->GiveAbility( FGameplayAbilitySpec(AbilitySpecDef, Effect.Handle) );
+			}
+		}	
+	}
+
+	// Update GameplayCue tags and events
 	for (const FGameplayEffectCue& Cue : Effect.Spec.Def->GameplayCues)
 	{
 		Owner->UpdateTagMap(Cue.GameplayCueTags, 1);
@@ -2413,23 +2402,6 @@ bool FActiveGameplayEffectsContainer::InternalRemoveActiveGameplayEffect(int32 I
 		// Mark the effect as pending removal
 		Effect.IsPendingRemove = true;
 
-		// Remove Granted Abilities
-		for (FGameplayAbilitySpecDef& AbilitySpecDef : Effect.Spec.GrantedAbilitySpecs)
-		{
-			if (AbilitySpecDef.AssignedHandle.IsValid())
-			{
-				switch(AbilitySpecDef.RemovalPolicy)
-				{
-				case EGameplayEffectGrantedAbilityRemovePolicy::CancelAbilityImmediately:
-					Owner->ClearAbility(AbilitySpecDef.AssignedHandle);
-					break;
-				case EGameplayEffectGrantedAbilityRemovePolicy::RemoveAbilityOnEnd:
-					Owner->SetRemoveAbilityOnEnd(AbilitySpecDef.AssignedHandle);
-					break;
-				}
-			}
-		}
-
 		// Invoke Remove GameplayCue event
 		bool ShouldInvokeGameplayCueEvent = true;
 		const bool bIsNetAuthority = IsNetAuthority();
@@ -2551,6 +2523,36 @@ void FActiveGameplayEffectsContainer::RemoveActiveGameplayEffectGrantedTagsAndMo
 
 	Owner->UpdateTagMap(Effect.Spec.DynamicGrantedTags, -1);
 
+	// Cancel/remove granted abilities
+	if (IsNetAuthority())
+	{
+		for (const FGameplayAbilitySpecDef& AbilitySpecDef : Effect.Spec.GrantedAbilitySpecs)
+		{
+			if (AbilitySpecDef.AssignedHandle.IsValid())
+			{
+				switch(AbilitySpecDef.RemovalPolicy)
+				{
+				case EGameplayEffectGrantedAbilityRemovePolicy::CancelAbilityImmediately:
+					{
+						Owner->ClearAbility(AbilitySpecDef.AssignedHandle);
+						break;
+					}
+				case EGameplayEffectGrantedAbilityRemovePolicy::RemoveAbilityOnEnd:
+					{
+						Owner->SetRemoveAbilityOnEnd(AbilitySpecDef.AssignedHandle);
+						break;
+					}
+				default:
+					{
+						// Do nothing to granted ability
+						break;
+					}
+				}
+			}
+		}
+	}
+
+	// Update GameplayCue tags and events
 	for (const FGameplayEffectCue& Cue : Effect.Spec.Def->GameplayCues)
 	{
 		if (bInvokeGameplayCueEvents)
@@ -2688,21 +2690,24 @@ bool FActiveGameplayEffectsContainer::NetDeltaSerialize(FNetDeltaSerializeInfo&
 
 		for (const FActiveGameplayEffect& Effect : this)
 		{
-			if (Effect.bIsInhibited == false)
+			if (!DeltaParms.bOutHasMoreUnmapped) // Do not invoke GCs when we have missing information (like AActor*s in EffectContext)
 			{
-				if (Effect.bPendingRepOnActiveGC)
+				if (Effect.bIsInhibited == false)
 				{
-					Owner->InvokeGameplayCueEvent(Effect.Spec, EGameplayCueEvent::OnActive);
-					
-				}
-				if (Effect.bPendingRepWhileActiveGC)
-				{
-					Owner->InvokeGameplayCueEvent(Effect.Spec, EGameplayCueEvent::WhileActive);
+					if (Effect.bPendingRepOnActiveGC)
+					{
+						Owner->InvokeGameplayCueEvent(Effect.Spec, EGameplayCueEvent::OnActive);
+						
+					}
+					if (Effect.bPendingRepWhileActiveGC)
+					{
+						Owner->InvokeGameplayCueEvent(Effect.Spec, EGameplayCueEvent::WhileActive);
+					}
 				}
-			}
 
-			Effect.bPendingRepOnActiveGC = false;
-			Effect.bPendingRepWhileActiveGC = false;
+				Effect.bPendingRepOnActiveGC = false;
+				Effect.bPendingRepWhileActiveGC = false;
+			}
 		}
 	}
 
diff --git a/Engine/Source/Runtime/GameplayAbilities/Private/GameplayEffectExecutionCalculation.cpp b/Engine/Source/Runtime/GameplayAbilities/Private/GameplayEffectExecutionCalculation.cpp
--- a/Engine/Source/Runtime/GameplayAbilities/Private/GameplayEffectExecutionCalculation.cpp
+++ b/Engine/Source/Runtime/GameplayAbilities/Private/GameplayEffectExecutionCalculation.cpp
@@ -52,6 +52,12 @@ const FGameplayEffectSpec& FGameplayEffectCustomExecutionParameters::GetOwningSp
 	return *OwningSpec;
 }
 
+FGameplayEffectSpec* FGameplayEffectCustomExecutionParameters::GetOwningSpecForPreExecuteMod() const
+{
+	check(OwningSpec);
+	return OwningSpec;
+}
+
 UAbilitySystemComponent* FGameplayEffectCustomExecutionParameters::GetTargetAbilitySystemComponent() const
 {
 	return TargetAbilitySystemComponent.Get();
diff --git a/Engine/Source/Runtime/GameplayAbilities/Private/GameplayEffectTests.cpp b/Engine/Source/Runtime/GameplayAbilities/Private/GameplayEffectTests.cpp
--- a/Engine/Source/Runtime/GameplayAbilities/Private/GameplayEffectTests.cpp
+++ b/Engine/Source/Runtime/GameplayAbilities/Private/GameplayEffectTests.cpp
@@ -110,7 +110,6 @@ public: // the tests
 
 		// just try and reduce the health attribute
 		{
-			ABILITY_LOG_SCOPE(TEXT("Apply InstantDamage"));
 			
 			CONSTRUCT_CLASS(UGameplayEffect, BaseDmgEffect);
 			AddModifier(BaseDmgEffect, GET_FIELD_CHECKED(UAbilitySystemTestAttributeSet, Health), EGameplayModOp::Additive, FScalableFloat(-DamageValue));
@@ -130,8 +129,6 @@ public: // the tests
 
 		// This is the same as GameplayEffectsTest_InstantDamage but modifies the Damage attribute and confirms it is remapped to -Health by UAbilitySystemTestAttributeSet::PostAttributeModify
 		{
-			ABILITY_LOG_SCOPE(TEXT("Apply InstantDamage"));
-
 			CONSTRUCT_CLASS(UGameplayEffect, BaseDmgEffect);
 			AddModifier(BaseDmgEffect, GET_FIELD_CHECKED(UAbilitySystemTestAttributeSet, Damage), EGameplayModOp::Additive, FScalableFloat(DamageValue));
 			BaseDmgEffect->DurationPolicy = EGameplayEffectDurationType::Instant;
@@ -155,8 +152,6 @@ public: // the tests
 
 		// apply the buff
 		{
-			ABILITY_LOG_SCOPE(TEXT("Apply Buff"));
-
 			CONSTRUCT_CLASS(UGameplayEffect, DamageBuffEffect);
 			AddModifier(DamageBuffEffect, GET_FIELD_CHECKED(UAbilitySystemTestAttributeSet, Mana), EGameplayModOp::Additive, FScalableFloat(BuffValue));
 			DamageBuffEffect->DurationPolicy = EGameplayEffectDurationType::Infinite;
@@ -169,8 +164,6 @@ public: // the tests
 
 		// remove the effect
 		{
-			ABILITY_LOG_SCOPE(TEXT("Remove Buff"));
-
 			DestComponent->RemoveActiveGameplayEffect(BuffHandle);
 		}
 
diff --git a/Engine/Source/Runtime/GameplayAbilities/Private/GameplayEffectTypes.cpp b/Engine/Source/Runtime/GameplayAbilities/Private/GameplayEffectTypes.cpp
--- a/Engine/Source/Runtime/GameplayAbilities/Private/GameplayEffectTypes.cpp
+++ b/Engine/Source/Runtime/GameplayAbilities/Private/GameplayEffectTypes.cpp
@@ -462,8 +462,6 @@ FString FGameplayTagRequirements::ToString() const
 
 void FActiveGameplayEffectsContainer::PrintAllGameplayEffects() const
 {
-	ABILITY_LOG_SCOPE(TEXT("ActiveGameplayEffects. Num: %d"), GetNumGameplayEffects());
-
 	for (const FActiveGameplayEffect& Effect : this)
 	{
 		Effect.PrintAll();
@@ -479,7 +477,6 @@ void FActiveGameplayEffect::PrintAll() const
 
 void FGameplayEffectSpec::PrintAll() const
 {
-	ABILITY_LOG_SCOPE(TEXT("GameplayEffectSpec"));
 	ABILITY_LOG(Log, TEXT("Def: %s"), *Def->GetName());
 
 	ABILITY_LOG(Log, TEXT("Duration: %.2f"), GetDuration());
@@ -536,4 +533,223 @@ FGameplayEffectSpecHandle::FGameplayEffectSpecHandle(FGameplayEffectSpec* DataPt
 	: Data(DataPtr)
 {
 
+}
+
+FGameplayCueParameters::FGameplayCueParameters(const FGameplayEffectSpecForRPC& Spec)
+: NormalizedMagnitude(0.0f)
+, RawMagnitude(0.0f)
+, MatchedTagName(NAME_None)
+, Location(ForceInitToZero)
+, Normal(ForceInitToZero)
+{
+	UAbilitySystemGlobals::Get().InitGameplayCueParameters(*this, Spec);
+}
+
+FGameplayCueParameters::FGameplayCueParameters(const struct FGameplayEffectContextHandle& InEffectContext)
+: NormalizedMagnitude(0.0f)
+, RawMagnitude(0.0f)
+, MatchedTagName(NAME_None)
+, Location(ForceInitToZero)
+, Normal(ForceInitToZero)
+{
+	UAbilitySystemGlobals::Get().InitGameplayCueParameters(*this, InEffectContext);
+}
+
+bool FGameplayCueParameters::NetSerialize(FArchive& Ar, class UPackageMap* Map, bool& bOutSuccess)
+{
+	enum RepFlag
+	{
+		REP_NormalizedMagnitude = 0,
+		REP_RawMagnitude,
+		REP_EffectContext,
+		REP_MatchedTagName,
+		REP_OriginalTag,
+		REP_AggregatedSourceTags,
+		REP_AggregatedTargetTags,
+		REP_Location,
+		REP_Normal,
+		REP_Instigator,
+		REP_EffectCauser,
+		REP_SourceObject,
+		
+		REP_MAX
+	};
+
+	uint16 RepBits = 0;
+	if (Ar.IsSaving())
+	{
+		if (NormalizedMagnitude != 0.f)
+		{
+			RepBits |= (1 << REP_NormalizedMagnitude);
+		}
+		if (RawMagnitude != 0.f)
+		{
+			RepBits |= (1 << REP_RawMagnitude);
+		}
+		if (EffectContext.IsValid())
+		{
+			RepBits |= (1 << REP_EffectContext);
+		}
+		if (MatchedTagName != NAME_None)
+		{
+			RepBits |= (1 << REP_MatchedTagName);
+		}
+		if (OriginalTag.IsValid())
+		{
+			RepBits |= (1 << REP_OriginalTag);
+		}
+		if (AggregatedSourceTags.Num() > 0)
+		{
+			RepBits |= (1 << REP_AggregatedSourceTags);
+		}
+		if (AggregatedTargetTags.Num() > 0)
+		{
+			RepBits |= (1 << REP_AggregatedTargetTags);
+		}
+		if (Location.IsNearlyZero() == false)
+		{
+			RepBits |= (1 << REP_Location);
+		}
+		if (Normal.IsNearlyZero() == false)
+		{
+			RepBits |= (1 << REP_Normal);
+		}
+		if (Instigator.IsValid())
+		{
+			RepBits |= (1 << REP_Instigator);
+		}
+		if (EffectCauser.IsValid())
+		{
+			RepBits |= (1 << REP_EffectCauser);
+		}
+		if (SourceObject.IsValid())
+		{
+			RepBits |= (1 << REP_SourceObject);
+		}
+	}
+
+	Ar.SerializeBits(&RepBits, REP_MAX);
+
+	if (RepBits & (1 << REP_NormalizedMagnitude))
+	{
+		Ar << NormalizedMagnitude;
+	}
+	if (RepBits & (1 << REP_RawMagnitude))
+	{
+		Ar << RawMagnitude;
+	}
+	if (RepBits & (1 << REP_EffectContext))
+	{
+		EffectContext.NetSerialize(Ar, Map, bOutSuccess);
+	}
+	if (RepBits & (1 << REP_MatchedTagName))
+	{
+		Ar << MatchedTagName;
+	}
+	if (RepBits & (1 << REP_OriginalTag))
+	{
+		Ar << OriginalTag;
+	}
+	if (RepBits & (1 << REP_AggregatedSourceTags))
+	{
+		AggregatedSourceTags.NetSerialize(Ar, Map, bOutSuccess);
+	}
+	if (RepBits & (1 << REP_AggregatedTargetTags))
+	{
+		AggregatedTargetTags.NetSerialize(Ar, Map, bOutSuccess);
+	}
+	if (RepBits & (1 << REP_Location))
+	{
+		Location.NetSerialize(Ar, Map, bOutSuccess);
+	}
+	if (RepBits & (1 << REP_Normal))
+	{
+		Normal.NetSerialize(Ar, Map, bOutSuccess);
+	}
+	if (RepBits & (1 << REP_Instigator))
+	{
+		Ar << Instigator;
+	}
+	if (RepBits & (1 << REP_EffectCauser))
+	{
+		Ar << EffectCauser;
+	}
+	if (RepBits & (1 << REP_SourceObject))
+	{
+		Ar << SourceObject;
+	}
+
+	bOutSuccess = true;
+	return true;
+}
+
+bool FGameplayCueParameters::IsInstigatorLocallyControlled() const
+{
+	if (EffectContext.IsValid())
+	{
+		return EffectContext.IsLocallyControlled();
+	}
+
+	APawn* Pawn = Cast<APawn>(Instigator.Get());
+	if (!Pawn)
+	{
+		Pawn = Cast<APawn>(EffectCauser.Get());
+	}
+	if (Pawn)
+	{
+		return Pawn->IsLocallyControlled();
+	}
+	return false;
+}
+
+bool FGameplayCueParameters::IsInstigatorLocallyControlledPlayer() const
+{
+	if (EffectContext.IsValid())
+	{
+		return EffectContext.IsLocallyControlledPlayer();
+	}
+
+	APawn* Pawn = Cast<APawn>(Instigator.Get());
+	if (!Pawn)
+	{
+		Pawn = Cast<APawn>(EffectCauser.Get());
+	}
+	if (Pawn && Pawn->Controller)
+	{
+		return Pawn->Controller->IsLocalPlayerController();
+	}
+	return false;
+}
+
+AActor* FGameplayCueParameters::GetInstigator() const
+{
+	if (Instigator.IsValid())
+	{
+		return Instigator.Get();
+	}
+
+	// Fallback to effect context if the explicit data on gameplaycue parameters is not there.
+	return EffectContext.GetInstigator();
+}
+
+AActor* FGameplayCueParameters::GetEffectCauser() const
+{
+	if (EffectCauser.IsValid())
+	{
+		return EffectCauser.Get();
+	}
+
+	// Fallback to effect context if the explicit data on gameplaycue parameters is not there.
+	return EffectContext.GetEffectCauser();
+}
+
+const UObject* FGameplayCueParameters::GetSourceObject() const
+{
+	if (SourceObject.IsValid())
+	{
+		return SourceObject.Get();
+	}
+
+	// Fallback to effect context if the explicit data on gameplaycue parameters is not there.
+	return EffectContext.GetSourceObject();
 }
\ No newline at end of file
diff --git a/Engine/Source/Runtime/GameplayAbilities/Private/GameplayTagResponseTable.cpp b/Engine/Source/Runtime/GameplayAbilities/Private/GameplayTagResponseTable.cpp
--- a/Engine/Source/Runtime/GameplayAbilities/Private/GameplayTagResponseTable.cpp
+++ b/Engine/Source/Runtime/GameplayAbilities/Private/GameplayTagResponseTable.cpp
@@ -1,6 +1,7 @@
 // Copyright 1998-2015 Epic Games, Inc. All Rights Reserved.
 
 #include "AbilitySystemPrivatePCH.h"
+#include "AbilitySystemComponent.h"
 #include "GameplayTagResponseTable.h"
 
 // --------------------------------------------------------------------------------------------------------------------------------------------------------
@@ -13,6 +14,10 @@
 UGameplayTagReponseTable::UGameplayTagReponseTable(const FObjectInitializer& ObjectInitializer)
 	: Super(ObjectInitializer)
 {
+	// Make an empty tag query. We will actually fill the tag out prior to evaluating the query with ::MakeQuery
+	FGameplayTagContainer Container;
+	Container.AddTagFast(FGameplayTag());
+	Query = FGameplayEffectQuery::MakeQuery_MatchAllOwningTags(Container);
 }
 
 void UGameplayTagReponseTable::RegisterResponseForEvents(UAbilitySystemComponent* ASC)
@@ -22,7 +27,10 @@ void UGameplayTagReponseTable::RegisterResponseForEvents(UAbilitySystemComponent
 		return;
 	}
 
-	RegisteredASCs.Add(ASC, FGameplayTagResponseAppliedInfo() );
+	TArray<FGameplayTagResponseAppliedInfo> AppliedInfoList;
+	AppliedInfoList.SetNum(Entries.Num());
+
+	RegisteredASCs.Add(ASC, AppliedInfoList );
 
 	for (int32 idx=0; idx < Entries.Num(); ++idx)
 	{
@@ -51,11 +59,14 @@ void UGameplayTagReponseTable::TagResponseEvent(const FGameplayTag Tag, int32 Ne
 	{
 		QUICK_SCOPE_CYCLE_COUNTER(ABILITY_TRT_CALC_COUNT);
 
-		TotalCount += ASC->GetAggregatedStackCount(MakeQuery(Entry.Positive.Tag));
-		TotalCount -= ASC->GetAggregatedStackCount(MakeQuery(Entry.Negative.Tag));
+		int32 Positive = GetCount(Entry.Positive, ASC);
+		int32 Negative = GetCount(Entry.Negative, ASC);
+
+		TotalCount = Positive - Negative;
 	}
 
-	FGameplayTagResponseAppliedInfo& Info = RegisteredASCs.FindChecked(ASC);
+	TArray<FGameplayTagResponseAppliedInfo>& InfoList = RegisteredASCs.FindChecked(ASC);
+	FGameplayTagResponseAppliedInfo& Info = InfoList[idx];
 
 	if (TotalCount < 0)
 	{
@@ -72,4 +83,43 @@ void UGameplayTagReponseTable::TagResponseEvent(const FGameplayTag Tag, int32 Ne
 		Remove(ASC, Info.PositiveHandle);
 		Remove(ASC, Info.NegativeHandle);
 	}
-}
\ No newline at end of file
+}
+
+int32 UGameplayTagReponseTable::GetCount(const FGameplayTagReponsePair& Pair, UAbilitySystemComponent* ASC) const
+{
+	int32 Count=0;
+	if (Pair.Tag.IsValid())
+	{
+		Count = ASC->GetAggregatedStackCount(MakeQuery(Pair.Tag));
+		if (Pair.SoftCountCap > 0)
+		{
+			Count = FMath::Min<int32>(Pair.SoftCountCap, Count);
+		}
+	}
+
+	return Count;
+}
+
+void UGameplayTagReponseTable::Remove(UAbilitySystemComponent* ASC, FActiveGameplayEffectHandle& Handle)
+{
+	if (Handle.IsValid())
+	{
+		ASC->RemoveActiveGameplayEffect(Handle);
+		Handle.Invalidate();
+	}
+}
+
+void UGameplayTagReponseTable::AddOrUpdate(UAbilitySystemComponent* ASC, const TSubclassOf<UGameplayEffect>& ResponseGameplayEffect, int32 TotalCount, FActiveGameplayEffectHandle& Handle)
+{
+	if (ResponseGameplayEffect)
+	{
+		if (Handle.IsValid())
+		{
+			ASC->SetActiveGameplayEffectLevel(Handle, TotalCount);
+		}
+		else
+		{
+			Handle = ASC->ApplyGameplayEffectToSelf(Cast<UGameplayEffect>(ResponseGameplayEffect->ClassDefaultObject), TotalCount, ASC->GetEffectContext());
+		}
+	}
+}
diff --git a/Engine/Source/Runtime/GameplayAbilities/Public/Abilities/GameplayAbility.h b/Engine/Source/Runtime/GameplayAbilities/Public/Abilities/GameplayAbility.h
--- a/Engine/Source/Runtime/GameplayAbilities/Public/Abilities/GameplayAbility.h
+++ b/Engine/Source/Runtime/GameplayAbilities/Public/Abilities/GameplayAbility.h
@@ -192,6 +192,10 @@ public:
 	/** Retrieves the actual AbilitySpec for this ability. Can only be called on instanced abilities. */
 	FGameplayAbilitySpec* GetCurrentAbilitySpec() const;
 
+	/** Retrieves the EffectContext of the GameplayEffect that granted this ability. Can only be called on instanced abilities. */
+	UFUNCTION(BlueprintCallable, Category = Ability)
+	FGameplayEffectContextHandle GetGrantedByEffectContext() const;
+
 	/** Returns an effect context, given a specified actor info */
 	virtual FGameplayEffectContextHandle GetEffectContext(const FGameplayAbilitySpecHandle Handle, const FGameplayAbilityActorInfo *ActorInfo) const;
 
diff --git a/Engine/Source/Runtime/GameplayAbilities/Public/Abilities/GameplayAbilityTargetActor.h b/Engine/Source/Runtime/GameplayAbilities/Public/Abilities/GameplayAbilityTargetActor.h
--- a/Engine/Source/Runtime/GameplayAbilities/Public/Abilities/GameplayAbilityTargetActor.h
+++ b/Engine/Source/Runtime/GameplayAbilities/Public/Abilities/GameplayAbilityTargetActor.h
@@ -97,4 +97,7 @@ public:
 
 	FDelegateHandle GenericConfirmHandle;
 	FDelegateHandle GenericCancelHandle;
+
+	UPROPERTY()
+	UAbilitySystemComponent* GenericDelegateBoundASC;
 };
\ No newline at end of file
diff --git a/Engine/Source/Runtime/GameplayAbilities/Public/Abilities/GameplayAbilityTargetTypes.h b/Engine/Source/Runtime/GameplayAbilities/Public/Abilities/GameplayAbilityTargetTypes.h
--- a/Engine/Source/Runtime/GameplayAbilities/Public/Abilities/GameplayAbilityTargetTypes.h
+++ b/Engine/Source/Runtime/GameplayAbilities/Public/Abilities/GameplayAbilityTargetTypes.h
@@ -69,6 +69,8 @@ struct GAMEPLAYABILITIES_API FGameplayAbilityTargetData
 
 	virtual void AddTargetDataToContext(FGameplayEffectContextHandle& Context, bool bIncludeActorArray);
 
+	virtual void AddTargetDataToGameplayCueParameters(FGameplayCueParameters& Parameters);
+
 	virtual TArray<TWeakObjectPtr<AActor> >	GetActors() const
 	{
 		return TArray<TWeakObjectPtr<AActor> >();
diff --git a/Engine/Source/Runtime/GameplayAbilities/Public/Abilities/GameplayAbilityTypes.h b/Engine/Source/Runtime/GameplayAbilities/Public/Abilities/GameplayAbilityTypes.h
--- a/Engine/Source/Runtime/GameplayAbilities/Public/Abilities/GameplayAbilityTypes.h
+++ b/Engine/Source/Runtime/GameplayAbilities/Public/Abilities/GameplayAbilityTypes.h
@@ -331,6 +331,9 @@ struct GAMEPLAYABILITIES_API FGameplayEventData
 	float EventMagnitude;
 };
 
+/** Used for cleaning up predicted data on network clients */
+DECLARE_MULTICAST_DELEGATE_OneParam(FGameplayEventMulticastDelegate, const FGameplayEventData*);
+
 /** 
  *	Structure that tells AbilitySystemComponent what to bind to an InputComponent (see BindAbilityActivationToInputComponent) 
  *	
diff --git a/Engine/Source/Runtime/GameplayAbilities/Public/Abilities/Tasks/AbilityTask.h b/Engine/Source/Runtime/GameplayAbilities/Public/Abilities/Tasks/AbilityTask.h
--- a/Engine/Source/Runtime/GameplayAbilities/Public/Abilities/Tasks/AbilityTask.h
+++ b/Engine/Source/Runtime/GameplayAbilities/Public/Abilities/Tasks/AbilityTask.h
@@ -119,3 +119,34 @@ protected:
 	bool CallOrAddReplicatedDelegate(EAbilityGenericReplicatedEvent::Type Event, FSimpleMulticastDelegate::FDelegate Delegate);
 };
 
+//For searching through lists of ability instances
+struct FAbilityInstanceNamePredicate
+{
+	FAbilityInstanceNamePredicate(FName DesiredInstanceName)
+	{
+		InstanceName = DesiredInstanceName;
+	}
+
+	bool operator()(const TWeakObjectPtr<UAbilityTask> A) const
+	{
+		return (A.IsValid() && !A.Get()->GetInstanceName().IsNone() && A.Get()->GetInstanceName().IsValid() && (A.Get()->GetInstanceName() == InstanceName));
+	}
+
+	FName InstanceName;
+};
+
+
+struct FAbilityInstanceClassPredicate
+{
+	FAbilityInstanceClassPredicate(TSubclassOf<UAbilityTask> Class)
+	{
+		TaskClass = Class;
+	}
+
+	bool operator()(const TWeakObjectPtr<UAbilityTask> A) const
+	{
+		return (A.IsValid() && (A.Get()->GetClass() == TaskClass));
+	}
+
+	TSubclassOf<UAbilityTask> TaskClass;
+};
diff --git a/Engine/Source/Runtime/GameplayAbilities/Public/Abilities/Tasks/AbilityTask_WaitConfirmCancel.h b/Engine/Source/Runtime/GameplayAbilities/Public/Abilities/Tasks/AbilityTask_WaitConfirmCancel.h
--- a/Engine/Source/Runtime/GameplayAbilities/Public/Abilities/Tasks/AbilityTask_WaitConfirmCancel.h
+++ b/Engine/Source/Runtime/GameplayAbilities/Public/Abilities/Tasks/AbilityTask_WaitConfirmCancel.h
@@ -11,8 +11,8 @@ DECLARE_DYNAMIC_MULTICAST_DELEGATE(FWaitConfirmCancelDelegate);
 // UAbilityTask_WaitConfirmCancel = Wait for Targeting confirm/cancel
 // UAbilityTask_WaitConfirm = Wait for server to confirm ability activation
 
-UCLASS(MinimalAPI)
-class UAbilityTask_WaitConfirmCancel : public UAbilityTask
+UCLASS()
+class GAMEPLAYABILITIES_API UAbilityTask_WaitConfirmCancel : public UAbilityTask
 {
 	GENERATED_UCLASS_BODY()
 
diff --git a/Engine/Source/Runtime/GameplayAbilities/Public/Abilities/Tasks/AbilityTask_WaitGameplayEffectApplied.h b/Engine/Source/Runtime/GameplayAbilities/Public/Abilities/Tasks/AbilityTask_WaitGameplayEffectApplied.h
--- a/Engine/Source/Runtime/GameplayAbilities/Public/Abilities/Tasks/AbilityTask_WaitGameplayEffectApplied.h
+++ b/Engine/Source/Runtime/GameplayAbilities/Public/Abilities/Tasks/AbilityTask_WaitGameplayEffectApplied.h
@@ -20,6 +20,7 @@ class UAbilityTask_WaitGameplayEffectApplied : public UAbilityTask
 	FGameplayTagRequirements SourceTagRequirements;
 	FGameplayTagRequirements TargetTagRequirements;
 	bool TriggerOnce;
+	bool ListenForPeriodicEffects;
 
 	void SetExternalActor(AActor* InActor);
 
diff --git a/Engine/Source/Runtime/GameplayAbilities/Public/Abilities/Tasks/AbilityTask_WaitGameplayEffectApplied_Self.h b/Engine/Source/Runtime/GameplayAbilities/Public/Abilities/Tasks/AbilityTask_WaitGameplayEffectApplied_Self.h
--- a/Engine/Source/Runtime/GameplayAbilities/Public/Abilities/Tasks/AbilityTask_WaitGameplayEffectApplied_Self.h
+++ b/Engine/Source/Runtime/GameplayAbilities/Public/Abilities/Tasks/AbilityTask_WaitGameplayEffectApplied_Self.h
@@ -20,7 +20,7 @@ class UAbilityTask_WaitGameplayEffectApplied_Self : public UAbilityTask_WaitGame
 	 * Optional External Owner can be used to run this task on someone else (not the owner of the ability). By default you can leave this empty.
 	 */	 	 
 	UFUNCTION(BlueprintCallable, Category="Ability|Tasks", meta = (HidePin = "WorldContextObject", DefaultToSelf = "WorldContextObject", BlueprintInternalUseOnly = "TRUE"))
-	static UAbilityTask_WaitGameplayEffectApplied_Self* WaitGameplayEffectAppliedToSelf(UObject* WorldContextObject, const FGameplayTargetDataFilterHandle SourceFilter, FGameplayTagRequirements SourceTagRequirements, FGameplayTagRequirements TargetTagRequirements, bool TriggerOnce=false, AActor* OptionalExternalOwner=nullptr );
+	static UAbilityTask_WaitGameplayEffectApplied_Self* WaitGameplayEffectAppliedToSelf(UObject* WorldContextObject, const FGameplayTargetDataFilterHandle SourceFilter, FGameplayTagRequirements SourceTagRequirements, FGameplayTagRequirements TargetTagRequirements, bool TriggerOnce=false, AActor* OptionalExternalOwner=nullptr, bool ListenForPeriodicEffect=false );
 	
 protected:
 
@@ -30,4 +30,5 @@ protected:
 
 private:
 	FDelegateHandle OnApplyGameplayEffectCallbackDelegateHandle;
+	FDelegateHandle OnPeriodicGameplayEffectExecuteCallbackDelegateHandle;
 };
\ No newline at end of file
diff --git a/Engine/Source/Runtime/GameplayAbilities/Public/Abilities/Tasks/AbilityTask_WaitGameplayEffectApplied_Target.h b/Engine/Source/Runtime/GameplayAbilities/Public/Abilities/Tasks/AbilityTask_WaitGameplayEffectApplied_Target.h
--- a/Engine/Source/Runtime/GameplayAbilities/Public/Abilities/Tasks/AbilityTask_WaitGameplayEffectApplied_Target.h
+++ b/Engine/Source/Runtime/GameplayAbilities/Public/Abilities/Tasks/AbilityTask_WaitGameplayEffectApplied_Target.h
@@ -20,7 +20,7 @@ class UAbilityTask_WaitGameplayEffectApplied_Target : public UAbilityTask_WaitGa
 	 * Optional External Owner can be used to run this task on someone else (not the owner of the ability). By default you can leave this empty.
 	 */
 	UFUNCTION(BlueprintCallable, Category="Ability|Tasks", meta = (HidePin = "WorldContextObject", DefaultToSelf = "WorldContextObject", BlueprintInternalUseOnly = "TRUE"))
-	static UAbilityTask_WaitGameplayEffectApplied_Target* WaitGameplayEffectAppliedToTarget(UObject* WorldContextObject, const FGameplayTargetDataFilterHandle TargetFilter, FGameplayTagRequirements SourceTagRequirements, FGameplayTagRequirements TargetTagRequirements, bool TriggerOnce=false, AActor* OptionalExternalOwner=nullptr);
+	static UAbilityTask_WaitGameplayEffectApplied_Target* WaitGameplayEffectAppliedToTarget(UObject* WorldContextObject, const FGameplayTargetDataFilterHandle TargetFilter, FGameplayTagRequirements SourceTagRequirements, FGameplayTagRequirements TargetTagRequirements, bool TriggerOnce=false, AActor* OptionalExternalOwner=nullptr, bool ListenForPeriodicEffects=false);
 	
 protected:
 
@@ -30,4 +30,5 @@ protected:
 
 private:
 	FDelegateHandle OnApplyGameplayEffectCallbackDelegateHandle;
+	FDelegateHandle OnPeriodicGameplayEffectExecuteCallbackDelegateHandle;
 };
\ No newline at end of file
diff --git a/Engine/Source/Runtime/GameplayAbilities/Public/Abilities/Tasks/AbilityTask_WaitGameplayEvent.h b/Engine/Source/Runtime/GameplayAbilities/Public/Abilities/Tasks/AbilityTask_WaitGameplayEvent.h
new file mode 100644
--- /dev/null
+++ b/Engine/Source/Runtime/GameplayAbilities/Public/Abilities/Tasks/AbilityTask_WaitGameplayEvent.h
@@ -0,0 +1,42 @@
+// Copyright 1998-2015 Epic Games, Inc. All Rights Reserved.
+#pragma once
+#include "AbilityTask.h"
+#include "Abilities/GameplayAbilityTypes.h"
+#include "GameplayTags.h"
+#include "GameplayAbilityTypes.h"
+#include "AbilityTask_WaitGameplayEvent.generated.h"
+
+DECLARE_DYNAMIC_MULTICAST_DELEGATE_OneParam(FWaitGameplayEventDelegate, FGameplayEventData, Payload);
+
+UCLASS(MinimalAPI)
+class UAbilityTask_WaitGameplayEvent : public UAbilityTask
+{
+	GENERATED_UCLASS_BODY()
+
+	UPROPERTY(BlueprintAssignable)
+	FWaitGameplayEventDelegate	EventReceived;
+
+	/**
+	 * 	Wait until the specified gameplay tag event is triggered. By default this will look at the owner of this ability. OptionalExternalTarget can be set to make this look at another actor's tags for changes. 
+	 *  It will keep listening as long as OnlyTriggerOnce = false.
+	 */
+	UFUNCTION(BlueprintCallable, Category = "Ability|Tasks", meta = (HidePin = "WorldContextObject", DefaultToSelf = "WorldContextObject", BlueprintInternalUseOnly = "TRUE"))
+	static UAbilityTask_WaitGameplayEvent* WaitGameplayEvent(UObject* WorldContextObject, FGameplayTag EventTag, AActor* OptionalExternalTarget=nullptr, bool OnlyTriggerOnce=false);
+
+	void SetExternalTarget(AActor* Actor);
+
+	UAbilitySystemComponent* GetTargetASC();
+
+	virtual void Activate() override;
+
+	virtual void GameplayEventCallback(const FGameplayEventData* Payload);
+
+	void OnDestroy(bool AbilityEnding) override;
+
+	FGameplayTag Tag;
+	TWeakObjectPtr<UAbilitySystemComponent> OptionalExternalTarget;
+	bool UseExternalTarget;	
+	bool OnlyTriggerOnce;
+
+	FDelegateHandle MyHandle;
+};
\ No newline at end of file
diff --git a/Engine/Source/Runtime/GameplayAbilities/Public/AbilitySystemBlueprintLibrary.h b/Engine/Source/Runtime/GameplayAbilities/Public/AbilitySystemBlueprintLibrary.h
--- a/Engine/Source/Runtime/GameplayAbilities/Public/AbilitySystemBlueprintLibrary.h
+++ b/Engine/Source/Runtime/GameplayAbilities/Public/AbilitySystemBlueprintLibrary.h
@@ -105,6 +105,9 @@ class GAMEPLAYABILITIES_API UAbilitySystemBlueprintLibrary : public UBlueprintFu
 	//		GameplayEffectContext
 	// -------------------------------------------------------------------------------
 
+	UFUNCTION(BlueprintPure, Category = "Ability|EffectContext", Meta = (DisplayName = "IsValid"))
+	static bool EffectContextIsValid(FGameplayEffectContextHandle EffectContext);
+
 	UFUNCTION(BlueprintPure, Category = "Ability|EffectContext", Meta = (DisplayName = "IsInstigatorLocallyControlled"))
 	static bool EffectContextIsInstigatorLocallyControlled(FGameplayEffectContextHandle EffectContext);
 
@@ -221,6 +224,14 @@ class GAMEPLAYABILITIES_API UAbilitySystemBlueprintLibrary : public UBlueprintFu
 	UFUNCTION(BlueprintCallable, Category = "Ability|GameplayEffect")
 	static FGameplayEffectSpecHandle SetStackCountToMax(FGameplayEffectSpecHandle SpecHandle);
 
+	// -------------------------------------------------------------------------------
+	//		GameplayEffectSpec
+	// -------------------------------------------------------------------------------
+
+	/** Gets the magnitude of change for an attribute on an APPLIED GameplayEffectSpec. */
+	UFUNCTION(BlueprintCallable, Category = "Ability|GameplayEffect")
+	static float GetModifiedAttributeMagnitude(FGameplayEffectSpecHandle SpecHandle, FGameplayAttribute Attribute);
+
 	// -------------------------------------------------------------------------------
 	//		FActiveGameplayEffectHandle
 	// -------------------------------------------------------------------------------
diff --git a/Engine/Source/Runtime/GameplayAbilities/Public/AbilitySystemComponent.h b/Engine/Source/Runtime/GameplayAbilities/Public/AbilitySystemComponent.h
--- a/Engine/Source/Runtime/GameplayAbilities/Public/AbilitySystemComponent.h
+++ b/Engine/Source/Runtime/GameplayAbilities/Public/AbilitySystemComponent.h
@@ -382,12 +382,20 @@ class GAMEPLAYABILITIES_API UAbilitySystemComponent : public UGameplayTasksCompo
 
 	void OnGameplayEffectAppliedToSelf(UAbilitySystemComponent* Source, const FGameplayEffectSpec& SpecApplied, FActiveGameplayEffectHandle ActiveHandle);
 
+	void OnPeriodicGameplayEffectExecuteOnTarget(UAbilitySystemComponent* Target, const FGameplayEffectSpec& SpecExecuted, FActiveGameplayEffectHandle ActiveHandle);
+
+	void OnPeriodicGameplayEffectExecuteOnSelf(UAbilitySystemComponent* Source, const FGameplayEffectSpec& SpecExecuted, FActiveGameplayEffectHandle ActiveHandle);
+
 	DECLARE_MULTICAST_DELEGATE_ThreeParams(FOnGameplayEffectAppliedDelegate, UAbilitySystemComponent*, const FGameplayEffectSpec&, FActiveGameplayEffectHandle);
 
 	FOnGameplayEffectAppliedDelegate OnGameplayEffectAppliedDelegateToSelf;
 
 	FOnGameplayEffectAppliedDelegate OnGameplayEffectAppliedDelegateToTarget;
 
+	FOnGameplayEffectAppliedDelegate OnPeriodicGameplayEffectExecuteDelegateOnSelf;
+
+	FOnGameplayEffectAppliedDelegate OnPeriodicGameplayEffectExecuteDelegateOnTarget;
+
 	// --------------------------------------------
 	// Tags
 	// --------------------------------------------
@@ -638,6 +646,9 @@ class GAMEPLAYABILITIES_API UAbilitySystemComponent : public UGameplayTasksCompo
 	/** Cancels all abilities regardless of tags. Will not cancel the ignore instance */
 	void CancelAllAbilities(UGameplayAbility* Ignore=nullptr);
 
+	/** Cancels all abilities and kills any remaining instanced abilities */
+	virtual void DestroyActiveState();
+
 	// ----------------------------------------------------------------------------------------------------------------
 	/** 
 	 * Called from ability activation or native code, will apply the correct ability blocking tags and cancel existing abilities. Subclasses can override the behavior 
@@ -679,6 +690,9 @@ class GAMEPLAYABILITIES_API UAbilitySystemComponent : public UGameplayTasksCompo
 	/** Returns an ability spec from a handle. If modifying call MarkAbilitySpecDirty */
 	FGameplayAbilitySpec* FindAbilitySpecFromInputID(int32 InputID);
 
+	/** Retrieves the EffectContext of the GameplayEffect of the active GameplayEffect. */
+	FGameplayEffectContextHandle GetEffectContextFromActiveGEHandle(FActiveGameplayEffectHandle Handle);
+
 	/** Call to mark that an ability spec has been modified */
 	void MarkAbilitySpecDirty(FGameplayAbilitySpec& Spec);
 
@@ -830,6 +844,9 @@ public:
 	/** Executes a gameplay event. Returns the number of successful ability activations triggered by the event */
 	virtual int32 HandleGameplayEvent(FGameplayTag EventTag, const FGameplayEventData* Payload);
 
+	/** Generic callbacks for gameplay events. See UAbilityTask_WaitGameplayEvent */
+	TMap<FGameplayTag, FGameplayEventMulticastDelegate> GenericGameplayEventCallbacks;
+
 	virtual void NotifyAbilityCommit(UGameplayAbility* Ability);
 	virtual void NotifyAbilityActivated(const FGameplayAbilitySpecHandle Handle, UGameplayAbility* Ability);
 	virtual void NotifyAbilityFailed(const FGameplayAbilitySpecHandle Handle, UGameplayAbility* Ability, FGameplayTagContainer FailureReason);
@@ -1109,6 +1126,7 @@ protected:
 	// ---------------------------------------------
 
 	virtual void OnRegister() override;
+	virtual void OnUnregister() override;
 
 	const UAttributeSet*	GetAttributeSubobject(const TSubclassOf<UAttributeSet> AttributeClass) const;
 	const UAttributeSet*	GetAttributeSubobjectChecked(const TSubclassOf<UAttributeSet> AttributeClass) const;
diff --git a/Engine/Source/Runtime/GameplayAbilities/Public/AbilitySystemGlobals.h b/Engine/Source/Runtime/GameplayAbilities/Public/AbilitySystemGlobals.h
--- a/Engine/Source/Runtime/GameplayAbilities/Public/AbilitySystemGlobals.h
+++ b/Engine/Source/Runtime/GameplayAbilities/Public/AbilitySystemGlobals.h
@@ -17,6 +17,9 @@ struct FGameplayEffectContext;
 struct FGameplayTag;
 struct FAttributeSetInitter;
 struct FGameplayEffectSpec;
+struct FGameplayEffectSpecForRPC;
+struct FGameplayCueParameters;
+struct FGameplayEffectContextHandle;
 
 /** Holds global data for the ability system. Can be configured per project via config file */
 UCLASS(config=Game)
@@ -160,6 +163,10 @@ class GAMEPLAYABILITIES_API UAbilitySystemGlobals : public UObject
 		}
 	}
 
+	// GameplayCue Parameters
+	virtual void InitGameplayCueParameters(FGameplayCueParameters& CueParameters, const FGameplayEffectSpecForRPC &Spec);
+	virtual void InitGameplayCueParameters(FGameplayCueParameters& CueParameters, const FGameplayEffectContextHandle& EffectContext);
+
 protected:
 
 	virtual void InitAttributeDefaults();
diff --git a/Engine/Source/Runtime/GameplayAbilities/Public/AbilitySystemLog.h b/Engine/Source/Runtime/GameplayAbilities/Public/AbilitySystemLog.h
--- a/Engine/Source/Runtime/GameplayAbilities/Public/AbilitySystemLog.h
+++ b/Engine/Source/Runtime/GameplayAbilities/Public/AbilitySystemLog.h
@@ -28,33 +28,18 @@ GAMEPLAYABILITIES_API DECLARE_LOG_CATEGORY_EXTERN(VLogAbilitySystem, Warning, Al
 	UE_VLOG(Actor, VLogAbilitySystem, Verbosity, Format, ##__VA_ARGS__); \
 }
 
-#define ABILITY_LOG_SCOPE( Format, ... ) 
-
 #else
 
 #define ABILITY_LOG(Verbosity, Format, ...) \
 { \
-	UE_LOG(LogAbilitySystem, Verbosity, TEXT("%s"), *AbilitySystemLog::Log(ELogVerbosity::Verbosity, FString::Printf(Format, ##__VA_ARGS__))); \
+	UE_LOG(LogAbilitySystem, Verbosity, Format, ##__VA_ARGS__); \
 }
 
-#if ENABLE_VISUAL_LOG
-#	define ABILITY_VLOG(Actor, Verbosity, Format, ...) \
-	{ \
-		if(FVisualLogger::IsRecording() || LogAbilitySystem.IsSuppressed(ELogVerbosity::Verbosity) == false ) \
-		{ \
-			const FString Str = AbilitySystemLog::Log(ELogVerbosity::Verbosity, FString::Printf(Format, ##__VA_ARGS__)); \
-			UE_LOG(LogAbilitySystem, Verbosity, TEXT("%s"), *Str); \
-			UE_VLOG(Actor, VLogAbilitySystem, Verbosity, TEXT("%s"), *Str); \
-		} \
-	}
-#else //ENABLE_VISUAL_LOG
-#	define ABILITY_VLOG(Actor, Verbosity, Format, ...) \
-	{ \
-		UE_LOG(LogAbilitySystem, Verbosity, TEXT("%s"), *AbilitySystemLog::Log(ELogVerbosity::Verbosity, FString::Printf(Format, ##__VA_ARGS__))); \
-	}
-#endif  //ENABLE_VISUAL_LOG
-
-#define ABILITY_LOG_SCOPE( Format, ... ) AbilitySystemLogScope PREPROCESSOR_JOIN(LogScope,__LINE__)( FString::Printf(Format, ##__VA_ARGS__));
+#define ABILITY_VLOG(Actor, Verbosity, Format, ...) \
+{ \
+	UE_LOG(LogAbilitySystem, Verbosity, Format, ##__VA_ARGS__); \
+	UE_VLOG(Actor, VLogAbilitySystem, Verbosity, Format, ##__VA_ARGS__); \
+}
 
 #endif //NO_LOGGING
 
@@ -79,47 +64,3 @@ GAMEPLAYABILITIES_API DECLARE_LOG_CATEGORY_EXTERN(VLogAbilitySystem, Warning, Al
 #define ABILITY_VLOG_ATTRIBUTE_GRAPH(Actor, Verbosity, AttributeName, OldValue, NewValue)
 
 #endif //ENABLE_VISUAL_LOG
-
-struct AbilitySystemLogScope
-{
-	AbilitySystemLogScope()
-	{
-		Init();
-	}
-
-	AbilitySystemLogScope(const FString& InScopeName) :
-		ScopeName(InScopeName)
-	{
-		Init();
-	}
-
-	~AbilitySystemLogScope();
-	
-	FString ScopeName;
-	bool PrintedInThisScope;
-
-private:
-
-	void Init();
-};
-
-class GAMEPLAYABILITIES_API AbilitySystemLog
-{
-public:
-
-	AbilitySystemLog()
-	{
-		Indent = 0;
-		NeedNewLine = false;
-	}
-
-	static FString Log(ELogVerbosity::Type, const FString& Str);
-
-	static void PushScope(AbilitySystemLogScope* Scope);
-
-	static void PopScope();	
-
-	int32	Indent;
-	TArray<AbilitySystemLogScope*>	ScopeStack;
-	bool NeedNewLine;
-};
diff --git a/Engine/Source/Runtime/GameplayAbilities/Public/AttributeSet.h b/Engine/Source/Runtime/GameplayAbilities/Public/AttributeSet.h
--- a/Engine/Source/Runtime/GameplayAbilities/Public/AttributeSet.h
+++ b/Engine/Source/Runtime/GameplayAbilities/Public/AttributeSet.h
@@ -218,6 +218,9 @@ public:
 	bool operator==(const FScalableFloat& Other) const;
 	bool operator!=(const FScalableFloat& Other) const;
 
+	//copy operator to prevent duplicate handles
+	void operator=(const FScalableFloat& Src);
+
 private:
 
 	/** Conditionally register the OnCurveTablePostReimport function with the re-import manager */
@@ -318,7 +321,7 @@ struct GAMEPLAYABILITIES_API FAttributeSetInitter
 	void PreloadAttributeSetData(UCurveTable* CurveData);
 
 	void InitAttributeSetDefaults(UAbilitySystemComponent* AbilitySystemComponent, FName GroupName, int32 Level, bool bInitialInit) const;
-
+	void ApplyAttributeDefault(UAbilitySystemComponent* AbilitySystemComponent, FGameplayAttribute& InAttribute, FName GroupName, int32 Level) const;
 private:
 
 	struct FAttributeDefaultValueList
diff --git a/Engine/Source/Runtime/GameplayAbilities/Public/GameplayCueInterface.h b/Engine/Source/Runtime/GameplayAbilities/Public/GameplayCueInterface.h
--- a/Engine/Source/Runtime/GameplayAbilities/Public/GameplayCueInterface.h
+++ b/Engine/Source/Runtime/GameplayAbilities/Public/GameplayCueInterface.h
@@ -105,6 +105,9 @@ struct FActiveGameplayCueContainer : public FFastArraySerializer
 
 	void PredictiveAdd(const FGameplayTag& Tag, FPredictionKey& PredictionKey);
 
+	/** Does explicit check for gameplay cue tag */
+	bool HasCue(const FGameplayTag& Tag) const;
+
 	bool NetDeltaSerialize(FNetDeltaSerializeInfo & DeltaParms)
 	{
 		return FastArrayDeltaSerialize<FActiveGameplayCue>(GameplayCues, DeltaParms, *this);
@@ -123,3 +126,21 @@ struct TStructOpsTypeTraits< FActiveGameplayCueContainer > : public TStructOpsTy
 		WithNetDeltaSerializer = true,
 	};
 };
+
+
+/**
+ *	Wrapper struct around a gameplaytag with the GameplayCue category. This also allows for a details customization
+ */
+USTRUCT(BlueprintType)
+struct FGameplayCueTag
+{
+	GENERATED_USTRUCT_BODY()
+	
+	UPROPERTY(EditAnywhere, BlueprintReadWrite, meta = (Categories="GameplayCue"), Category="GameplayCue")
+	FGameplayTag GameplayCueTag;
+
+	bool IsValid() const
+	{
+		return GameplayCueTag.IsValid();
+	}
+};
diff --git a/Engine/Source/Runtime/GameplayAbilities/Public/GameplayCueManager.h b/Engine/Source/Runtime/GameplayAbilities/Public/GameplayCueManager.h
--- a/Engine/Source/Runtime/GameplayAbilities/Public/GameplayCueManager.h
+++ b/Engine/Source/Runtime/GameplayAbilities/Public/GameplayCueManager.h
@@ -220,6 +220,11 @@ class GAMEPLAYABILITIES_API UGameplayCueManager : public UDataAsset
 	/** Handles cleaning up an object library if it matches the passed in object */
 	void HandleAssetDeleted(UObject *Object);
 
+	/** Warns if we move a GameplayCue notify out of the valid search paths */
+	void HandleAssetRenamed(const FAssetData& Data, const FString& String);
+
+	void VerifyNotifyAssetIsInValidPath(FString Path);
+
 	bool RegisteredEditorCallbacks;
 
 	bool bAccelerationMapOutdated;
diff --git a/Engine/Source/Runtime/GameplayAbilities/Public/GameplayEffect.h b/Engine/Source/Runtime/GameplayAbilities/Public/GameplayEffect.h
--- a/Engine/Source/Runtime/GameplayAbilities/Public/GameplayEffect.h
+++ b/Engine/Source/Runtime/GameplayAbilities/Public/GameplayEffect.h
@@ -969,9 +969,6 @@ struct GAMEPLAYABILITIES_API FGameplayEffectSpec
 	/** Adds a new modified attribute struct, will always add so check to see if it exists first */
 	FGameplayEffectModifiedAttribute* AddModifiedAttribute(const FGameplayAttribute& Attribute);
 
-	/** Deletes any modified attributes that aren't needed. Call before replication */
-	void PruneModifiedAttributes();
-
 	/**
 	 * Helper function to attempt to calculate the duration of the spec from its GE definition
 	 * 
@@ -1139,6 +1136,8 @@ struct GAMEPLAYABILITIES_API FGameplayEffectSpecForRPC
 
 	FGameplayEffectSpecForRPC(const FGameplayEffectSpec& InSpec);
 
+	void InitFromSpec(const FGameplayEffectSpec& InSpec);
+
 	/** GameplayEfect definition. The static data that this spec points to. */
 	UPROPERTY()
 	const UGameplayEffect* Def;
@@ -1723,7 +1722,7 @@ private:
 	bool HandleActiveGameplayEffectStackOverflow(const FActiveGameplayEffect& ActiveStackableGE, const FGameplayEffectSpec& OldSpec, const FGameplayEffectSpec& OverflowingSpec);
 
 	/** After application has gone through, give stacking rules a chance to do something as the source of the gameplay effect (E.g., remove an old version) */
-	void ApplyStackingLogicPostApplyAsSource(UAbilitySystemComponent* Target, const FGameplayEffectSpec& SpecApplied, FActiveGameplayEffectHandle ActiveHandle);
+	virtual void ApplyStackingLogicPostApplyAsSource(UAbilitySystemComponent* Target, const FGameplayEffectSpec& SpecApplied, FActiveGameplayEffectHandle ActiveHandle) { }
 
 	mutable int32 ScopedLockCount;
 	int32 PendingRemoves;
diff --git a/Engine/Source/Runtime/GameplayAbilities/Public/GameplayEffectExecutionCalculation.h b/Engine/Source/Runtime/GameplayAbilities/Public/GameplayEffectExecutionCalculation.h
--- a/Engine/Source/Runtime/GameplayAbilities/Public/GameplayEffectExecutionCalculation.h
+++ b/Engine/Source/Runtime/GameplayAbilities/Public/GameplayEffectExecutionCalculation.h
@@ -25,6 +25,9 @@ public:
 	/** Simple accessor to owning gameplay spec */
 	const FGameplayEffectSpec& GetOwningSpec() const;
 
+	/** Non const access. Be careful with this, especially when modifying a spec after attribute capture. */
+	FGameplayEffectSpec* GetOwningSpecForPreExecuteMod() const;
+
 	/** Simple accessor to target ability system component */
 	UAbilitySystemComponent* GetTargetAbilitySystemComponent() const;
 
diff --git a/Engine/Source/Runtime/GameplayAbilities/Public/GameplayEffectTypes.h b/Engine/Source/Runtime/GameplayAbilities/Public/GameplayEffectTypes.h
--- a/Engine/Source/Runtime/GameplayAbilities/Public/GameplayEffectTypes.h
+++ b/Engine/Source/Runtime/GameplayAbilities/Public/GameplayEffectTypes.h
@@ -700,7 +700,7 @@ struct TStructOpsTypeTraits<FGameplayEffectContextHandle> : public TStructOpsTyp
 
 
 USTRUCT(BlueprintType)
-struct FGameplayCueParameters
+struct GAMEPLAYABILITIES_API FGameplayCueParameters
 {
 	GENERATED_USTRUCT_BODY()
 
@@ -708,8 +708,15 @@ struct FGameplayCueParameters
 	: NormalizedMagnitude(0.0f)
 	, RawMagnitude(0.0f)
 	, MatchedTagName(NAME_None)
+	, Location(ForceInitToZero)
+	, Normal(ForceInitToZero)
 	{}
 
+	/** Projects can override this via UAbilitySystemGlobals */
+	FGameplayCueParameters(const struct FGameplayEffectSpecForRPC &Spec);
+
+	FGameplayCueParameters(const struct FGameplayEffectContextHandle& EffectContext);
+
 	/** Magnitude of source gameplay effect, normalzed from 0-1. Use this for "how strong is the gameplay effect" (0=min, 1=,max) */
 	UPROPERTY(BlueprintReadWrite, Category=GameplayCue)
 	float NormalizedMagnitude;
@@ -737,6 +744,45 @@ struct FGameplayCueParameters
 	/** The aggregated target tags taken from the effect spec */
 	UPROPERTY(BlueprintReadWrite, Category=GameplayCue)
 	FGameplayTagContainer AggregatedTargetTags;
+
+	UPROPERTY(BlueprintReadWrite, Category=GameplayCue)
+	FVector_NetQuantize10 Location;
+
+	UPROPERTY(BlueprintReadWrite, Category=GameplayCue)
+	FVector_NetQuantizeNormal Normal;
+
+	/** Instigator actor, the actor that owns the ability system component */
+	UPROPERTY(BlueprintReadWrite, Category=GameplayCue)
+	TWeakObjectPtr<AActor> Instigator;
+
+	/** The physical actor that actually did the damage, can be a weapon or projectile */
+	UPROPERTY(BlueprintReadWrite, Category=GameplayCue)
+	TWeakObjectPtr<AActor> EffectCauser;
+
+	/** Object this effect was created from, can be an actor or static object. Useful to bind an effect to a gameplay object */
+	UPROPERTY(BlueprintReadWrite, Category=GameplayCue)
+	TWeakObjectPtr<const UObject> SourceObject;
+
+	bool NetSerialize(FArchive& Ar, class UPackageMap* Map, bool& bOutSuccess);
+
+	bool IsInstigatorLocallyControlled() const;
+
+	bool IsInstigatorLocallyControlledPlayer() const;
+
+	AActor* GetInstigator() const;
+
+	AActor* GetEffectCauser() const;
+
+	const UObject* GetSourceObject() const;
+};
+
+template<>
+struct TStructOpsTypeTraits<FGameplayCueParameters> : public TStructOpsTypeTraitsBase
+{
+	enum
+	{
+		WithNetSerializer = true		
+	};
 };
 
 UENUM(BlueprintType)
@@ -771,7 +817,7 @@ namespace EGameplayTagEventType
 {
 	enum Type
 	{		
-		/** Event only happens when tag is branch new or removed */
+		/** Event only happens when tag is new or completely removed */
 		NewOrRemoved,
 
 		/** Event happens any time tag "count" changes */
diff --git a/Engine/Source/Runtime/GameplayAbilities/Public/GameplayEffectUIData_TextOnly.h b/Engine/Source/Runtime/GameplayAbilities/Public/GameplayEffectUIData_TextOnly.h
--- a/Engine/Source/Runtime/GameplayAbilities/Public/GameplayEffectUIData_TextOnly.h
+++ b/Engine/Source/Runtime/GameplayAbilities/Public/GameplayEffectUIData_TextOnly.h
@@ -9,10 +9,10 @@
  * If your game needs only text, this is a reasonable class to use. To include more data, make a custom subclass of UGameplayEffectUIData.
  */
 UCLASS()
-class UGameplayEffectUIData_TextOnly : public UGameplayEffectUIData
+class GAMEPLAYABILITIES_API UGameplayEffectUIData_TextOnly : public UGameplayEffectUIData
 {
 	GENERATED_UCLASS_BODY()
 
-	UPROPERTY(EditDefaultsOnly, BlueprintReadOnly, Category=Data)
+	UPROPERTY(EditDefaultsOnly, BlueprintReadOnly, Category = Data, meta = (MultiLine = "true"))
 	FText Description;
 };
\ No newline at end of file
diff --git a/Engine/Source/Runtime/GameplayAbilities/Public/GameplayTagResponseTable.h b/Engine/Source/Runtime/GameplayAbilities/Public/GameplayTagResponseTable.h
--- a/Engine/Source/Runtime/GameplayAbilities/Public/GameplayTagResponseTable.h
+++ b/Engine/Source/Runtime/GameplayAbilities/Public/GameplayTagResponseTable.h
@@ -8,6 +8,7 @@
 
 
 class UGameplayEffect;
+class UAbilitySystemComponent;
 
 USTRUCT()
 struct FGameplayTagReponsePair
@@ -21,6 +22,10 @@ struct FGameplayTagReponsePair
 	/** The GameplayEffect to apply in reponse to the tag */
 	UPROPERTY(EditAnywhere, Category="Response")
 	TSubclassOf<UGameplayEffect> ResponseGameplayEffect;
+
+	/** The max "count" this response can achieve. This will not prevent counts from being applied, but will be used when calculating the net count of a tag. 0=no cap. */
+	UPROPERTY(EditAnywhere, Category="Response", meta=(ClampMin = "0"))
+	int32 SoftCountCap;
 };
 
 USTRUCT()
@@ -41,7 +46,7 @@ struct FGameplayTagResponseTableEntry
  *	A data driven table for applying gameplay effects based on tag count. This allows designers to map a 
  *	"tag count" -> "response Gameplay Effect" relationship.
  *	
- *	For example, "for every stack of "Status.Haste" I get 1 level of GE_Response_Haste. This class facilitates
+ *	For example, "for every count of "Status.Haste" I get 1 level of GE_Response_Haste. This class facilitates
  *	building this table and automatically registering and responding to tag events on the ability system component.
  */
 UCLASS()
@@ -61,13 +66,16 @@ protected:
 
 	UFUNCTION()
 	void TagResponseEvent(const FGameplayTag Tag, int32 NewCount, UAbilitySystemComponent* ASC, int32 idx);
+	
+	/** Temporary structs to avoid extra heap allocations every time we recalculate tag count */
+	mutable FGameplayEffectQuery Query;
 
-	FGameplayEffectQuery MakeQuery(const FGameplayTag& Tag)
+	FGameplayEffectQuery& MakeQuery(const FGameplayTag& Tag) const
 	{
-		return FGameplayEffectQuery::MakeQuery_MatchAnyOwningTags(FGameplayTagContainer(Tag));
+		Query.OwningTagQuery.ReplaceTagFast(Tag);
+		return Query;
 	}
 
-
 	// ----------------------------------------------------
 
 	struct FGameplayTagResponseAppliedInfo
@@ -76,29 +84,11 @@ protected:
 		FActiveGameplayEffectHandle NegativeHandle;
 	};
 
-	TMap< TWeakObjectPtr<UAbilitySystemComponent>, FGameplayTagResponseAppliedInfo > RegisteredASCs;
+	TMap< TWeakObjectPtr<UAbilitySystemComponent>, TArray< FGameplayTagResponseAppliedInfo> > RegisteredASCs;
 
-	void Remove(UAbilitySystemComponent* ASC, FActiveGameplayEffectHandle& Handle)
-	{
-		if (Handle.IsValid())
-		{
-			ASC->RemoveActiveGameplayEffect(Handle);
-			Handle.Invalidate();
-		}
-	}
+	void Remove(UAbilitySystemComponent* ASC, FActiveGameplayEffectHandle& Handle);
 
-	void AddOrUpdate(UAbilitySystemComponent* ASC, const TSubclassOf<UGameplayEffect>& ResponseGameplayEffect, int32 TotalCount, FActiveGameplayEffectHandle& Handle)
-	{
-		if (ResponseGameplayEffect)
-		{
-			if (Handle.IsValid())
-			{
-				ASC->SetActiveGameplayEffectLevel(Handle, TotalCount);
-			}
-			else
-			{
-				Handle = ASC->ApplyGameplayEffectToSelf(Cast<UGameplayEffect>(ResponseGameplayEffect->ClassDefaultObject), TotalCount, ASC->GetEffectContext());
-			}
-		}
-	}
-};
\ No newline at end of file
+	void AddOrUpdate(UAbilitySystemComponent* ASC, const TSubclassOf<UGameplayEffect>& ResponseGameplayEffect, int32 TotalCount, FActiveGameplayEffectHandle& Handle);
+
+	int32 GetCount(const FGameplayTagReponsePair& Pair, UAbilitySystemComponent* ASC) const;
+};
diff --git a/Engine/Source/Runtime/GameplayTags/Classes/GameplayTagContainer.h b/Engine/Source/Runtime/GameplayTags/Classes/GameplayTagContainer.h
--- a/Engine/Source/Runtime/GameplayTags/Classes/GameplayTagContainer.h
+++ b/Engine/Source/Runtime/GameplayTags/Classes/GameplayTagContainer.h
@@ -279,6 +279,8 @@ struct GAMEPLAYTAGS_API FGameplayTagContainer
 
 	/** An empty Gameplay Tag Container */
 	static const FGameplayTagContainer EmptyContainer;
+		
+	bool NetSerialize(FArchive& Ar, class UPackageMap* Map, bool& bOutSuccess);
 
 protected:
 	/**
@@ -309,6 +311,8 @@ protected:
 	// Allow the redirection helper class access to RemoveTagByExplicitName.  It can then (through friendship) allow
 	// access to others without exposing everything the Container has privately to everyone.
 	friend class FGameplayTagRedirectHelper;
+	friend struct FGameplayTagQuery;
+	friend struct FGameplayTagQueryExpression;
 
 private:
 
@@ -342,6 +346,7 @@ struct TStructOpsTypeTraits<FGameplayTagContainer> : public TStructOpsTypeTraits
 	{
 		WithSerializer = true,
 		WithIdenticalViaEquality = true,
+		WithNetSerializer = true,
 		WithCopy = true
 	};
 };
@@ -440,6 +445,23 @@ private:
 	}
 
 public:
+
+	/** Replaces existing tags with passed in tags. Does not modify the tag query expression logic. Useful when you need to cache off and update often used query. Must use same sized tag container! */
+	void ReplaceTagsFast(FGameplayTagContainer const& Tags)
+	{
+		ensure(Tags.Num() == TagDictionary.Num());
+		TagDictionary.Reset();
+		TagDictionary.Append(Tags.GameplayTags);
+	}
+
+	/** Replaces existing tags with passed in tag. Does not modify the tag query expression logic. Useful when you need to cache off and update often used query. */		 
+	void ReplaceTagFast(FGameplayTag const& Tag)
+	{
+		ensure(1 == TagDictionary.Num());
+		TagDictionary.Reset();
+		TagDictionary.Add(Tag);
+	}
+
 	/** Returns true if the given tags match this query, or false otherwise. */
 	bool Matches(FGameplayTagContainer const& Tags) const;
 
@@ -541,10 +563,7 @@ struct GAMEPLAYTAGS_API FGameplayTagQueryExpression
 	FGameplayTagQueryExpression& AddTags(FGameplayTagContainer const& Tags)
 	{
 		ensure(UsesTagSet());
-		for (auto T : Tags)
-		{
-			TagSet.Add(T);
-		}
+		TagSet.Append(Tags.GameplayTags);
 		return *this;
 	}
 
diff --git a/Engine/Source/Runtime/GameplayTags/Classes/GameplayTagsManager.h b/Engine/Source/Runtime/GameplayTags/Classes/GameplayTagsManager.h
--- a/Engine/Source/Runtime/GameplayTags/Classes/GameplayTagsManager.h
+++ b/Engine/Source/Runtime/GameplayTags/Classes/GameplayTagsManager.h
@@ -167,6 +167,10 @@ class GAMEPLAYTAGS_API UGameplayTagsManager : public UObject
 	 */
 	void OnObjectReimported(class UFactory* ImportFactory, UObject* InObject);
 
+	TSharedPtr<FGameplayTagNode> FindTagNode(FName TagName) const;
+
+	static void AddNewGameplayTagToINI(FString NewTag);
+
 #endif //WITH_EDITOR
 
 	/** 
diff --git a/Engine/Source/Runtime/GameplayTags/GameplayTags.Build.cs b/Engine/Source/Runtime/GameplayTags/GameplayTags.Build.cs
--- a/Engine/Source/Runtime/GameplayTags/GameplayTags.Build.cs
+++ b/Engine/Source/Runtime/GameplayTags/GameplayTags.Build.cs
@@ -46,6 +46,9 @@ namespace UnrealBuildTool.Rules
 			if (UEBuildConfiguration.bBuildEditor == true)
 			{
 				PrivateDependencyModuleNames.Add("UnrealEd");
+                PrivateDependencyModuleNames.Add("SourceControl");
+                PrivateDependencyModuleNames.Add("Slate");
+                PrivateDependencyModuleNames.Add("SlateCore");
 			}
 		}
 	}
diff --git a/Engine/Source/Runtime/GameplayTags/Private/GameplayTagContainer.cpp b/Engine/Source/Runtime/GameplayTags/Private/GameplayTagContainer.cpp
--- a/Engine/Source/Runtime/GameplayTags/Private/GameplayTagContainer.cpp
+++ b/Engine/Source/Runtime/GameplayTags/Private/GameplayTagContainer.cpp
@@ -588,7 +588,7 @@ void FGameplayTagContainer::AppendMatchingTags(FGameplayTagContainer const& Othe
 {
 	for(TArray<FGameplayTag>::TConstIterator It(OtherA.GameplayTags); It; ++It)
 	{
-		if (OtherB.HasTag(*It, EGameplayTagMatchType::IncludeParentTags, EGameplayTagMatchType::Explicit))
+		if (OtherB.HasTag(*It, EGameplayTagMatchType::Explicit, EGameplayTagMatchType::IncludeParentTags))
 		{
 			AddTag(*It);		
 		}
@@ -713,6 +713,12 @@ FString FGameplayTagContainer::ToStringSimple() const
 	return RetString;
 }
 
+bool FGameplayTagContainer::NetSerialize(FArchive& Ar, class UPackageMap* Map, bool& bOutSuccess)
+{
+	Ar << GameplayTags;
+	return true;
+}
+
 FText FGameplayTagContainer::ToMatchingText(EGameplayContainerMatchType MatchType, bool bInvertCondition) const
 {
 	enum class EMatchingTypes : int8
diff --git a/Engine/Source/Runtime/GameplayTags/Private/GameplayTagsManager.cpp b/Engine/Source/Runtime/GameplayTags/Private/GameplayTagsManager.cpp
--- a/Engine/Source/Runtime/GameplayTags/Private/GameplayTagsManager.cpp
+++ b/Engine/Source/Runtime/GameplayTags/Private/GameplayTagsManager.cpp
@@ -2,13 +2,15 @@
 
 #include "GameplayTagsModulePrivatePCH.h"
 #include "GameplayTagsSettings.h"
-//#include "AssetRegistryModule.h"
-
+#include "SNotificationList.h"
+#include "NotificationManager.h"
 
 #if WITH_EDITOR
 #include "UnrealEd.h"
 #endif
 
+#define LOCTEXT_NAMESPACE "GameplayTagManager"
+
 UGameplayTagsManager::UGameplayTagsManager(const FObjectInitializer& ObjectInitializer)
 	: Super(ObjectInitializer),
 	bHasHandledRedirectors(false)
@@ -524,6 +526,72 @@ void UGameplayTagsManager::OnObjectReimported(UFactory* ImportFactory, UObject*
 		ConstructGameplayTagTree();
 	}
 }
+
+TSharedPtr<FGameplayTagNode> UGameplayTagsManager::FindTagNode(FName TagName) const
+{
+	return FindTagNode(GameplayRootTag, TagName);
+}
+
+void UGameplayTagsManager::AddNewGameplayTagToINI(FString NewTag)
+{
+	if(NewTag.IsEmpty())
+	{
+		return;
+	}
+
+	if (ShouldImportTagsFromINI() == false)
+	{
+		return;
+	}
+
+	UGameplayTagsSettings* Settings = GetMutableDefault<UGameplayTagsSettings>();
+	if (Settings)
+	{
+
+		FString RelativeConfigFilePath = Settings->GetDefaultConfigFilename();
+		FString ConfigPath = FPaths::ConvertRelativePathToFull(RelativeConfigFilePath);
+
+		if (ISourceControlModule::Get().IsEnabled())
+		{
+			FText ErrorMessage;
+
+			if (!SourceControlHelpers::CheckoutOrMarkForAdd(ConfigPath, FText::FromString(ConfigPath), NULL, ErrorMessage))
+			{
+				FNotificationInfo Info(ErrorMessage);
+				Info.ExpireDuration = 3.0f;
+				FSlateNotificationManager::Get().AddNotification(Info);
+			}
+		}
+		else
+		{
+			if (!FPlatformFileManager::Get().GetPlatformFile().SetReadOnly(*ConfigPath, false))
+			{
+				FText NotificationErrorText = FText::Format(LOCTEXT("FailedToMakeWritable", "Could not make {0} writable."), FText::FromString(ConfigPath));
+
+				FNotificationInfo Info(NotificationErrorText);
+				Info.ExpireDuration = 3.0f;
+
+				FSlateNotificationManager::Get().AddNotification(Info);
+			}
+		}
+
+		Settings->GameplayTags.Add(NewTag);
+		Settings->SortTags();
+		IGameplayTagsModule::Get().GetGameplayTagsManager().DestroyGameplayTagTree();
+
+		{
+#if STATS
+			FString PerfMessage = FString::Printf(TEXT("ConstructGameplayTagTree GameplayTag tables"));
+			SCOPE_LOG_TIME_IN_SECONDS(*PerfMessage, nullptr)
+#endif
+			IGameplayTagsModule::Get().GetGameplayTagsManager().ConstructGameplayTagTree();
+		}
+
+		Settings->UpdateDefaultConfigFile();
+	}
+	
+}
+
 #endif // WITH_EDITOR
 
 FGameplayTag UGameplayTagsManager::RequestGameplayTag(FName TagName, bool ErrorIfNotFound) const
@@ -872,3 +940,5 @@ void FGameplayTagNode::ResetNode()
 	ChildTags.Empty();
 	ParentNode.Reset();
 }
+
+#undef LOCTEXT_NAMESPACE
\ No newline at end of file
