Added GetActiveGameplayEffectDebugString for blueprint debugging

--------
Integrated using branch UE4-To-UE4-Orion (reversed) of change#2653621 by [REDACTED_PERSON] on 2015/08/12 17:40:15.

diff --git a/Engine/Source/Runtime/GameplayAbilities/Private/AbilitySystemBlueprintLibrary.cpp b/Engine/Source/Runtime/GameplayAbilities/Private/AbilitySystemBlueprintLibrary.cpp
--- a/Engine/Source/Runtime/GameplayAbilities/Private/AbilitySystemBlueprintLibrary.cpp
+++ b/Engine/Source/Runtime/GameplayAbilities/Private/AbilitySystemBlueprintLibrary.cpp
@@ -711,4 +711,15 @@ float UAbilitySystemBlueprintLibrary::GetModifiedAttributeMagnitude(FGameplayEff
 		}
 	}
 	return Delta;
+}
+
+FString UAbilitySystemBlueprintLibrary::GetActiveGameplayEffectDebugString(FActiveGameplayEffectHandle ActiveHandle)
+{
+	FString Str;
+	UAbilitySystemComponent* ASC = ActiveHandle.GetOwningAbilitySystemComponent();
+	if (ASC)
+	{
+		Str = ASC->GetActiveGEDebugString(ActiveHandle);
+	}
+	return Str;
 }
\ No newline at end of file
diff --git a/Engine/Source/Runtime/GameplayAbilities/Public/AbilitySystemBlueprintLibrary.h b/Engine/Source/Runtime/GameplayAbilities/Public/AbilitySystemBlueprintLibrary.h
--- a/Engine/Source/Runtime/GameplayAbilities/Public/AbilitySystemBlueprintLibrary.h
+++ b/Engine/Source/Runtime/GameplayAbilities/Public/AbilitySystemBlueprintLibrary.h
@@ -239,4 +239,7 @@ class GAMEPLAYABILITIES_API UAbilitySystemBlueprintLibrary : public UBlueprintFu
 	/** Returns current stack count of an active Gameplay Effect. Will return 0 if the GameplayEffect is no longer valid. */
 	UFUNCTION(BlueprintCallable, Category = "Ability|GameplayEffect")
 	static int32 GetActiveGameplayEffectStackCount(FActiveGameplayEffectHandle ActiveHandle);
+
+	UFUNCTION(BlueprintPure, Category = "Ability|GameplayEffect", Meta = (DisplayName = "Get Active GameplayEffect Debug String "))
+	static FString GetActiveGameplayEffectDebugString(FActiveGameplayEffectHandle ActiveHandle);
 };
