Added optional parameter to ApplyGE* functions to take starting stack count

--------
Integrated using branch UE4-To-UE4-Orion (reversed) of change#2653624 by [REDACTED_PERSON] on 2015/08/12 17:40:31.

diff --git a/Engine/Source/Runtime/GameplayAbilities/Private/Abilities/GameplayAbility.cpp b/Engine/Source/Runtime/GameplayAbilities/Private/Abilities/GameplayAbility.cpp
--- a/Engine/Source/Runtime/GameplayAbilities/Private/Abilities/GameplayAbility.cpp
+++ b/Engine/Source/Runtime/GameplayAbilities/Private/Abilities/GameplayAbility.cpp
@@ -1328,7 +1328,7 @@ void UGameplayAbility::OnAvatarSet(const FGameplayAbilityActorInfo* ActorInfo, c
 
 // -------------------------------------------------------
 
-FActiveGameplayEffectHandle UGameplayAbility::BP_ApplyGameplayEffectToOwner(TSubclassOf<UGameplayEffect> GameplayEffectClass, int32 GameplayEffectLevel)
+FActiveGameplayEffectHandle UGameplayAbility::BP_ApplyGameplayEffectToOwner(TSubclassOf<UGameplayEffect> GameplayEffectClass, int32 GameplayEffectLevel, int32 Stacks)
 {
 	check(CurrentActorInfo);
 	check(CurrentSpecHandle.IsValid());
@@ -1336,31 +1336,32 @@ FActiveGameplayEffectHandle UGameplayAbility::BP_ApplyGameplayEffectToOwner(TSub
 	if ( GameplayEffectClass )
 	{
 		const UGameplayEffect* GameplayEffect = GameplayEffectClass->GetDefaultObject<UGameplayEffect>();
-		return ApplyGameplayEffectToOwner(CurrentSpecHandle, CurrentActorInfo, CurrentActivationInfo, GameplayEffect, GameplayEffectLevel);
+		return ApplyGameplayEffectToOwner(CurrentSpecHandle, CurrentActorInfo, CurrentActivationInfo, GameplayEffect, GameplayEffectLevel, Stacks);
 	}
 
 	ABILITY_LOG(Error, TEXT("BP_ApplyGameplayEffectToOwner called on ability %s with no GameplayEffectClass."), *GetName());
 	return FActiveGameplayEffectHandle();
 }
 
-FActiveGameplayEffectHandle UGameplayAbility::K2_ApplyGameplayEffectToOwner(const UGameplayEffect* GameplayEffect, int32 GameplayEffectLevel)
+FActiveGameplayEffectHandle UGameplayAbility::K2_ApplyGameplayEffectToOwner(const UGameplayEffect* GameplayEffect, int32 GameplayEffectLevel, int32 Stacks)
 {
 	if ( GameplayEffect )
 	{
-		return BP_ApplyGameplayEffectToOwner(GameplayEffect->GetClass(), GameplayEffectLevel);
+		return BP_ApplyGameplayEffectToOwner(GameplayEffect->GetClass(), GameplayEffectLevel, Stacks);
 	}
 
 	ABILITY_LOG(Error, TEXT("K2_ApplyGameplayEffectToOwner called on ability %s with no GameplayEffect."), *GetName());
 	return FActiveGameplayEffectHandle();
 }
 
-FActiveGameplayEffectHandle UGameplayAbility::ApplyGameplayEffectToOwner(const FGameplayAbilitySpecHandle Handle, const FGameplayAbilityActorInfo* ActorInfo, const FGameplayAbilityActivationInfo ActivationInfo, const UGameplayEffect* GameplayEffect, float GameplayEffectLevel) const
+FActiveGameplayEffectHandle UGameplayAbility::ApplyGameplayEffectToOwner(const FGameplayAbilitySpecHandle Handle, const FGameplayAbilityActorInfo* ActorInfo, const FGameplayAbilityActivationInfo ActivationInfo, const UGameplayEffect* GameplayEffect, float GameplayEffectLevel, int32 Stacks) const
 {
 	if (GameplayEffect && (HasAuthorityOrPredictionKey(ActorInfo, &ActivationInfo)))
 	{
 		FGameplayEffectSpecHandle SpecHandle = MakeOutgoingGameplayEffectSpec(Handle, ActorInfo, ActivationInfo, GameplayEffect->GetClass(), GameplayEffectLevel);
 		if (SpecHandle.IsValid())
 		{
+			SpecHandle.Data->StackCount = Stacks;
 			return ApplyGameplayEffectSpecToOwner(Handle, ActorInfo, ActivationInfo, SpecHandle);
 		}
 	}
@@ -1389,17 +1390,17 @@ FActiveGameplayEffectHandle UGameplayAbility::ApplyGameplayEffectSpecToOwner(con
 
 // -------------------------------
 
-TArray<FActiveGameplayEffectHandle> UGameplayAbility::BP_ApplyGameplayEffectToTarget(FGameplayAbilityTargetDataHandle Target, TSubclassOf<UGameplayEffect> GameplayEffectClass, int32 GameplayEffectLevel)
+TArray<FActiveGameplayEffectHandle> UGameplayAbility::BP_ApplyGameplayEffectToTarget(FGameplayAbilityTargetDataHandle Target, TSubclassOf<UGameplayEffect> GameplayEffectClass, int32 GameplayEffectLevel, int32 Stacks)
 {
-	return ApplyGameplayEffectToTarget(CurrentSpecHandle, CurrentActorInfo, CurrentActivationInfo, Target, GameplayEffectClass, GameplayEffectLevel);
+	return ApplyGameplayEffectToTarget(CurrentSpecHandle, CurrentActorInfo, CurrentActivationInfo, Target, GameplayEffectClass, GameplayEffectLevel, Stacks);
 }
 
-TArray<FActiveGameplayEffectHandle> UGameplayAbility::K2_ApplyGameplayEffectToTarget(FGameplayAbilityTargetDataHandle Target, const UGameplayEffect* GameplayEffect, int32 GameplayEffectLevel)
+TArray<FActiveGameplayEffectHandle> UGameplayAbility::K2_ApplyGameplayEffectToTarget(FGameplayAbilityTargetDataHandle Target, const UGameplayEffect* GameplayEffect, int32 GameplayEffectLevel, int32 Stacks)
 {
-	return BP_ApplyGameplayEffectToTarget(Target, GameplayEffect->GetClass(), GameplayEffectLevel);
+	return BP_ApplyGameplayEffectToTarget(Target, GameplayEffect->GetClass(), GameplayEffectLevel, Stacks);
 }
 
-TArray<FActiveGameplayEffectHandle> UGameplayAbility::ApplyGameplayEffectToTarget(const FGameplayAbilitySpecHandle Handle, const FGameplayAbilityActorInfo* ActorInfo, const FGameplayAbilityActivationInfo ActivationInfo, FGameplayAbilityTargetDataHandle Target, TSubclassOf<UGameplayEffect> GameplayEffectClass, float GameplayEffectLevel) const
+TArray<FActiveGameplayEffectHandle> UGameplayAbility::ApplyGameplayEffectToTarget(const FGameplayAbilitySpecHandle Handle, const FGameplayAbilityActorInfo* ActorInfo, const FGameplayAbilityActivationInfo ActivationInfo, FGameplayAbilityTargetDataHandle Target, TSubclassOf<UGameplayEffect> GameplayEffectClass, float GameplayEffectLevel, int32 Stacks) const
 {
 	TArray<FActiveGameplayEffectHandle> EffectHandles;
 
@@ -1413,6 +1414,7 @@ TArray<FActiveGameplayEffectHandle> UGameplayAbility::ApplyGameplayEffectToTarge
 	else if (HasAuthorityOrPredictionKey(ActorInfo, &ActivationInfo))
 	{
 		FGameplayEffectSpecHandle SpecHandle = MakeOutgoingGameplayEffectSpec(Handle, ActorInfo, ActivationInfo, GameplayEffectClass, GameplayEffectLevel);
+		SpecHandle.Data->StackCount = Stacks;
 		EffectHandles.Append(ApplyGameplayEffectSpecToTarget(Handle, ActorInfo, ActivationInfo, SpecHandle, Target));
 	}
 
diff --git a/Engine/Source/Runtime/GameplayAbilities/Public/Abilities/GameplayAbility.h b/Engine/Source/Runtime/GameplayAbilities/Public/Abilities/GameplayAbility.h
--- a/Engine/Source/Runtime/GameplayAbilities/Public/Abilities/GameplayAbility.h
+++ b/Engine/Source/Runtime/GameplayAbilities/Public/Abilities/GameplayAbility.h
@@ -447,13 +447,13 @@ protected:
 	// ---------
 
 	UFUNCTION(BlueprintCallable, Category = Ability, DisplayName="ApplyGameplayEffectToOwner")
-	FActiveGameplayEffectHandle BP_ApplyGameplayEffectToOwner(TSubclassOf<UGameplayEffect> GameplayEffectClass, int32 GameplayEffectLevel = 1);
+	FActiveGameplayEffectHandle BP_ApplyGameplayEffectToOwner(TSubclassOf<UGameplayEffect> GameplayEffectClass, int32 GameplayEffectLevel = 1, int32 Stacks = 1);
 
 	UFUNCTION(BlueprintCallable, Category = Ability, DisplayName="ApplyGameplayEffectToOwner", meta=(DeprecatedFunction, DeprecationMessage = "Use new ApplyGameplayEffectToOwner"))
-	FActiveGameplayEffectHandle K2_ApplyGameplayEffectToOwner(const UGameplayEffect* GameplayEffect, int32 GameplayEffectLevel = 1);
+	FActiveGameplayEffectHandle K2_ApplyGameplayEffectToOwner(const UGameplayEffect* GameplayEffect, int32 GameplayEffectLevel = 1, int32 Stacks = 1);
 
 	/** Non blueprintcallable, safe to call on CDO/NonInstance abilities */
-	FActiveGameplayEffectHandle ApplyGameplayEffectToOwner(const FGameplayAbilitySpecHandle Handle, const FGameplayAbilityActorInfo* ActorInfo, const FGameplayAbilityActivationInfo ActivationInfo, const UGameplayEffect* GameplayEffect, float GameplayEffectLevel) const;
+	FActiveGameplayEffectHandle ApplyGameplayEffectToOwner(const FGameplayAbilitySpecHandle Handle, const FGameplayAbilityActorInfo* ActorInfo, const FGameplayAbilityActivationInfo ActivationInfo, const UGameplayEffect* GameplayEffect, float GameplayEffectLevel, int32 Stacks = 1) const;
 
 	UFUNCTION(BlueprintCallable, Category = Ability, DisplayName = "ApplyGameplayEffectSpecToOwner")
 	FActiveGameplayEffectHandle K2_ApplyGameplayEffectSpecToOwner(const FGameplayEffectSpecHandle EffectSpecHandle);
@@ -465,13 +465,13 @@ protected:
 	// ---------
 
 	UFUNCTION(BlueprintCallable, Category = Ability, DisplayName = "ApplyGameplayEffectToTarget")
-	TArray<FActiveGameplayEffectHandle> BP_ApplyGameplayEffectToTarget(FGameplayAbilityTargetDataHandle TargetData, TSubclassOf<UGameplayEffect> GameplayEffectClass, int32 GameplayEffectLevel = 1);
+	TArray<FActiveGameplayEffectHandle> BP_ApplyGameplayEffectToTarget(FGameplayAbilityTargetDataHandle TargetData, TSubclassOf<UGameplayEffect> GameplayEffectClass, int32 GameplayEffectLevel = 1, int32 Stacks = 1);
 
 	UFUNCTION(BlueprintCallable, Category = Ability, DisplayName = "ApplyGameplayEffectToTarget", meta=(DeprecatedFunction, DeprecationMessage = "Use new ApplyGameplayEffectToTarget"))
-	TArray<FActiveGameplayEffectHandle> K2_ApplyGameplayEffectToTarget(FGameplayAbilityTargetDataHandle TargetData, const UGameplayEffect* GameplayEffect, int32 GameplayEffectLevel = 1);
+	TArray<FActiveGameplayEffectHandle> K2_ApplyGameplayEffectToTarget(FGameplayAbilityTargetDataHandle TargetData, const UGameplayEffect* GameplayEffect, int32 GameplayEffectLevel = 1, int32 Stacks = 1);
 
 	/** Non blueprintcallable, safe to call on CDO/NonInstance abilities */
-	TArray<FActiveGameplayEffectHandle> ApplyGameplayEffectToTarget(const FGameplayAbilitySpecHandle Handle, const FGameplayAbilityActorInfo* ActorInfo, const FGameplayAbilityActivationInfo ActivationInfo, FGameplayAbilityTargetDataHandle Target, TSubclassOf<UGameplayEffect> GameplayEffectClass, float GameplayEffectLevel) const;
+	TArray<FActiveGameplayEffectHandle> ApplyGameplayEffectToTarget(const FGameplayAbilitySpecHandle Handle, const FGameplayAbilityActorInfo* ActorInfo, const FGameplayAbilityActivationInfo ActivationInfo, FGameplayAbilityTargetDataHandle Target, TSubclassOf<UGameplayEffect> GameplayEffectClass, float GameplayEffectLevel, int32 Stacks) const;
 
 	UFUNCTION(BlueprintCallable, Category = Ability, DisplayName = "ApplyGameplayEffectSpecToTarget")
 	TArray<FActiveGameplayEffectHandle> K2_ApplyGameplayEffectSpecToTarget(const FGameplayEffectSpecHandle EffectSpecHandle, FGameplayAbilityTargetDataHandle TargetData);
