-Added GE StackingExpirationPolicy. This can be used to do waterfall style expiration of GE stacks (remove 1 stack and reset duration)
-Added WaitGameplayEffectStackChange task

--------
Integrated using branch UE4-To-UE4-Orion (reversed) of change#2653626 by [REDACTED_PERSON] on 2015/08/12 17:40:40.

diff --git a/Engine/Source/Runtime/GameplayAbilities/Private/Abilities/Tasks/AbilityTask_WaitGameplayEffectStackChange.cpp b/Engine/Source/Runtime/GameplayAbilities/Private/Abilities/Tasks/AbilityTask_WaitGameplayEffectStackChange.cpp
new file mode 100644
--- /dev/null
+++ b/Engine/Source/Runtime/GameplayAbilities/Private/Abilities/Tasks/AbilityTask_WaitGameplayEffectStackChange.cpp
@@ -0,0 +1,63 @@
+// Copyright 1998-2015 Epic Games, Inc. All Rights Reserved.
+
+#include "AbilitySystemPrivatePCH.h"
+#include "Abilities/Tasks/AbilityTask_WaitGameplayEffectStackChange.h"
+
+#include "AbilitySystemComponent.h"
+
+#include "Abilities/GameplayAbility.h"
+
+UAbilityTask_WaitGameplayEffectStackChange::UAbilityTask_WaitGameplayEffectStackChange(const FObjectInitializer& ObjectInitializer)
+	: Super(ObjectInitializer)
+{
+	Registered = false;
+}
+
+UAbilityTask_WaitGameplayEffectStackChange* UAbilityTask_WaitGameplayEffectStackChange::WaitForGameplayEffectStackChange(UObject* WorldContextObject, FActiveGameplayEffectHandle InHandle)
+{
+	auto MyObj = NewAbilityTask<UAbilityTask_WaitGameplayEffectStackChange>(WorldContextObject);
+	MyObj->Handle = InHandle;
+	return MyObj;
+}
+
+void UAbilityTask_WaitGameplayEffectStackChange::Activate()
+{
+	if (Handle.IsValid() == false)
+	{
+		InvalidHandle.Broadcast(Handle, 0, 0);
+		EndTask();
+		return;;
+	}
+
+	UAbilitySystemComponent* EffectOwningAbilitySystemComponent = Handle.GetOwningAbilitySystemComponent();
+
+	if (EffectOwningAbilitySystemComponent)
+	{
+		FOnActiveGameplayEffectStackChange* DelPtr = EffectOwningAbilitySystemComponent->OnGameplayEffectStackChangeDelegate(Handle);
+		if (DelPtr)
+		{
+			OnGameplayEffectStackChangeDelegateHandle = DelPtr->AddUObject(this, &UAbilityTask_WaitGameplayEffectStackChange::OnGameplayEffectStackChange);
+			Registered = true;
+		}
+	}
+}
+
+void UAbilityTask_WaitGameplayEffectStackChange::OnDestroy(bool AbilityIsEnding)
+{
+	UAbilitySystemComponent* EffectOwningAbilitySystemComponent = Handle.GetOwningAbilitySystemComponent();
+	if (EffectOwningAbilitySystemComponent && OnGameplayEffectStackChangeDelegateHandle.IsValid())
+	{
+		FOnActiveGameplayEffectRemoved* DelPtr = EffectOwningAbilitySystemComponent->OnGameplayEffectRemovedDelegate(Handle);
+		if (DelPtr)
+		{
+			DelPtr->Remove(OnGameplayEffectStackChangeDelegateHandle);
+		}
+	}
+
+	Super::OnDestroy(AbilityIsEnding);
+}
+
+void UAbilityTask_WaitGameplayEffectStackChange::OnGameplayEffectStackChange(FActiveGameplayEffectHandle InHandle, int32 NewCount, int32 OldCount)
+{
+	OnChange.Broadcast(InHandle, NewCount, OldCount);
+}
\ No newline at end of file
diff --git a/Engine/Source/Runtime/GameplayAbilities/Private/AbilitySystemComponent.cpp b/Engine/Source/Runtime/GameplayAbilities/Private/AbilitySystemComponent.cpp
--- a/Engine/Source/Runtime/GameplayAbilities/Private/AbilitySystemComponent.cpp
+++ b/Engine/Source/Runtime/GameplayAbilities/Private/AbilitySystemComponent.cpp
@@ -396,6 +396,17 @@ FOnActiveGameplayEffectRemoved& UAbilitySystemComponent::OnAnyGameplayEffectRemo
 	return ActiveGameplayEffects.OnActiveGameplayEffectRemovedDelegate;
 }
 
+FOnActiveGameplayEffectStackChange* UAbilitySystemComponent::OnGameplayEffectStackChangeDelegate(FActiveGameplayEffectHandle Handle)
+{
+	FActiveGameplayEffect* ActiveEffect = ActiveGameplayEffects.GetActiveGameplayEffect(Handle);
+	if (ActiveEffect)
+	{
+		return &ActiveEffect->OnStackChangeDelegate;
+	}
+
+	return nullptr;
+}
+
 int32 UAbilitySystemComponent::GetNumActiveGameplayEffects() const
 {
 	return ActiveGameplayEffects.GetNumGameplayEffects();
@@ -925,6 +936,18 @@ int32 UAbilitySystemComponent::GetCurrentStackCount(FGameplayAbilitySpecHandle H
 	return 0;
 }
 
+FString UAbilitySystemComponent::GetActiveGEDebugString(FActiveGameplayEffectHandle Handle) const
+{
+	FString Str;
+
+	if (const FActiveGameplayEffect* ActiveGE = ActiveGameplayEffects.GetActiveGameplayEffect(Handle))
+	{
+		Str = FString::Printf(TEXT("%s - (Level: %.2f. Stacks: %d)"), *ActiveGE->Spec.Def->GetName(), ActiveGE->Spec.GetLevel(), ActiveGE->Spec.StackCount);
+	}
+
+	return Str;
+}
+
 FActiveGameplayEffectHandle UAbilitySystemComponent::FindActiveGameplayEffectHandle(FGameplayAbilitySpecHandle Handle) const
 {
 	for (const FActiveGameplayEffect& ActiveGE : &ActiveGameplayEffects)
diff --git a/Engine/Source/Runtime/GameplayAbilities/Private/GameplayEffect.cpp b/Engine/Source/Runtime/GameplayAbilities/Private/GameplayEffect.cpp
--- a/Engine/Source/Runtime/GameplayAbilities/Private/GameplayEffect.cpp
+++ b/Engine/Source/Runtime/GameplayAbilities/Private/GameplayEffect.cpp
@@ -1624,12 +1624,14 @@ void FActiveGameplayEffectsContainer::OnMagnitudeDependencyChange(FActiveGamepla
 	}
 }
 
-void FActiveGameplayEffectsContainer::OnStackCountChange(FActiveGameplayEffect& ActiveEffect)
+void FActiveGameplayEffectsContainer::OnStackCountChange(FActiveGameplayEffect& ActiveEffect, int32 OldStackCount)
 {
 	MarkItemDirty(ActiveEffect);
 	UpdateAllAggregatorModMagnitudes(ActiveEffect);
 	Owner->NotifyTagMap_StackCountChange(ActiveEffect.Spec.Def->InheritableOwnedTagsContainer.CombinedTags);
 	Owner->NotifyTagMap_StackCountChange(ActiveEffect.Spec.DynamicGrantedTags);
+
+	ActiveEffect.OnStackChangeDelegate.Broadcast(ActiveEffect.Handle, ActiveEffect.Spec.StackCount, OldStackCount);
 }
 
 void FActiveGameplayEffectsContainer::UpdateAllAggregatorModMagnitudes(FActiveGameplayEffect& ActiveEffect)
@@ -2008,6 +2010,7 @@ FActiveGameplayEffect* FActiveGameplayEffectsContainer::ApplyGameplayEffectSpec(
 
 	bool bSetDuration = true;
 	bool bSetPeriod = true;
+	int32 StartingStackCount = 0;
 
 	// Check if there's an active GE this application should stack upon
 	if (ExistingStackableGE)
@@ -2019,6 +2022,7 @@ FActiveGameplayEffect* FActiveGameplayEffectsContainer::ApplyGameplayEffectSpec(
 		}
 
 		FGameplayEffectSpec& ExistingSpec = ExistingStackableGE->Spec;
+		StartingStackCount = ExistingSpec.StackCount;
 		
 		// How to apply multiple stacks at once? What if we trigger an overflow which can reject the application?
 		// We still want to apply the stacks that didnt push us over, but we also want to call HandleActiveGameplayEffectStackOverflow.
@@ -2247,7 +2251,7 @@ FActiveGameplayEffect* FActiveGameplayEffectsContainer::ApplyGameplayEffectSpec(
 	// as a result of stacking. In reality it could in complicated cases with differing sets of dynamically-granted tags.
 	if (ExistingStackableGE)
 	{
-		OnStackCountChange(*ExistingStackableGE);
+		OnStackCountChange(*ExistingStackableGE, StartingStackCount);
 	}
 	else
 	{
@@ -2394,8 +2398,9 @@ bool FActiveGameplayEffectsContainer::InternalRemoveActiveGameplayEffect(int32 I
 		if (StacksToRemove > 0 && Effect.Spec.StackCount > StacksToRemove)
 		{
 			// This won't be a full remove, only a change in StackCount.
+			int32 StartingStackCount = Effect.Spec.StackCount;
 			Effect.Spec.StackCount -= StacksToRemove;
-			OnStackCountChange(Effect);
+			OnStackCountChange(Effect, StartingStackCount);
 			return false;
 		}
 
@@ -2754,7 +2759,40 @@ void FActiveGameplayEffectsContainer::CheckDuration(FActiveGameplayEffectHandle
 			float Duration = Effect.GetDuration();
 			float CurrentTime = GetWorldTime();
 
+			int32 StacksToRemove = -2;
+			bool RefreshStartTime = false;
+			bool RefreshDurationTimer = false;
+			bool CheckForFinalPeriodicExec = false;
+
 			if (Duration > 0.f && (((Effect.StartWorldTime + Duration) < CurrentTime) || FMath::IsNearlyZero(CurrentTime - Duration - Effect.StartWorldTime, KINDA_SMALL_NUMBER)))
+			{
+				// Figure out what to do based on the expiration policy
+				switch(Effect.Spec.Def->StackExpirationPolicy)
+				{
+				case EGameplayEffectStackingExpirationPolicy::ClearEntireStack:
+					StacksToRemove = -1; // Remove all stacks
+					CheckForFinalPeriodicExec = true;					
+					break;
+
+				case EGameplayEffectStackingExpirationPolicy::RemoveSingleStackAndRefreshDuration:
+					StacksToRemove = 1;
+					CheckForFinalPeriodicExec = (Effect.Spec.StackCount == 1);
+					RefreshStartTime = true;
+					RefreshDurationTimer = true;
+					break;
+				case EGameplayEffectStackingExpirationPolicy::RefreshDuration:
+					RefreshStartTime = true;
+					RefreshDurationTimer = true;
+					break;
+				};					
+			}
+			else
+			{
+				// Effect isn't finished, just refresh its duration timer
+				RefreshDurationTimer = true;
+			}
+
+			if (CheckForFinalPeriodicExec)
 			{
 				// This gameplay effect has hit its duration. Check if it needs to execute one last time before removing it.
 				if (Effect.PeriodHandle.IsValid() && TimerManager.TimerExists(Effect.PeriodHandle))
@@ -2768,10 +2806,22 @@ void FActiveGameplayEffectsContainer::CheckDuration(FActiveGameplayEffectHandle
 					// Forcibly clear the periodic ticks because this effect is going to be removed
 					TimerManager.ClearTimer(Effect.PeriodHandle);
 				}
+			}
 
-				InternalRemoveActiveGameplayEffect(ActiveGEIdx, -1, false);
+			if (StacksToRemove >= -1)
+			{
+				InternalRemoveActiveGameplayEffect(ActiveGEIdx, StacksToRemove, false);
 			}
-			else
+
+			if (RefreshStartTime)
+			{
+				Effect.StartServerWorldTime = GetServerWorldTime();
+				Effect.CachedStartServerWorldTime = Effect.StartServerWorldTime;
+				Effect.StartWorldTime = GetWorldTime();
+				MarkItemDirty(Effect);
+			}
+
+			if (RefreshDurationTimer)
 			{
 				// Always reset the timer, since the duration might have been modified
 				FTimerDelegate Delegate = FTimerDelegate::CreateUObject(Owner, &UAbilitySystemComponent::CheckDurationExpired, Effect.Handle);
diff --git a/Engine/Source/Runtime/GameplayAbilities/Public/Abilities/Tasks/AbilityTask_WaitGameplayEffectStackChange.h b/Engine/Source/Runtime/GameplayAbilities/Public/Abilities/Tasks/AbilityTask_WaitGameplayEffectStackChange.h
new file mode 100644
--- /dev/null
+++ b/Engine/Source/Runtime/GameplayAbilities/Public/Abilities/Tasks/AbilityTask_WaitGameplayEffectStackChange.h
@@ -0,0 +1,45 @@
+// Copyright 1998-2015 Epic Games, Inc. All Rights Reserved.
+#pragma once
+#include "AbilityTask.h"
+#include "Abilities/GameplayAbilityTypes.h"
+#include "GameplayEffect.h"
+#include "AbilityTask_WaitGameplayEffectStackChange.generated.h"
+
+
+DECLARE_DYNAMIC_MULTICAST_DELEGATE_ThreeParams(FWaitGameplayEffectStackChangeDelegate, FActiveGameplayEffectHandle, Handle, int32, NewCount, int32, OldCount);
+
+class AActor;
+class UPrimitiveComponent;
+
+/**
+ *	Waits for the actor to activate another ability
+ */
+UCLASS(MinimalAPI)
+class UAbilityTask_WaitGameplayEffectStackChange : public UAbilityTask
+{
+	GENERATED_UCLASS_BODY()
+
+	UPROPERTY(BlueprintAssignable)
+	FWaitGameplayEffectStackChangeDelegate	OnChange;
+
+	UPROPERTY(BlueprintAssignable)
+	FWaitGameplayEffectStackChangeDelegate	InvalidHandle;
+
+	virtual void Activate() override;
+
+	UFUNCTION()
+	void OnGameplayEffectStackChange(FActiveGameplayEffectHandle Handle, int32 NewCount, int32 OldCount);
+
+	/** Wait until the specified gameplay effect is removed. */
+	UFUNCTION(BlueprintCallable, Category="Ability|Tasks", meta = (HidePin = "WorldContextObject", DefaultToSelf = "WorldContextObject", BlueprintInternalUseOnly = "TRUE"))
+	static UAbilityTask_WaitGameplayEffectStackChange* WaitForGameplayEffectStackChange(UObject* WorldContextObject, FActiveGameplayEffectHandle Handle);
+
+	FActiveGameplayEffectHandle Handle;
+
+protected:
+
+	virtual void OnDestroy(bool AbilityIsEnding) override;
+	bool Registered;
+
+	FDelegateHandle OnGameplayEffectStackChangeDelegateHandle;
+};
\ No newline at end of file
diff --git a/Engine/Source/Runtime/GameplayAbilities/Public/AbilitySystemComponent.h b/Engine/Source/Runtime/GameplayAbilities/Public/AbilitySystemComponent.h
--- a/Engine/Source/Runtime/GameplayAbilities/Public/AbilitySystemComponent.h
+++ b/Engine/Source/Runtime/GameplayAbilities/Public/AbilitySystemComponent.h
@@ -340,6 +340,9 @@ class GAMEPLAYABILITIES_API UAbilitySystemComponent : public UGameplayTasksCompo
 	/** Returns current stack count of an already applied GE, but given the ability spec handle that was granted by the GE */
 	int32 GetCurrentStackCount(FGameplayAbilitySpecHandle Handle) const;
 
+	/** Returns debug string describing active gameplay effect */
+	FString GetActiveGEDebugString(FActiveGameplayEffectHandle Handle) const;
+
 	/** Gets the GE Handle of the GE that granted the passed in Ability */
 	FActiveGameplayEffectHandle FindActiveGameplayEffectHandle(FGameplayAbilitySpecHandle Handle) const;
 
@@ -455,6 +458,8 @@ class GAMEPLAYABILITIES_API UAbilitySystemComponent : public UGameplayTasksCompo
 	FOnActiveGameplayEffectRemoved* OnGameplayEffectRemovedDelegate(FActiveGameplayEffectHandle Handle);
 	FOnActiveGameplayEffectRemoved& OnAnyGameplayEffectRemovedDelegate();
 
+	FOnActiveGameplayEffectStackChange* OnGameplayEffectStackChangeDelegate(FActiveGameplayEffectHandle Handle);
+
 	UFUNCTION(BlueprintCallable, Category = GameplayEffects, meta=(DisplayName = "ApplyGameplayEffectToTarget"))
 	FActiveGameplayEffectHandle BP_ApplyGameplayEffectToTarget(TSubclassOf<UGameplayEffect> GameplayEffectClass, UAbilitySystemComponent *Target, float Level, FGameplayEffectContextHandle Context);
 
diff --git a/Engine/Source/Runtime/GameplayAbilities/Public/GameplayEffect.h b/Engine/Source/Runtime/GameplayAbilities/Public/GameplayEffect.h
--- a/Engine/Source/Runtime/GameplayAbilities/Public/GameplayEffect.h
+++ b/Engine/Source/Runtime/GameplayAbilities/Public/GameplayEffect.h
@@ -509,6 +509,20 @@ enum class EGameplayEffectStackingPeriodPolicy : uint8
 	NeverReset,
 };
 
+/** Enumeration of policies for dealing gameplay effect stacks that expire (in duration based effects). */
+UENUM()
+enum class EGameplayEffectStackingExpirationPolicy : uint8
+{
+	/** The entire stack is cleared when the active gameplay effect expires  */
+	ClearEntireStack,
+
+	/** The current stack count will be decremented by 1 and the duration refreshed. The GE is not "reapplied", just continues to exist with one less stacks. */
+	RemoveSingleStackAndRefreshDuration,
+
+	/** The duration of the gameplay effect is refreshed. This essentially makes the effect infinite in duration. This can be used to manually handle stack decrements via XXX callback */
+	RefreshDuration,
+};
+
 /**
  * UGameplayEffect
  *	The GameplayEffect definition. This is the data asset defined in the editor that drives everything.
@@ -691,6 +705,10 @@ public:
 	UPROPERTY(EditDefaultsOnly, BlueprintReadOnly, Category = Stacking)
 	EGameplayEffectStackingPeriodPolicy StackPeriodResetPolicy;
 
+	/** Policy for how to handle duration expiring on this gameplay effect */
+	UPROPERTY(EditDefaultsOnly, BlueprintReadOnly, Category = Stacking)
+	EGameplayEffectStackingExpirationPolicy StackExpirationPolicy;
+
 	// ----------------------------------------------------------------------
 	//	Granted abilities
 	// ----------------------------------------------------------------------
@@ -1267,7 +1285,9 @@ struct GAMEPLAYABILITIES_API FActiveGameplayEffect : public FFastArraySerializer
 
 	bool IsPendingRemove;
 
-	FOnActiveGameplayEffectRemoved	OnRemovedDelegate;
+	FOnActiveGameplayEffectRemoved OnRemovedDelegate;
+
+	FOnActiveGameplayEffectStackChange OnStackChangeDelegate;
 
 	FTimerHandle PeriodHandle;
 
@@ -1706,7 +1726,7 @@ private:
 
 	void OnMagnitudeDependencyChange(FActiveGameplayEffectHandle Handle, const FAggregator* ChangedAgg);
 
-	void OnStackCountChange(FActiveGameplayEffect& ActiveEffect);
+	void OnStackCountChange(FActiveGameplayEffect& ActiveEffect, int32 OldStackCount);
 
 	void UpdateAllAggregatorModMagnitudes(FActiveGameplayEffect& ActiveEffect);
 
diff --git a/Engine/Source/Runtime/GameplayAbilities/Public/GameplayEffectTypes.h b/Engine/Source/Runtime/GameplayAbilities/Public/GameplayEffectTypes.h
--- a/Engine/Source/Runtime/GameplayAbilities/Public/GameplayEffectTypes.h
+++ b/Engine/Source/Runtime/GameplayAbilities/Public/GameplayEffectTypes.h
@@ -803,6 +803,8 @@ DECLARE_MULTICAST_DELEGATE_TwoParams(FOnGameplayEffectTagCountChanged, const FGa
 
 DECLARE_MULTICAST_DELEGATE(FOnActiveGameplayEffectRemoved);
 
+DECLARE_MULTICAST_DELEGATE_ThreeParams(FOnActiveGameplayEffectStackChange, FActiveGameplayEffectHandle, int32, int32);
+
 DECLARE_MULTICAST_DELEGATE_TwoParams(FOnGameplayAttributeChange, float, const FGameplayEffectModCallbackData*);
 
 DECLARE_DELEGATE_RetVal(FGameplayTagContainer, FGetGameplayTags);
