Integrate two changes from Orion to main:

2646936 - Add analog cursor to login screen and frontend for PS4
2656390 - Fix issue on client with anim montages error correcting if the replication structure says it's stopped

diff --git a/Engine/Source/Runtime/GameplayAbilities/Private/AbilitySystemComponent_Abilities.cpp b/Engine/Source/Runtime/GameplayAbilities/Private/AbilitySystemComponent_Abilities.cpp
--- a/Engine/Source/Runtime/GameplayAbilities/Private/AbilitySystemComponent_Abilities.cpp
+++ b/Engine/Source/Runtime/GameplayAbilities/Private/AbilitySystemComponent_Abilities.cpp
@@ -2272,7 +2272,7 @@ void UAbilitySystemComponent::OnRep_ReplicatedAnimMontage()
 			float CurrentPosition = AnimInstance->Montage_GetPosition(LocalAnimMontageInfo.AnimMontage);
 			int32 CurrentSectionID = LocalAnimMontageInfo.AnimMontage->GetSectionIndexFromPosition(CurrentPosition);
 			// Only check threshold if we are located in the same section. Different sections require a bit more work as we could be jumping around the timeline.
-			if( (CurrentSectionID == RepSectionID) && (FMath::Abs(CurrentPosition - RepAnimMontageInfo.Position) > MONTAGE_REP_POS_ERR_THRESH) )
+			if ((CurrentSectionID == RepSectionID) && (FMath::Abs(CurrentPosition - RepAnimMontageInfo.Position) > MONTAGE_REP_POS_ERR_THRESH) && RepAnimMontageInfo.IsStopped == 0)
 			{
 				AnimInstance->Montage_SetPosition(LocalAnimMontageInfo.AnimMontage, RepAnimMontageInfo.Position);
 			}
