Fix issue with SetByCaller GE durations
-Now default these to a duration of 1.f (arbitrary) until the SetByCaller magnitude is set. Otherwise the system will treat the GE as an instant effect and early out before the SetByCallerMagnitude can be read.

--------
Integrated using branch UE4-To-UE4-Orion (reversed) of change#2663350 by [REDACTED_PERSON] on 2015/08/20 15:49:52.

diff --git a/Engine/Source/Runtime/GameplayAbilities/Private/GameplayEffect.cpp b/Engine/Source/Runtime/GameplayAbilities/Private/GameplayEffect.cpp
--- a/Engine/Source/Runtime/GameplayAbilities/Private/GameplayEffect.cpp
+++ b/Engine/Source/Runtime/GameplayAbilities/Private/GameplayEffect.cpp
@@ -224,7 +224,7 @@ bool FGameplayEffectModifierMagnitude::CanCalculateMagnitude(const FGameplayEffe
 	return InRelevantSpec.HasValidCapturedAttributes(ReqCaptureDefs);
 }
 
-bool FGameplayEffectModifierMagnitude::AttemptCalculateMagnitude(const FGameplayEffectSpec& InRelevantSpec, OUT float& OutCalculatedMagnitude) const
+bool FGameplayEffectModifierMagnitude::AttemptCalculateMagnitude(const FGameplayEffectSpec& InRelevantSpec, OUT float& OutCalculatedMagnitude, bool WarnIfSetByCallerFail, float DefaultSetbyCaller) const
 {
 	const bool bCanCalc = CanCalculateMagnitude(InRelevantSpec);
 	if (bCanCalc)
@@ -250,7 +250,7 @@ bool FGameplayEffectModifierMagnitude::AttemptCalculateMagnitude(const FGameplay
 			break;
 
 			case EGameplayEffectMagnitudeCalculation::SetByCaller:
-				OutCalculatedMagnitude = InRelevantSpec.GetSetByCallerMagnitude(SetByCallerMagnitude.DataName);
+				OutCalculatedMagnitude = InRelevantSpec.GetSetByCallerMagnitude(SetByCallerMagnitude.DataName, WarnIfSetByCallerFail, DefaultSetbyCaller);
 				break;
 			default:
 				ABILITY_LOG(Error, TEXT("Unknown MagnitudeCalculationType %d in AttemptCalculateMagnitude"), (int32)MagnitudeCalculationType);
@@ -640,7 +640,10 @@ bool FGameplayEffectSpec::AttemptCalculateDurationFromDef(OUT float& OutDefDurat
 	}
 	else
 	{
-		bCalculatedDuration = Def->DurationMagnitude.AttemptCalculateMagnitude(*this, OutDefDuration);
+		// The last parameters (false, 1.f) are so that if SetByCaller hasn't been set yet, we don't warn and default
+		// to 1.f. This is so that the rest of the system doesn't treat the effect as an instant effect. 1.f is arbitrary
+		// and this makes it illegal to SetByCaller something into an instant effect.
+		bCalculatedDuration = Def->DurationMagnitude.AttemptCalculateMagnitude(*this, OutDefDuration, false, 1.f);
 	}
 
 	return bCalculatedDuration;
@@ -810,15 +813,15 @@ void FGameplayEffectSpec::SetSetByCallerMagnitude(FName DataName, float Magnitud
 	SetByCallerMagnitudes.Add(DataName) = Magnitude;
 }
 
-float FGameplayEffectSpec::GetSetByCallerMagnitude(FName DataName) const
+float FGameplayEffectSpec::GetSetByCallerMagnitude(FName DataName, bool WarnIfNotFound, float DefaultIfNotFound) const
 {
-	float Magnitude = 0.f;
+	float Magnitude = DefaultIfNotFound;
 	const float* Ptr = SetByCallerMagnitudes.Find(DataName);
 	if (Ptr)
 	{
 		Magnitude = *Ptr;
 	}
-	else
+	else if (WarnIfNotFound)
 	{
 		ABILITY_LOG(Error, TEXT("FGameplayEffectSpec::GetMagnitude called for Data %s on Def %s when magnitude had not yet been set by caller."), *DataName.ToString(), *Def->GetName());
 	}
diff --git a/Engine/Source/Runtime/GameplayAbilities/Public/GameplayEffect.h b/Engine/Source/Runtime/GameplayAbilities/Public/GameplayEffect.h
--- a/Engine/Source/Runtime/GameplayAbilities/Public/GameplayEffect.h
+++ b/Engine/Source/Runtime/GameplayAbilities/Public/GameplayEffect.h
@@ -224,7 +224,7 @@ public:
 	 * 
 	 * @return True if the calculation was successful, false if it was not
 	 */
-	bool AttemptCalculateMagnitude(const FGameplayEffectSpec& InRelevantSpec, OUT float& OutCalculatedMagnitude) const;
+	bool AttemptCalculateMagnitude(const FGameplayEffectSpec& InRelevantSpec, OUT float& OutCalculatedMagnitude, bool WarnIfSetByCallerFail=true, float DefaultSetbyCaller=0.f) const;
 
 	/** Attempts to recalculate the magnitude given a changed aggregator. This will only recalculate if we are a modifier that is linked (non snapshot) to the given aggregator. */
 	bool AttemptRecalculateMagnitudeFromDependentChange(const FGameplayEffectSpec& InRelevantSpec, OUT float& OutCalculatedMagnitude, const FAggregator* ChangedAggregator) const;
@@ -1021,7 +1021,7 @@ struct GAMEPLAYABILITIES_API FGameplayEffectSpec
 	void SetSetByCallerMagnitude(FName DataName, float Magnitude);
 
 	/** Returns the magnitude of a SetByCaller modifier. Will return 0.f and Warn if the magnitude has not been set. */
-	float GetSetByCallerMagnitude(FName DataName) const;
+	float GetSetByCallerMagnitude(FName DataName, bool WarnIfNotFound=true, float DefaultIfNotFound=0.f) const;
 
 	void SetLevel(float InLevel);
 
