Add safety checks and fallbacks in GetGameplayCueManager

--------
Integrated using branch UE4-To-UE4-Orion (reversed) of change#2664473 by [REDACTED_PERSON] on 2015/08/21 11:19:20.

diff --git a/Engine/Source/Runtime/GameplayAbilities/Private/AbilitySystemGlobals.cpp b/Engine/Source/Runtime/GameplayAbilities/Private/AbilitySystemGlobals.cpp
--- a/Engine/Source/Runtime/GameplayAbilities/Private/AbilitySystemGlobals.cpp
+++ b/Engine/Source/Runtime/GameplayAbilities/Private/AbilitySystemGlobals.cpp
@@ -249,13 +249,28 @@ UGameplayCueManager* UAbilitySystemGlobals::GetGameplayCueManager()
 {
 	if (GlobalGameplayCueManager == nullptr)
 	{
-		GlobalGameplayCueManager = LoadObject<UGameplayCueManager>(nullptr, *GlobalGameplayCueManagerName.ToString(), nullptr, LOAD_None, nullptr);
+		if (GlobalGameplayCueManagerName.IsValid())
+		{
+			GlobalGameplayCueManager = LoadObject<UGameplayCueManager>(nullptr, *GlobalGameplayCueManagerName.ToString(), nullptr, LOAD_None, nullptr);
+			if (GlobalGameplayCueManager == nullptr)
+			{
+				ABILITY_LOG(Error, TEXT("Unable to Load GameplayCueManager %s"), *GlobalGameplayCueManagerName.ToString() );
+			}
+		}
+
+		if ( GlobalGameplayCueManager == nullptr)
+		{
+			// Fallback to CDO
+			GlobalGameplayCueManager = UGameplayCueManager::StaticClass()->GetDefaultObject<UGameplayCueManager>();
+		}
+
 		if (GameplayCueNotifyPaths.Num() > 0)
 		{
 			GlobalGameplayCueManager->LoadObjectLibraryFromPaths(GameplayCueNotifyPaths);
 		}
 	}
 
+	check(GlobalGameplayCueManager);
 	return GlobalGameplayCueManager;
 }
 
