minor cleanup

--------
Integrated using branch UE4-To-UE4-Orion (reversed) of change#2665027 by [REDACTED_PERSON] on 2015/08/21 17:37:29.

diff --git a/Engine/Source/Runtime/GameplayAbilities/Private/GameplayEffect.cpp b/Engine/Source/Runtime/GameplayAbilities/Private/GameplayEffect.cpp
--- a/Engine/Source/Runtime/GameplayAbilities/Private/GameplayEffect.cpp
+++ b/Engine/Source/Runtime/GameplayAbilities/Private/GameplayEffect.cpp
@@ -2079,9 +2079,7 @@ FActiveGameplayEffect* FActiveGameplayEffectsContainer::ApplyGameplayEffectSpec(
 		}
 		else
 		{
-			ExistingStackableGE->StartServerWorldTime = GetServerWorldTime();
-			ExistingStackableGE->CachedStartServerWorldTime = ExistingStackableGE->StartServerWorldTime;
-			ExistingStackableGE->StartWorldTime = GetWorldTime();
+			RestartActiveGameplayEffectDuration(*ExistingStackableGE);
 		}
 
 		// Make sure the GE actually wants to reset its period
@@ -2628,6 +2626,14 @@ void FActiveGameplayEffectsContainer::InternalApplyExpirationEffects(const FGame
 	}
 }
 
+void FActiveGameplayEffectsContainer::RestartActiveGameplayEffectDuration(FActiveGameplayEffect& ActiveGameplayEffect)
+{
+	ActiveGameplayEffect.StartServerWorldTime = GetServerWorldTime();
+	ActiveGameplayEffect.CachedStartServerWorldTime = ActiveGameplayEffect.StartServerWorldTime;
+	ActiveGameplayEffect.StartWorldTime = GetWorldTime();
+	MarkItemDirty(ActiveGameplayEffect);
+}
+
 void FActiveGameplayEffectsContainer::OnOwnerTagChange(FGameplayTag TagChange, int32 NewCount)
 {
 	// It may be beneficial to do a scoped lock on attribute re-evaluation during this function
@@ -2818,10 +2824,7 @@ void FActiveGameplayEffectsContainer::CheckDuration(FActiveGameplayEffectHandle
 
 			if (RefreshStartTime)
 			{
-				Effect.StartServerWorldTime = GetServerWorldTime();
-				Effect.CachedStartServerWorldTime = Effect.StartServerWorldTime;
-				Effect.StartWorldTime = GetWorldTime();
-				MarkItemDirty(Effect);
+				RestartActiveGameplayEffectDuration(Effect);
 			}
 
 			if (RefreshDurationTimer)
diff --git a/Engine/Source/Runtime/GameplayAbilities/Public/GameplayEffect.h b/Engine/Source/Runtime/GameplayAbilities/Public/GameplayEffect.h
--- a/Engine/Source/Runtime/GameplayAbilities/Public/GameplayEffect.h
+++ b/Engine/Source/Runtime/GameplayAbilities/Public/GameplayEffect.h
@@ -1707,6 +1707,8 @@ private:
 	/** Internal helper function to apply expiration effects from a removed/expired gameplay effect spec */
 	void InternalApplyExpirationEffects(const FGameplayEffectSpec& ExpiringSpec, bool bPrematureRemoval);
 
+	void RestartActiveGameplayEffectDuration(FActiveGameplayEffect& ActiveGameplayEffect);
+
 	// -------------------------------------------------------------------------------------------
 
 	TMap<FGameplayAttribute, FAggregatorRef>		AttributeAggregatorMap;
