Convert granted ability level property into scalable float so that GEs can grant abilities of arbitrary levels

--------
Integrated using branch UE4-To-UE4-Orion (reversed) of change#2665101 by [REDACTED_PERSON] on 2015/08/21 18:15:28.

diff --git a/Engine/Source/Runtime/GameplayAbilities/Private/GameplayAbilityTypes.cpp b/Engine/Source/Runtime/GameplayAbilities/Private/GameplayAbilityTypes.cpp
--- a/Engine/Source/Runtime/GameplayAbilities/Private/GameplayAbilityTypes.cpp
+++ b/Engine/Source/Runtime/GameplayAbilities/Private/GameplayAbilityTypes.cpp
@@ -172,9 +172,9 @@ void FGameplayAbilitySpecContainer::RegisterWithOwner(UAbilitySystemComponent* I
 
 // ----------------------------------------------------
 
-FGameplayAbilitySpec:: FGameplayAbilitySpec(FGameplayAbilitySpecDef& InDef, FActiveGameplayEffectHandle InGameplayEffectHandle)
+FGameplayAbilitySpec:: FGameplayAbilitySpec(FGameplayAbilitySpecDef& InDef, int32 InGameplayEffectLevel, FActiveGameplayEffectHandle InGameplayEffectHandle)
 	: Ability(InDef.Ability ? InDef.Ability->GetDefaultObject<UGameplayAbility>() : nullptr)
-	, Level(InDef.Level)
+	, Level(InDef.LevelScalableFloat.GetValueAtLevel(InGameplayEffectLevel))
 	, InputID(InDef.InputID)
 	, SourceObject(InDef.SourceObject)
 	, InputPressed(false)
diff --git a/Engine/Source/Runtime/GameplayAbilities/Private/GameplayEffect.cpp b/Engine/Source/Runtime/GameplayAbilities/Private/GameplayEffect.cpp
--- a/Engine/Source/Runtime/GameplayAbilities/Private/GameplayEffect.cpp
+++ b/Engine/Source/Runtime/GameplayAbilities/Private/GameplayEffect.cpp
@@ -92,6 +92,15 @@ void UGameplayEffect::PostLoad()
 	ClearTagsContainer.RemoveAllTags();
 
 	UpdateInheritedTagProperties();
+
+	for (FGameplayAbilitySpecDef& Def : GrantedAbilities)
+	{
+		if (Def.Level != INDEX_NONE)
+		{
+			Def.LevelScalableFloat.SetValue(Def.Level);
+			Def.Level = INDEX_NONE;
+		}
+	}
 }
 
 void UGameplayEffect::PostInitProperties()
@@ -2334,7 +2343,7 @@ void FActiveGameplayEffectsContainer::AddActiveGameplayEffectGrantedTagsAndModif
 			// would regrant the ability every time the stack was applied
 			if (AbilitySpecDef.AssignedHandle.IsValid() == false)
 			{
-				Owner->GiveAbility( FGameplayAbilitySpec(AbilitySpecDef, Effect.Handle) );
+				Owner->GiveAbility( FGameplayAbilitySpec(AbilitySpecDef, Effect.Spec.GetLevel(), Effect.Handle) );
 			}
 		}	
 	}
diff --git a/Engine/Source/Runtime/GameplayAbilities/Public/GameplayAbilitySpec.h b/Engine/Source/Runtime/GameplayAbilities/Public/GameplayAbilitySpec.h
--- a/Engine/Source/Runtime/GameplayAbilities/Public/GameplayAbilitySpec.h
+++ b/Engine/Source/Runtime/GameplayAbilities/Public/GameplayAbilitySpec.h
@@ -90,10 +90,11 @@ USTRUCT(BlueprintType)
 struct FGameplayAbilitySpecDef
 {
 	FGameplayAbilitySpecDef()
-		: Level(1)
+		: Level(INDEX_NONE)
 		, InputID(INDEX_NONE)
 		, RemovalPolicy(EGameplayEffectGrantedAbilityRemovePolicy::CancelAbilityImmediately)
 	{
+		LevelScalableFloat.SetValue(1.f);
 	}
 
 	GENERATED_USTRUCT_BODY()
@@ -101,8 +102,12 @@ struct FGameplayAbilitySpecDef
 	UPROPERTY(EditDefaultsOnly, Category="Ability Definition", NotReplicated)
 	TSubclassOf<UGameplayAbility> Ability;
 	
-	UPROPERTY(EditDefaultsOnly, Category="Ability Definition", NotReplicated)
+	// Deprecated for LevelScalableFloat
+	UPROPERTY(NotReplicated)
 	int32 Level; 
+
+	UPROPERTY(EditDefaultsOnly, Category="Ability Definition", NotReplicated, DisplayName="Level")
+	FScalableFloat LevelScalableFloat; 
 	
 	UPROPERTY(EditDefaultsOnly, Category="Ability Definition", NotReplicated)
 	int32 InputID;
@@ -200,7 +205,7 @@ struct GAMEPLAYABILITIES_API FGameplayAbilitySpec : public FFastArraySerializerI
 		Handle.GenerateNewHandle();
 	}
 	
-	FGameplayAbilitySpec(FGameplayAbilitySpecDef& InDef, FActiveGameplayEffectHandle InGameplayEffectHandle = FActiveGameplayEffectHandle());
+	FGameplayAbilitySpec(FGameplayAbilitySpecDef& InDef, int32 InGameplayEffectLevel, FActiveGameplayEffectHandle InGameplayEffectHandle = FActiveGameplayEffectHandle());
 
 	/** Handle for outside sources to refer to this spec by */
 	UPROPERTY()
