Make FGameplayAbilityTargetData::ApplyGameplayEffectSpec virtual and take in a non const spec for cases where we want to modify the spec in target data subclasses

--------
Integrated using branch UE4-To-UE4-Orion (reversed) of change#2669350 by [REDACTED_PERSON] on 2015/08/26 11:37:25.

diff --git a/Engine/Source/Runtime/GameplayAbilities/Private/GameplayAbilityTargetTypes.cpp b/Engine/Source/Runtime/GameplayAbilities/Private/GameplayAbilityTargetTypes.cpp
--- a/Engine/Source/Runtime/GameplayAbilities/Private/GameplayAbilityTargetTypes.cpp
+++ b/Engine/Source/Runtime/GameplayAbilities/Private/GameplayAbilityTargetTypes.cpp
@@ -15,7 +15,7 @@ TArray<FActiveGameplayEffectHandle> FGameplayAbilityTargetData::ApplyGameplayEff
 	return ApplyGameplayEffectSpec(TempSpecToApply, PredictionKey);
 }
 
-TArray<FActiveGameplayEffectHandle> FGameplayAbilityTargetData::ApplyGameplayEffectSpec(const FGameplayEffectSpec& InSpec, FPredictionKey PredictionKey)
+TArray<FActiveGameplayEffectHandle> FGameplayAbilityTargetData::ApplyGameplayEffectSpec(FGameplayEffectSpec& InSpec, FPredictionKey PredictionKey)
 {
 	TArray<FActiveGameplayEffectHandle>	AppliedHandles;
 
diff --git a/Engine/Source/Runtime/GameplayAbilities/Public/Abilities/GameplayAbilityTargetTypes.h b/Engine/Source/Runtime/GameplayAbilities/Public/Abilities/GameplayAbilityTargetTypes.h
--- a/Engine/Source/Runtime/GameplayAbilities/Public/Abilities/GameplayAbilityTargetTypes.h
+++ b/Engine/Source/Runtime/GameplayAbilities/Public/Abilities/GameplayAbilityTargetTypes.h
@@ -65,7 +65,7 @@ struct GAMEPLAYABILITIES_API FGameplayAbilityTargetData
 
 	TArray<FActiveGameplayEffectHandle> ApplyGameplayEffect(const UGameplayEffect* GameplayEffect, const FGameplayEffectContextHandle& InEffectContext, float Level, FPredictionKey PredictionKey = FPredictionKey());
 
-	TArray<FActiveGameplayEffectHandle> ApplyGameplayEffectSpec(const FGameplayEffectSpec& Spec, FPredictionKey PredictionKey = FPredictionKey());
+	virtual TArray<FActiveGameplayEffectHandle> ApplyGameplayEffectSpec(FGameplayEffectSpec& Spec, FPredictionKey PredictionKey = FPredictionKey());
 
 	virtual void AddTargetDataToContext(FGameplayEffectContextHandle& Context, bool bIncludeActorArray);
 
