Fix case where gameplaycue editor was not updating references to notifies that were renamed prioir to being saved to disk.

--------
Integrated using branch UE4-To-UE4-Orion (reversed) of change#2669880 by [REDACTED_PERSON] on 2015/08/26 16:34:32.

diff --git a/Engine/Source/Runtime/GameplayAbilities/Private/GameplayCueManager.cpp b/Engine/Source/Runtime/GameplayAbilities/Private/GameplayCueManager.cpp
--- a/Engine/Source/Runtime/GameplayAbilities/Private/GameplayCueManager.cpp
+++ b/Engine/Source/Runtime/GameplayAbilities/Private/GameplayCueManager.cpp
@@ -434,6 +434,8 @@ void UGameplayCueManager::HandleAssetRenamed(const FAssetData& Data, const FStri
 			if (StaticCDO || ActorCDO)
 			{
 				VerifyNotifyAssetIsInValidPath(Data.PackagePath.ToString());
+				GlobalCueSet->UpdateCueByStringRefs(String + TEXT("_C"), Data.ObjectPath.ToString() + TEXT("_C"));
+				OnGameplayCueNotifyAddOrRemove.Broadcast();
 			}
 		}
 	}
diff --git a/Engine/Source/Runtime/GameplayAbilities/Private/GameplayCueSet.cpp b/Engine/Source/Runtime/GameplayAbilities/Private/GameplayCueSet.cpp
--- a/Engine/Source/Runtime/GameplayAbilities/Private/GameplayCueSet.cpp
+++ b/Engine/Source/Runtime/GameplayAbilities/Private/GameplayCueSet.cpp
@@ -90,6 +90,22 @@ void UGameplayCueSet::RemoveCuesByStringRefs(const TArray<FStringAssetReference>
 	}
 }
 
+#if WITH_EDITOR
+void UGameplayCueSet::UpdateCueByStringRefs(const FStringAssetReference& CueToRemove, FString NewPath)
+{
+	
+	for (int32 idx = 0; idx < GameplayCueData.Num(); ++idx)
+	{
+		if (GameplayCueData[idx].GameplayCueNotifyObj == CueToRemove)
+		{
+			GameplayCueData[idx].GameplayCueNotifyObj = NewPath;
+			BuildAccelerationMap_Internal();
+			break;
+		}
+	}
+}
+#endif
+
 void UGameplayCueSet::Empty()
 {
 	GameplayCueData.Empty();
diff --git a/Engine/Source/Runtime/GameplayAbilities/Public/GameplayCueSet.h b/Engine/Source/Runtime/GameplayAbilities/Public/GameplayCueSet.h
--- a/Engine/Source/Runtime/GameplayAbilities/Public/GameplayCueSet.h
+++ b/Engine/Source/Runtime/GameplayAbilities/Public/GameplayCueSet.h
@@ -59,6 +59,13 @@ class GAMEPLAYABILITIES_API UGameplayCueSet : public UDataAsset
 	/** Removes all cues from the set matching the supplied string refs */
 	virtual void RemoveCuesByStringRefs(const TArray<FStringAssetReference>& CuesToRemove);
 
+#if WITH_EDITOR
+
+	/** Updates an existing cue */
+	virtual void UpdateCueByStringRefs(const FStringAssetReference& CueToRemove, FString NewPath);
+
+#endif
+
 	/** Removes all cues from the set */
 	virtual void Empty();
 
