virtual function for gameplaycuemanager to specify whether it should async load all cues at startup

--------
Integrated using branch UE4-To-UE4-Orion (reversed) of change#2678508 by [REDACTED_PERSON] on 2015/09/03 09:57:33.

diff --git a/Engine/Source/Runtime/GameplayAbilities/Private/GameplayCueManager.cpp b/Engine/Source/Runtime/GameplayAbilities/Private/GameplayCueManager.cpp
--- a/Engine/Source/Runtime/GameplayAbilities/Private/GameplayCueManager.cpp
+++ b/Engine/Source/Runtime/GameplayAbilities/Private/GameplayCueManager.cpp
@@ -253,7 +253,7 @@ void UGameplayCueManager::LoadObjectLibrary_Internal()
 	// ---------------------------------------------------------
 
 	const bool bSyncFullyLoad = IsRunningCommandlet();
-	const bool bAsyncLoadAtStartup = !bSyncFullyLoad;
+	const bool bAsyncLoadAtStartup = !bSyncFullyLoad && ShouldAsyncLoadAtStartup();
 	if (bSyncFullyLoad)
 	{
 #if STATS
diff --git a/Engine/Source/Runtime/GameplayAbilities/Public/GameplayCueManager.h b/Engine/Source/Runtime/GameplayAbilities/Public/GameplayCueManager.h
--- a/Engine/Source/Runtime/GameplayAbilities/Public/GameplayCueManager.h
+++ b/Engine/Source/Runtime/GameplayAbilities/Public/GameplayCueManager.h
@@ -245,6 +245,8 @@ protected:
 
 	void LoadObjectLibrary_Internal();
 
+	virtual bool ShouldAsyncLoadAtStartup() const { return true; }
+
 	void BuildCuesToAddToGlobalSet(const TArray<FAssetData>& AssetDataList, FName TagPropertyName, bool bAsyncLoadAfterAdd, TArray<struct FGameplayCueReferencePair>& OutCuesToAdd);
 
 	TArray<FString>	LoadedPaths;
