UE4 - Improved performance for ability system singleton.

diff --git a/Engine/Source/Runtime/GameplayAbilities/Private/GameplayAbilitiesModule.cpp b/Engine/Source/Runtime/GameplayAbilities/Private/GameplayAbilitiesModule.cpp
--- a/Engine/Source/Runtime/GameplayAbilities/Private/GameplayAbilitiesModule.cpp
+++ b/Engine/Source/Runtime/GameplayAbilities/Private/GameplayAbilitiesModule.cpp
@@ -13,6 +13,7 @@ class FGameplayAbilitiesModule : public IGameplayAbilitiesModule
 
 	virtual UAbilitySystemGlobals* GetAbilitySystemGlobals() override
 	{
+		QUICK_SCOPE_CYCLE_COUNTER(STAT_IGameplayAbilitiesModule_GetAbilitySystemGlobals);
 		// Defer loading of globals to the first time it is requested
 		if (!AbilitySystemGlobals)
 		{
@@ -31,6 +32,7 @@ class FGameplayAbilitiesModule : public IGameplayAbilitiesModule
 
 	virtual bool IsAbilitySystemGlobalsAvailable() override
 	{
+		QUICK_SCOPE_CYCLE_COUNTER(STAT_IGameplayAbilitiesModule_IsAbilitySystemGlobalsAvailable);
 		return AbilitySystemGlobals != nullptr;
 	}
 
diff --git a/Engine/Source/Runtime/GameplayAbilities/Public/GameplayAbilitiesModule.h b/Engine/Source/Runtime/GameplayAbilities/Public/GameplayAbilitiesModule.h
--- a/Engine/Source/Runtime/GameplayAbilities/Public/GameplayAbilitiesModule.h
+++ b/Engine/Source/Runtime/GameplayAbilities/Public/GameplayAbilitiesModule.h
@@ -22,7 +22,9 @@ public:
 	 */
 	static inline IGameplayAbilitiesModule& Get()
 	{
-		return FModuleManager::LoadModuleChecked< IGameplayAbilitiesModule >("GameplayAbilities");
+		QUICK_SCOPE_CYCLE_COUNTER(STAT_IGameplayAbilitiesModule_Get);
+		static IGameplayAbilitiesModule& Singleton = FModuleManager::LoadModuleChecked< IGameplayAbilitiesModule >("GameplayAbilities");
+		return Singleton;
 	}
 
 	/**
@@ -32,6 +34,7 @@ public:
 	 */
 	static inline bool IsAvailable()
 	{
+		QUICK_SCOPE_CYCLE_COUNTER(STAT_IGameplayAbilitiesModule_IsAvailable);
 		return FModuleManager::Get().IsModuleLoaded( "GameplayAbilities" );
 	}
 
