Merging CL 2691463 to Main

Allow derived gameplay cue managers to not begin async loading of object libraries.
Reduce priority of async object library loads

diff --git a/Engine/Source/Runtime/GameplayAbilities/Private/AbilitySystemGlobals.cpp b/Engine/Source/Runtime/GameplayAbilities/Private/AbilitySystemGlobals.cpp
--- a/Engine/Source/Runtime/GameplayAbilities/Private/AbilitySystemGlobals.cpp
+++ b/Engine/Source/Runtime/GameplayAbilities/Private/AbilitySystemGlobals.cpp
@@ -217,6 +217,13 @@ void UAbilitySystemGlobals::InitGameplayCueParameters(FGameplayCueParameters& Cu
 
 // --------------------------------------------------------------------
 
+void UAbilitySystemGlobals::StartAsyncLoadingObjectLibraries()
+{
+	GlobalGameplayCueManager->LoadObjectLibraryFromPaths(GameplayCueNotifyPaths);
+}
+
+// --------------------------------------------------------------------
+
 /** Initialize FAttributeSetInitter. This is virtual so projects can override what class they use */
 void UAbilitySystemGlobals::AllocAttributeSetInitter()
 {
@@ -263,10 +270,10 @@ UGameplayCueManager* UAbilitySystemGlobals::GetGameplayCueManager()
 			// Fallback to CDO
 			GlobalGameplayCueManager = UGameplayCueManager::StaticClass()->GetDefaultObject<UGameplayCueManager>();
 		}
-
-		if (GameplayCueNotifyPaths.Num() > 0)
+		
+		if (GlobalGameplayCueManager->ShouldAsyncLoadObjectLibrariesAtStart() && GameplayCueNotifyPaths.Num() > 0)
 		{
-			GlobalGameplayCueManager->LoadObjectLibraryFromPaths(GameplayCueNotifyPaths);
+			StartAsyncLoadingObjectLibraries();
 		}
 	}
 
diff --git a/Engine/Source/Runtime/GameplayAbilities/Private/GameplayCueManager.cpp b/Engine/Source/Runtime/GameplayAbilities/Private/GameplayCueManager.cpp
--- a/Engine/Source/Runtime/GameplayAbilities/Private/GameplayCueManager.cpp
+++ b/Engine/Source/Runtime/GameplayAbilities/Private/GameplayCueManager.cpp
@@ -314,7 +314,7 @@ void UGameplayCueManager::BuildCuesToAddToGlobalSet(const TArray<FAssetData>& As
 
 				if (bAsyncLoadAfterAdd)
 				{
-					StreamableManager.SimpleAsyncLoad(StringRef);
+					StreamableManager.SimpleAsyncLoad(StringRef, FStreamableManager::DefaultAsyncLoadPriority - 1);
 				}
 			}
 			else
diff --git a/Engine/Source/Runtime/GameplayAbilities/Public/AbilitySystemGlobals.h b/Engine/Source/Runtime/GameplayAbilities/Public/AbilitySystemGlobals.h
--- a/Engine/Source/Runtime/GameplayAbilities/Public/AbilitySystemGlobals.h
+++ b/Engine/Source/Runtime/GameplayAbilities/Public/AbilitySystemGlobals.h
@@ -167,6 +167,11 @@ class GAMEPLAYABILITIES_API UAbilitySystemGlobals : public UObject
 	virtual void InitGameplayCueParameters(FGameplayCueParameters& CueParameters, const FGameplayEffectSpecForRPC &Spec);
 	virtual void InitGameplayCueParameters(FGameplayCueParameters& CueParameters, const FGameplayEffectContextHandle& EffectContext);
 
+	// Trigger async loading of the gameplay cue object libraries. By default, the manager will do this on creation,
+	// but that behaviour can be changed by a derived class overriding ShouldAsyncLoadObjectLibrariesAtStart and returning false.
+	// In that case, this function must be called to begin the load
+	virtual void StartAsyncLoadingObjectLibraries();
+
 protected:
 
 	virtual void InitAttributeDefaults();
diff --git a/Engine/Source/Runtime/GameplayAbilities/Public/GameplayCueManager.h b/Engine/Source/Runtime/GameplayAbilities/Public/GameplayCueManager.h
--- a/Engine/Source/Runtime/GameplayAbilities/Public/GameplayCueManager.h
+++ b/Engine/Source/Runtime/GameplayAbilities/Public/GameplayCueManager.h
@@ -235,6 +235,8 @@ class GAMEPLAYABILITIES_API UGameplayCueManager : public UDataAsset
 
 	static UWorld* PreviewWorld;
 #endif
+	
+	virtual bool ShouldAsyncLoadObjectLibrariesAtStart() const { return true; }
 
 protected:
 
