Fixed ability system crash when reconnecting to a server

* Fixed crash caused by ability instances being deleted without being removed from the instance array.  This would happen to any ability that was flagged as not cancel-able.

=== Bugs Fixed ===
https://jira.ol.epicgames.net/browse/OR-6805

--------
Integrated using branch UE4-To-UE4-Orion (reversed) of change#2695342 by [REDACTED_PERSON] on 2015/09/17 11:42:56.

diff --git a/Engine/Source/Runtime/GameplayAbilities/Private/AbilitySystemComponent_Abilities.cpp b/Engine/Source/Runtime/GameplayAbilities/Private/AbilitySystemComponent_Abilities.cpp
--- a/Engine/Source/Runtime/GameplayAbilities/Private/AbilitySystemComponent_Abilities.cpp
+++ b/Engine/Source/Runtime/GameplayAbilities/Private/AbilitySystemComponent_Abilities.cpp
@@ -718,7 +718,7 @@ void UAbilitySystemComponent::DestroyActiveState()
 
 		// Mark pending kill any remaining instanced abilities
 		// (CancelAbilities() will only MarkPending kill InstancePerExecution abilities).
-		for (const FGameplayAbilitySpec& Spec : GetActivatableAbilities())
+		for (FGameplayAbilitySpec& Spec : ActivatableAbilities.Items)
 		{
 			TArray<UGameplayAbility*> AbilitiesToCancel = Spec.GetAbilityInstances();
 			for (UGameplayAbility* InstanceAbility : AbilitiesToCancel)
@@ -728,6 +728,9 @@ void UAbilitySystemComponent::DestroyActiveState()
 					InstanceAbility->MarkPendingKill();
 				}
 			}
+
+			Spec.ReplicatedInstances.Empty();
+			Spec.NonReplicatedInstances.Empty();
 		}
 	}
 }
