adding more logs to ensure condition that is failing

--------
Integrated using branch UE4-To-UE4-Orion (reversed) of change#2695196 by [REDACTED_PERSON] on 2015/09/17 10:45:14.

diff --git a/Engine/Source/Runtime/GameplayAbilities/Private/AbilitySystemComponent_Abilities.cpp b/Engine/Source/Runtime/GameplayAbilities/Private/AbilitySystemComponent_Abilities.cpp
--- a/Engine/Source/Runtime/GameplayAbilities/Private/AbilitySystemComponent_Abilities.cpp
+++ b/Engine/Source/Runtime/GameplayAbilities/Private/AbilitySystemComponent_Abilities.cpp
@@ -2180,7 +2180,12 @@ void UAbilitySystemComponent::AnimMontage_UpdateReplicatedData()
 		if (CurrentSectionID != INDEX_NONE)
 		{
 			int32 NextSectionID = AnimInstance->Montage_GetNextSectionID(LocalAnimMontageInfo.AnimMontage, CurrentSectionID);
-			ensure(NextSectionID < (256 - 1));
+			if (NextSectionID >= (256 - 1))
+			{
+				ABILITY_LOG( Error, TEXT("AnimMontage_UpdateReplicatedData. NextSectionID = %d.  RepAnimMontageInfo.Position: %.2f, CurrentSectionID: %d. LocalAnimMontageInfo.AnimMontage %s"), 
+					NextSectionID, RepAnimMontageInfo.Position, CurrentSectionID, *GetNameSafe(LocalAnimMontageInfo.AnimMontage) );
+				ensure(NextSectionID < (256 - 1));
+			}
 			RepAnimMontageInfo.NextSectionID = uint8(NextSectionID + 1);
 		}
 		else
