Added a check for set by caller magnitudes in FGameplayEffectSpec::AttemptCalculateDurationFromDef. This prevents errors when trying to create gameplay effect specs in blueprints from definitions with set by caller magnitudes on the duration.

diff --git a/Engine/Source/Runtime/GameplayAbilities/Private/GameplayEffect.cpp b/Engine/Source/Runtime/GameplayAbilities/Private/GameplayEffect.cpp
--- a/Engine/Source/Runtime/GameplayAbilities/Private/GameplayEffect.cpp
+++ b/Engine/Source/Runtime/GameplayAbilities/Private/GameplayEffect.cpp
@@ -647,6 +647,10 @@ bool FGameplayEffectSpec::AttemptCalculateDurationFromDef(OUT float& OutDefDurat
 	{
 		OutDefDuration = UGameplayEffect::INSTANT_APPLICATION;
 	}
+	else if (Def->DurationMagnitude.GetMagnitudeCalculationType() == EGameplayEffectMagnitudeCalculation::SetByCaller)
+	{
+		return false;
+	}
 	else
 	{
 		// The last parameters (false, 1.f) are so that if SetByCaller hasn't been set yet, we don't warn and default
