Copying //UE4/Orion-Staging to Dev-Main (Originates from //Orion/Main @ CL-2759376)

diff --git a/Engine/Source/Programs/UnrealBuildTool/Windows/WindowsDeploy.cs b/Engine/Source/Programs/UnrealBuildTool/Windows/WindowsDeploy.cs
new file mode 100644
--- /dev/null
+++ b/Engine/Source/Programs/UnrealBuildTool/Windows/WindowsDeploy.cs
@@ -0,0 +1,57 @@
+﻿// Copyright 1998-2014 Epic Games, Inc. All Rights Reserved.
+
+using System;
+using System.Collections.Generic;
+using System.Text;
+using System.IO;
+using System.Diagnostics;
+
+// @todo UWP: this file is a work in progress and is not yet complete for the F5 scenario for UWP
+
+namespace UnrealBuildTool
+{
+	/// <summary>
+	///  Base class to handle deploy of a target for a given platform
+	/// </summary>
+	class BaseWindowsDeploy : UEBuildDeploy
+	{
+		public bool PrepForUATPackageOrDeploy(string InProjectName, string InProjectDirectory, string InExecutablePath, string InEngineDir, bool bForDistribution, string CookFlavor)
+		{
+			// set the icon on the original exe this will be used in the task bar when the bootstrap exe runs
+			if (File.Exists(Path.Combine(InProjectDirectory, "Build/Windows/Application.ico")))
+			{
+				GroupIconResource GroupIcon = null;
+				GroupIcon = GroupIconResource.FromIco(Path.Combine(InProjectDirectory, "Build/Windows/Application.ico"));
+
+				// Update the icon on the original exe because this will be used when the game is running in the task bar
+				using (ModuleResourceUpdate Update = new ModuleResourceUpdate(InExecutablePath, false))
+				{
+					const int IconResourceId = 101;
+					if (GroupIcon != null)
+					{
+						Update.SetIcons(IconResourceId, GroupIcon);
+					}
+				}
+			}
+			return true;
+		}
+
+
+
+		public override bool PrepTargetForDeployment(UEBuildTarget InTarget)
+		{
+			if ((InTarget.TargetType != TargetRules.TargetType.Editor && InTarget.TargetType != TargetRules.TargetType.Program) && (BuildHostPlatform.Current.Platform == UnrealTargetPlatform.Win32 || BuildHostPlatform.Current.Platform == UnrealTargetPlatform.Win64))
+			{
+				string InAppName = InTarget.AppName;
+				Log.TraceInformation("Prepping {0} for deployment to {1}", InAppName, InTarget.Platform.ToString());
+				System.DateTime PrepDeployStartTime = DateTime.UtcNow;
+
+				string TargetFilename = InTarget.RulesAssembly.GetTargetFileName(InAppName).FullName;
+				string ProjectSourceFolder = new FileInfo(TargetFilename).DirectoryName + "/";
+
+				PrepForUATPackageOrDeploy(InAppName, InTarget.ProjectDirectory.FullName, InTarget.OutputPath.FullName, BuildConfiguration.RelativeEnginePath, false, "");
+			}
+			return true;
+		}
+	}
+}
diff --git a/Engine/Source/Runtime/Analytics/AnalyticsET/Private/AnalyticsET.cpp b/Engine/Source/Runtime/Analytics/AnalyticsET/Private/AnalyticsET.cpp
--- a/Engine/Source/Runtime/Analytics/AnalyticsET/Private/AnalyticsET.cpp
+++ b/Engine/Source/Runtime/Analytics/AnalyticsET/Private/AnalyticsET.cpp
@@ -12,6 +12,7 @@
 #include "EngineVersion.h"
 
 IMPLEMENT_MODULE( FAnalyticsET, AnalyticsET );
+DEFINE_LOG_CATEGORY(LogAnalytics);
 
 class FAnalyticsProviderET : 
 	public IAnalyticsProvider,
diff --git a/Engine/Source/Runtime/Analytics/AnalyticsET/Public/AnalyticsET.h b/Engine/Source/Runtime/Analytics/AnalyticsET/Public/AnalyticsET.h
--- a/Engine/Source/Runtime/Analytics/AnalyticsET/Public/AnalyticsET.h
+++ b/Engine/Source/Runtime/Analytics/AnalyticsET/Public/AnalyticsET.h
@@ -5,7 +5,7 @@
 #include "IAnalyticsProviderModule.h"
 #include "Core.h"
 
-DEFINE_LOG_CATEGORY_STATIC(LogAnalytics, Display, All);
+DECLARE_LOG_CATEGORY_EXTERN(LogAnalytics, Display, All);
 
 class IAnalyticsProvider;
 
diff --git a/Engine/Source/Runtime/Analytics/QoSReporter/Private/QoSReporter.cpp b/Engine/Source/Runtime/Analytics/QoSReporter/Private/QoSReporter.cpp
--- a/Engine/Source/Runtime/Analytics/QoSReporter/Private/QoSReporter.cpp
+++ b/Engine/Source/Runtime/Analytics/QoSReporter/Private/QoSReporter.cpp
@@ -204,18 +204,10 @@ void FQoSReporter::AddServerHeartbeatAttributes(TArray<FAnalyticsEventAttribute>
 	{
 		const TMap<FString, IPerfCounters::FJsonVariant>& PerfCounterMap = PerfCounters->GetAllCounters();
 		int NumAddedCounters = 0;
-		// we want to skip reporting average FPS if there are no human players
-		bool bShouldSkipAverageFPS = (PerfCounters->Get("NumClients", 0) == 0);
 
 		// only add string and number values
 		for (const auto& It : PerfCounterMap)
 		{
-			// do not report average FPS unless there are actual people playing
-			if (bShouldSkipAverageFPS && It.Key == TEXT("AverageFPS"))
-			{
-				continue;
-			}
-
 			const IPerfCounters::FJsonVariant& JsonValue = It.Value;
 			switch (JsonValue.Format)
 			{
@@ -260,6 +252,7 @@ void FQoSReporter::AddServerHeartbeatAttributes(TArray<FAnalyticsEventAttribute>
 	else if (UE_SERVER)
 	{
 		UE_LOG(LogQoSReporter, Warning, TEXT("PerfCounters module is not available, could not send proper server heartbeat."));
+		OutArray.Add(FAnalyticsEventAttribute(TEXT("MisconfiguredPerfCounters"), 1));
 	}
 #else
 	#if UE_SERVER
diff --git a/Engine/Source/Runtime/AnimGraphRuntime/Private/AnimNodes/AnimNode_LayeredBoneBlend.cpp b/Engine/Source/Runtime/AnimGraphRuntime/Private/AnimNodes/AnimNode_LayeredBoneBlend.cpp
--- a/Engine/Source/Runtime/AnimGraphRuntime/Private/AnimNodes/AnimNode_LayeredBoneBlend.cpp
+++ b/Engine/Source/Runtime/AnimGraphRuntime/Private/AnimNodes/AnimNode_LayeredBoneBlend.cpp
@@ -60,6 +60,7 @@ void FAnimNode_LayeredBoneBlend::Update(const FAnimationUpdateContext& Context)
 	// initialize children
 	BasePose.Update(Context);
 
+	bHasRelevantPoses = false;
 	const int NumPoses = BlendPoses.Num();
 	if ( NumPoses > 0 )
 	{
@@ -67,6 +68,7 @@ void FAnimNode_LayeredBoneBlend::Update(const FAnimationUpdateContext& Context)
 		{
 			if (BlendWeights[ChildIndex] > ZERO_ANIMWEIGHT_THRESH)
 			{
+				bHasRelevantPoses = true;
 				BlendPoses[ChildIndex].Update(Context);
 			}
 		}
@@ -78,7 +80,7 @@ void FAnimNode_LayeredBoneBlend::Evaluate(FPoseContext& Output)
 	ANIM_MT_SCOPE_CYCLE_COUNTER(BlendPosesInGraph, Output.AnimInstance->IsRunningParallelEvaluation());
 
 	const int NumPoses = BlendPoses.Num();
-	if (NumPoses == 0)
+	if ((NumPoses == 0) || !bHasRelevantPoses)
 	{
 		BasePose.Evaluate(Output);
 	}
diff --git a/Engine/Source/Runtime/AnimGraphRuntime/Private/AnimNodes/AnimNode_SaveCachedPose.cpp b/Engine/Source/Runtime/AnimGraphRuntime/Private/AnimNodes/AnimNode_SaveCachedPose.cpp
--- a/Engine/Source/Runtime/AnimGraphRuntime/Private/AnimNodes/AnimNode_SaveCachedPose.cpp
+++ b/Engine/Source/Runtime/AnimGraphRuntime/Private/AnimNodes/AnimNode_SaveCachedPose.cpp
@@ -76,7 +76,10 @@ void FAnimNode_SaveCachedPose::GatherDebugData(FNodeDebugData& DebugData)
 	FString DebugLine = DebugData.GetNodeName(this);
 	DebugData.AddDebugItem(DebugLine);
 
-	if (!DebugDataCounter.IsSynchronizedWith(DebugData.AnimInstance->DebugDataCounter))
+	// we need to fix SaveCachePose to collect the highest weight branch, not the first one,
+	// for it to be useful in debug.
+	if(false)
+// 	if (!DebugDataCounter.IsSynchronizedWith(DebugData.AnimInstance->DebugDataCounter))
 	{
 		DebugDataCounter.SynchronizeWith(DebugData.AnimInstance->DebugDataCounter);
 		Pose.GatherDebugData(DebugData);
diff --git a/Engine/Source/Runtime/AnimGraphRuntime/Public/AnimNodes/AnimNode_LayeredBoneBlend.h b/Engine/Source/Runtime/AnimGraphRuntime/Public/AnimNodes/AnimNode_LayeredBoneBlend.h
--- a/Engine/Source/Runtime/AnimGraphRuntime/Public/AnimNodes/AnimNode_LayeredBoneBlend.h
+++ b/Engine/Source/Runtime/AnimGraphRuntime/Public/AnimNodes/AnimNode_LayeredBoneBlend.h
@@ -34,6 +34,9 @@ public:
 	UPROPERTY(EditAnywhere, BlueprintReadWrite, Category=Config)
 	TEnumAsByte<enum ECurveBlendOption::Type>	CurveBlendOption;
 
+	UPROPERTY(Transient)
+	bool bHasRelevantPoses;
+
 protected:
 	TArray<FPerBoneBlendWeight> DesiredBoneBlendWeights;
 	TArray<FPerBoneBlendWeight> CurrentBoneBlendWeights;
diff --git a/Engine/Source/Runtime/Engine/Classes/Animation/AnimNode_StateMachine.h b/Engine/Source/Runtime/Engine/Classes/Animation/AnimNode_StateMachine.h
--- a/Engine/Source/Runtime/Engine/Classes/Animation/AnimNode_StateMachine.h
+++ b/Engine/Source/Runtime/Engine/Classes/Animation/AnimNode_StateMachine.h
@@ -229,6 +229,8 @@ protected:
 	void EvaluateTransitionStandardBlendInternal(FPoseContext& Output, FAnimationActiveTransitionEntry& Transition, const FPoseContext& PreviousStateResult, const FPoseContext& NextStateResult);
 	void EvaluateTransitionCustomBlend(FPoseContext& Output, FAnimationActiveTransitionEntry& Transition, bool bIntermediatePoseIsValid);
 
+	FAnimNode_AssetPlayerBase* GetRelevantAssetPlayerFromState(const FAnimationUpdateContext& Context, const FBakedAnimationState& StateInfo);
+
 public:
 	friend class UAnimInstance;
 };
diff --git a/Engine/Source/Runtime/Engine/Classes/Animation/AnimStateMachineTypes.h b/Engine/Source/Runtime/Engine/Classes/Animation/AnimStateMachineTypes.h
--- a/Engine/Source/Runtime/Engine/Classes/Animation/AnimStateMachineTypes.h
+++ b/Engine/Source/Runtime/Engine/Classes/Animation/AnimStateMachineTypes.h
@@ -180,13 +180,13 @@ struct FBakedStateExitTransition
 	UPROPERTY()
 	int32 TransitionIndex;
 
-	// For transitions that go automatically at the end of the sequence, this is the index for the sole sequence player node in the graph
-	UPROPERTY()
-	int32 StateSequencePlayerToQueryIndex;
-
 	// What the transition rule node needs to return to take this transition (for bidirectional transitions)
 	UPROPERTY()
 	bool bDesiredTransitionReturnValue;
+
+	// Automatic Transition Rule based on animation remaining time.
+	UPROPERTY()
+	bool bAutomaticRemainingTimeRule;
 	
 	UPROPERTY()
 	TArray<int32> PoseEvaluatorLinks;
@@ -195,8 +195,8 @@ struct FBakedStateExitTransition
 		: CanTakeDelegateIndex(INDEX_NONE)
 		, CustomResultNodeIndex(INDEX_NONE)
 		, TransitionIndex(INDEX_NONE)
-		, StateSequencePlayerToQueryIndex(INDEX_NONE)
 		, bDesiredTransitionReturnValue(true)
+		, bAutomaticRemainingTimeRule(false)
 	{
 	}
 };
diff --git a/Engine/Source/Runtime/Engine/Classes/Engine/GameInstance.h b/Engine/Source/Runtime/Engine/Classes/Engine/GameInstance.h
--- a/Engine/Source/Runtime/Engine/Classes/Engine/GameInstance.h
+++ b/Engine/Source/Runtime/Engine/Classes/Engine/GameInstance.h
@@ -92,7 +92,7 @@ public:
 	void InitializeStandalone();
 #if WITH_EDITOR
 	/* Called to initialize the game instance for PIE instances of the game */
-	bool InitializePIE(bool bAnyBlueprintErrors, int32 PIEInstance, bool bRunAsDedicated);
+	virtual bool InitializePIE(bool bAnyBlueprintErrors, int32 PIEInstance, bool bRunAsDedicated);
 
 	virtual bool StartPIEGameInstance(ULocalPlayer* LocalPlayer, bool bInSimulateInEditor, bool bAnyBlueprintErrors, bool bStartInSpectatorMode);
 
diff --git a/Engine/Source/Runtime/Engine/Private/Animation/AnimInstance.cpp b/Engine/Source/Runtime/Engine/Private/Animation/AnimInstance.cpp
--- a/Engine/Source/Runtime/Engine/Private/Animation/AnimInstance.cpp
+++ b/Engine/Source/Runtime/Engine/Private/Animation/AnimInstance.cpp
@@ -674,7 +674,7 @@ bool UAnimInstance::NeedsImmediateUpdate(float DeltaSeconds) const
 	return
 		CVarUseParallelAnimUpdate.GetValueOnGameThread() == 0 ||
 		CVarUseParallelAnimationEvaluation.GetValueOnGameThread() == 0 ||
-		IsRunningDedicatedServer() ||
+		GetWorld()->IsServer() ||
 		!bCanUseParallelUpdateAnimation ||
 		DeltaSeconds == 0.0f ||
 		RootMotionMode == ERootMotionMode::RootMotionFromEverything;
diff --git a/Engine/Source/Runtime/Engine/Private/Animation/AnimMontage.cpp b/Engine/Source/Runtime/Engine/Private/Animation/AnimMontage.cpp
--- a/Engine/Source/Runtime/Engine/Private/Animation/AnimMontage.cpp
+++ b/Engine/Source/Runtime/Engine/Private/Animation/AnimMontage.cpp
@@ -25,6 +25,9 @@ UAnimMontage::UAnimMontage(const FObjectInitializer& ObjectInitializer)
 	BlendOut.SetBlendTime(0.25f);
 	BlendOutTriggerTime = -1.f;
 	SyncSlotIndex = 0;
+
+	BlendInTime_DEPRECATED = -1.f;
+	BlendOutTime_DEPRECATED = -1.f;
 }
 
 bool UAnimMontage::IsValidSlot(FName InSlotName) const
@@ -462,13 +465,13 @@ void UAnimMontage::PostLoad()
 	}
 
 	// fix up blending time deprecated variable
-	if (BlendInTime_DEPRECATED == -1.f)
+	if (BlendInTime_DEPRECATED != -1.f)
 	{
 		BlendIn.SetBlendTime(BlendInTime_DEPRECATED);
 		BlendInTime_DEPRECATED = -1.f;
 	}
 
-	if(BlendOutTime_DEPRECATED == -1.f)
+	if(BlendOutTime_DEPRECATED != -1.f)
 	{
 		BlendOut.SetBlendTime(BlendOutTime_DEPRECATED);
 		BlendOutTime_DEPRECATED = -1.f;
@@ -1013,7 +1016,8 @@ void UAnimMontage::TickAssetPlayerInstance(FAnimTickRecord& Instance, class UAni
 {
 	// nothing has to happen here
 	// we just have to make sure we set Context data correct
-	if (ensure (Context.IsLeader()))
+	//if (ensure (Context.IsLeader()))
+	if ((Context.IsLeader()))
 	{
 		const float CurrentTime = Instance.Montage.CurrentPosition;
 		const float PreviousTime = Instance.Montage.PreviousPosition;
diff --git a/Engine/Source/Runtime/Engine/Private/Animation/AnimNode_StateMachine.cpp b/Engine/Source/Runtime/Engine/Private/Animation/AnimNode_StateMachine.cpp
--- a/Engine/Source/Runtime/Engine/Private/Animation/AnimNode_StateMachine.cpp
+++ b/Engine/Source/Runtime/Engine/Private/Animation/AnimNode_StateMachine.cpp
@@ -479,6 +479,24 @@ void FAnimNode_StateMachine::Update(const FAnimationUpdateContext& Context)
 	ElapsedTime += Context.GetDeltaTime();
 }
 
+FAnimNode_AssetPlayerBase* FAnimNode_StateMachine::GetRelevantAssetPlayerFromState(const FAnimationUpdateContext& Context, const FBakedAnimationState& StateInfo)
+{
+	FAnimNode_AssetPlayerBase* ResultPlayer = nullptr;
+	float MaxWeight = 0.0f;
+	for (const int32& PlayerIdx : StateInfo.PlayerNodeIndices)
+	{
+		if (FAnimNode_AssetPlayerBase* Player = Context.AnimInstance->GetNodeFromIndex<FAnimNode_AssetPlayerBase>(PlayerIdx))
+		{
+			if (!Player->bIgnoreForRelevancyTest && (Player->GetCachedBlendWeight() > MaxWeight))
+			{
+				MaxWeight = Player->GetCachedBlendWeight();
+				ResultPlayer = Player;
+			}
+		}
+	}
+	return ResultPlayer;
+}
+
 bool FAnimNode_StateMachine::FindValidTransition(const FAnimationUpdateContext& Context, const FBakedAnimationState& StateInfo, /*out*/ FAnimationPotentialTransition& OutPotentialTransition, /*out*/ TArray<int32, TInlineAllocator<4>>& OutVisitedStateIndices)
 {
 	// There is a possibility we'll revisit states connected through conduits,
@@ -530,33 +548,24 @@ bool FAnimNode_StateMachine::FindValidTransition(const FAnimationUpdateContext&
 			// attempt to evaluate native rule
 			ResultNode->bCanEnterTransition = ResultNode->NativeTransitionDelegate.Execute();
 		}
-		else
+		else if (TransitionRule.bAutomaticRemainingTimeRule)
 		{
-			bool bStillCallEvaluate = true;
-
-			if (TransitionRule.StateSequencePlayerToQueryIndex != INDEX_NONE)
+			bool bCanEnterTransition = false;
+			if (FAnimNode_AssetPlayerBase* RelevantPlayer = GetRelevantAssetPlayerFromState(Context, StateInfo))
 			{
-				// Simple automatic rule
-				FAnimNode_SequencePlayer* SequencePlayer = GetNodeFromPropertyIndex<FAnimNode_SequencePlayer>(Context.AnimInstance, AnimBlueprintClass, TransitionRule.StateSequencePlayerToQueryIndex);
-				if ((SequencePlayer != nullptr) && (SequencePlayer->Sequence != nullptr))
+				if (UAnimationAsset* AnimAsset = RelevantPlayer->GetAnimAsset())
 				{
-					const float SequenceLength = SequencePlayer->Sequence->GetMaxCurrentTime();
-					const float PlayerTime = SequencePlayer->GetAccumulatedTime();
-					const float PlayerTimeLeft = SequenceLength - PlayerTime;
-
+					const float AnimTimeRemaining = AnimAsset->GetMaxCurrentTime() - RelevantPlayer->GetAccumulatedTime();
 					const FAnimationTransitionBetweenStates& TransitionInfo = GetTransitionInfo(TransitionRule.TransitionIndex);
-
-					ResultNode->bCanEnterTransition = (PlayerTimeLeft <= TransitionInfo.CrossfadeDuration);
-
-					bStillCallEvaluate = false;
+					bCanEnterTransition = (AnimTimeRemaining <= TransitionInfo.CrossfadeDuration);
 				}
 			}
-
-			if (bStillCallEvaluate)
-			{
-				// Execute it and see if we can take this rule
-				ResultNode->EvaluateGraphExposedInputs.Execute(Context);
-			}
+			ResultNode->bCanEnterTransition = bCanEnterTransition;
+		}			
+		else 
+		{
+			// Execute it and see if we can take this rule
+			ResultNode->EvaluateGraphExposedInputs.Execute(Context);
 		}
 
 		if (ResultNode->bCanEnterTransition == TransitionRule.bDesiredTransitionReturnValue)
diff --git a/Engine/Source/Runtime/Engine/Private/GameEngine.cpp b/Engine/Source/Runtime/Engine/Private/GameEngine.cpp
--- a/Engine/Source/Runtime/Engine/Private/GameEngine.cpp
+++ b/Engine/Source/Runtime/Engine/Private/GameEngine.cpp
@@ -203,15 +203,18 @@ void UGameEngine::ConditionallyOverrideSettings(int32& ResolutionX, int32& Resol
 	FParse::Value(FCommandLine::Get(), TEXT("ResX="), ResolutionX);
 	FParse::Value(FCommandLine::Get(), TEXT("ResY="), ResolutionY);
 
-		// consume available desktop area
-		FDisplayMetrics DisplayMetrics;
+	// consume available desktop area
+	FDisplayMetrics DisplayMetrics;
 	FDisplayMetrics::GetDisplayMetrics(DisplayMetrics);
 		
 	int32 DesktopResolutionX = DisplayMetrics.PrimaryDisplayWidth;
 	int32 DesktopResolutionY = DisplayMetrics.PrimaryDisplayHeight;
 
+	// Optionally force the resolution by passing -ForceRes
+	const bool bForceRes = FParse::Param(FCommandLine::Get(), TEXT("ForceRes"));
+
 	//Dont allow a resolution bigger then the desktop found a convenient one
-	if (!IsRunningDedicatedServer() && ((ResolutionX <= 0 || ResolutionX >= DesktopResolutionX) || (ResolutionY <= 0 || ResolutionY >= DesktopResolutionY)))
+	if (!bForceRes && !IsRunningDedicatedServer() && ((ResolutionX <= 0 || ResolutionX >= DesktopResolutionX) || (ResolutionY <= 0 || ResolutionY >= DesktopResolutionY)))
 	{
 		ResolutionX = DesktopResolutionX;
 		ResolutionY = DesktopResolutionY;
@@ -282,13 +285,15 @@ TSharedRef<SWindow> UGameEngine::CreateGameWindow()
 	const FText PlatformBits = FText::FromString( TEXT( "32" ) );
 #endif	//PLATFORM_64BITS
 
+	// Note: If these parameters are updated or renamed, please update the tooltip on the ProjectDisplayedTitle property
 	FFormatNamedArguments Args;
 	Args.Add( TEXT("GameName"), FText::FromString( FApp::GetGameName() ) );
 	Args.Add( TEXT("PlatformArchitecture"), PlatformBits );
 	Args.Add( TEXT("RHIName"), FText::FromName( LegacyShaderPlatformToShaderFormat( GMaxRHIShaderPlatform ) ) );
 
-	const FText AppName = FText::Format( NSLOCTEXT("UnrealEd", "GameWindowTitle", "{GameName} ({PlatformArchitecture}-bit, {RHIName})"), Args );
+	const FText DefaultWindowTitle = NSLOCTEXT("UnrealEd", "GameWindowTitle", "{GameName} ({PlatformArchitecture}-bit, {RHIName})");
 	const FText WindowTitleOverride = GetDefault<UGeneralProjectSettings>()->ProjectDisplayedTitle;
+	const FText WindowTitle = FText::Format(WindowTitleOverride.IsEmpty() ? DefaultWindowTitle : WindowTitleOverride, Args);
 
 	// Allow optional winX/winY parameters to set initial window position
 	EAutoCenter::Type AutoCenterType = EAutoCenter::PrimaryWorkArea;
@@ -301,7 +306,7 @@ TSharedRef<SWindow> UGameEngine::CreateGameWindow()
 
 	TSharedRef<SWindow> Window = SNew(SWindow)
 	.ClientSize(FVector2D( ResX, ResY ))
-	.Title(WindowTitleOverride.IsEmpty() ? AppName : WindowTitleOverride)
+	.Title(WindowTitle)
 	.AutoCenter(AutoCenterType)
 	.ScreenPosition(FVector2D(WinX, WinY))
 	.FocusWhenFirstShown(true)
diff --git a/Engine/Source/Runtime/Engine/Private/LocalPlayer.cpp b/Engine/Source/Runtime/Engine/Private/LocalPlayer.cpp
--- a/Engine/Source/Runtime/Engine/Private/LocalPlayer.cpp
+++ b/Engine/Source/Runtime/Engine/Private/LocalPlayer.cpp
@@ -1444,6 +1444,21 @@ void ULocalPlayer::SetControllerId( int32 NewControllerId )
 
 FString ULocalPlayer::GetNickname() const
 {
+	// Try to get platform identity first
+	IOnlineSubsystem* PlatformSubsystem = IOnlineSubsystem::GetByPlatform(false);
+	if (PlatformSubsystem)
+	{
+		IOnlineIdentityPtr OnlineIdentityInt = PlatformSubsystem->GetIdentityInterface();
+		if (OnlineIdentityInt.IsValid())
+		{
+			FString PlayerNickname = OnlineIdentityInt->GetPlayerNickname(ControllerId);
+			if (!PlayerNickname.IsEmpty())
+			{
+				return PlayerNickname;
+			}
+		}
+	}
+
 	UWorld* World = GetWorld();
 	if (World != NULL)
 	{
diff --git a/Engine/Source/Runtime/Engine/Private/NetConnection.cpp b/Engine/Source/Runtime/Engine/Private/NetConnection.cpp
--- a/Engine/Source/Runtime/Engine/Private/NetConnection.cpp
+++ b/Engine/Source/Runtime/Engine/Private/NetConnection.cpp
@@ -100,9 +100,7 @@ void UNetConnection::InitBase(UNetDriver* InDriver,class FSocket* InSocket, cons
 	// Reset Handler
 	Handler.Reset(NULL);
 
-#if ENABLER_HANDLERS == 1
 	Handler = MakeUnique<PacketHandler>();
-#endif
 
 	if(Handler.IsValid())
 	{
@@ -169,12 +167,10 @@ void UNetConnection::InitConnection(UNetDriver* InDriver, EConnectionState InSta
 {
 	Driver = InDriver;
 
-#if ENABLER_HANDLERS == 1
 	if (!Handler.IsValid())
 	{
 		Handler = MakeUnique<PacketHandler>();
 	}
-#endif
 
 	// We won't be sending any packets, so use a default size
 	MaxPacket = (InMaxPacket == 0 || InMaxPacket > MAX_PACKET_SIZE) ? MAX_PACKET_SIZE : InMaxPacket;
diff --git a/Engine/Source/Runtime/Engine/Private/PhysicsEngine/ClothManager.cpp b/Engine/Source/Runtime/Engine/Private/PhysicsEngine/ClothManager.cpp
--- a/Engine/Source/Runtime/Engine/Private/PhysicsEngine/ClothManager.cpp
+++ b/Engine/Source/Runtime/Engine/Private/PhysicsEngine/ClothManager.cpp
@@ -19,8 +19,6 @@ void FClothManagerData::PrepareCloth(float DeltaTime, FTickFunction& TickFunctio
 	QUICK_SCOPE_CYCLE_COUNTER(STAT_FClothManager_PrepareCloth);
 	if (SkeletalMeshComponents.Num())
 	{
-		// TODO: OriC, I took the earlier version but you should verify which is correct. You probably want the above check at this location :)
-		//const bool bParallelCloth = !!CVarParallelCloth.GetValueOnAnyThread();
 		IsPreparingClothInParallel.AtomicSet(bParallelCloth);
 		for (USkeletalMeshComponent* SkeletalMeshComponent : SkeletalMeshComponents)
 		{
diff --git a/Engine/Source/Runtime/Engine/Private/PhysicsEngine/PhysScene.cpp b/Engine/Source/Runtime/Engine/Private/PhysicsEngine/PhysScene.cpp
--- a/Engine/Source/Runtime/Engine/Private/PhysicsEngine/PhysScene.cpp
+++ b/Engine/Source/Runtime/Engine/Private/PhysicsEngine/PhysScene.cpp
@@ -1005,9 +1005,6 @@ void FPhysScene::WaitClothScene()
 		QUICK_SCOPE_CYCLE_COUNTER(STAT_FPhysScene_WaitClothScene);
 		FTaskGraphInterface::Get().WaitUntilTasksComplete(ThingsToComplete, ENamedThreads::GameThread);
 	}
-	
-	PhysicsSubsceneCompletion[PST_Cloth] = nullptr;
-	bPhysXSceneExecuting[PST_Cloth] = false;
 }
 
 void FPhysScene::SceneCompletionTask(ENamedThreads::Type CurrentThread, const FGraphEventRef& MyCompletionGraphEvent, EPhysicsSceneType SceneType)
@@ -1071,12 +1068,8 @@ void FPhysScene::ProcessPhysScene(uint32 SceneType)
 	}
 #endif // WITH_PHYSX
 
-    if(SceneType != PST_Cloth)	//Cloth potentially runs this code off the game thread where setting to null is not thread safe. For cloth we do it in a later game thread sync point
-	{
 	PhysicsSubsceneCompletion[SceneType] = NULL;
 	bPhysXSceneExecuting[SceneType] = false;
-	}
-	
 }
 #if WITH_PHYSX
 void FPhysScene::UpdateActiveTransforms(uint32 SceneType)
@@ -1360,7 +1353,7 @@ void FPhysScene::StartCloth()
 					FDelegateGraphTask::CreateAndDispatchWhenReady(
 					FDelegateGraphTask::FDelegate::CreateRaw(this, &FPhysScene::SceneCompletionTask, PST_Cloth),
 					GET_STATID(STAT_FDelegateGraphTask_ProcessPhysScene_Cloth), PhysicsSubsceneCompletion[PST_Cloth],
-					CurrentThread, CurrentThread
+					CurrentThread, ENamedThreads::GameThread
 					)
 					);
 			}
diff --git a/Engine/Source/Runtime/Engine/Private/Tests/AutomationCommon.cpp b/Engine/Source/Runtime/Engine/Private/Tests/AutomationCommon.cpp
--- a/Engine/Source/Runtime/Engine/Private/Tests/AutomationCommon.cpp
+++ b/Engine/Source/Runtime/Engine/Private/Tests/AutomationCommon.cpp
@@ -12,6 +12,8 @@
 #include "Matinee/MatineeActor.h"
 
 DEFINE_LOG_CATEGORY_STATIC(LogEngineAutomationLatentCommand, Log, All);
+DEFINE_LOG_CATEGORY(LogEditorAutomationTests);
+DEFINE_LOG_CATEGORY(LogEngineAutomationTests);
 
 //declare static variable
 FOnEditorAutomationMapLoad AutomationCommon::OnEditorAutomationMapLoad;
diff --git a/Engine/Source/Runtime/Engine/Public/Tests/AutomationCommon.h b/Engine/Source/Runtime/Engine/Public/Tests/AutomationCommon.h
--- a/Engine/Source/Runtime/Engine/Public/Tests/AutomationCommon.h
+++ b/Engine/Source/Runtime/Engine/Public/Tests/AutomationCommon.h
@@ -9,9 +9,8 @@
 ///////////////////////////////////////////////////////////////////////
 // Common Latent commands which are used across test type. I.e. Engine, Network, etc...
 
-DEFINE_LOG_CATEGORY_STATIC(LogEditorAutomationTests, Log, All);
-DEFINE_LOG_CATEGORY_STATIC(LogEngineAutomationTests, Log, All);
-DEFINE_LOG_CATEGORY_STATIC(LogAnalytics, Log, All);
+ENGINE_API DECLARE_LOG_CATEGORY_EXTERN(LogEditorAutomationTests, Log, All);
+ENGINE_API DECLARE_LOG_CATEGORY_EXTERN(LogEngineAutomationTests, Log, All);
 
 DECLARE_MULTICAST_DELEGATE_OneParam(FOnEditorAutomationMapLoad, const FString&);
 
diff --git a/Engine/Source/Runtime/EngineSettings/Classes/GeneralProjectSettings.h b/Engine/Source/Runtime/EngineSettings/Classes/GeneralProjectSettings.h
--- a/Engine/Source/Runtime/EngineSettings/Classes/GeneralProjectSettings.h
+++ b/Engine/Source/Runtime/EngineSettings/Classes/GeneralProjectSettings.h
@@ -55,7 +55,7 @@ class ENGINESETTINGS_API UGeneralProjectSettings
 	UPROPERTY(config, EditAnywhere, Category=Publisher)
 	FString SupportContact;
 
-	/** The project's title as displayed on the window title bar */
+	/** The project's title as displayed on the window title bar (can include the tokens {GameName}, {PlatformArchitecture}, or {RHIName}, which will be replaced with the specified text) */
 	UPROPERTY(config, EditAnywhere, Category=Displayed)
 	FText ProjectDisplayedTitle;
 };
diff --git a/Engine/Source/Runtime/GameplayAbilities/Private/Abilities/GameplayAbility.cpp b/Engine/Source/Runtime/GameplayAbilities/Private/Abilities/GameplayAbility.cpp
--- a/Engine/Source/Runtime/GameplayAbilities/Private/Abilities/GameplayAbility.cpp
+++ b/Engine/Source/Runtime/GameplayAbilities/Private/Abilities/GameplayAbility.cpp
@@ -58,6 +58,7 @@ UGameplayAbility::UGameplayAbility(const FObjectInitializer& ObjectInitializer)
 
 	bServerRespectsRemoteAbilityCancellation = true;
 	bReplicateInputDirectly = false;
+	RemoteInstanceEnded = false;
 }
 
 int32 UGameplayAbility::GetFunctionCallspace(UFunction* Function, void* Parameters, FFrame* Stack)
@@ -341,17 +342,20 @@ bool UGameplayAbility::CommitAbility(const FGameplayAbilitySpecHandle Handle, co
 	return true;
 }
 
-bool UGameplayAbility::CommitAbilityCooldown(const FGameplayAbilitySpecHandle Handle, const FGameplayAbilityActorInfo* ActorInfo, const FGameplayAbilityActivationInfo ActivationInfo)
+bool UGameplayAbility::CommitAbilityCooldown(const FGameplayAbilitySpecHandle Handle, const FGameplayAbilityActorInfo* ActorInfo, const FGameplayAbilityActivationInfo ActivationInfo, const bool ForceCooldown)
 {
 	if (UAbilitySystemGlobals::Get().ShouldIgnoreCooldowns())
 	{
 		return true;
 	}
 
-	// Last chance to fail (maybe we no longer have resources to commit since we after we started this ability activation)
-	if (!CheckCooldown(Handle, ActorInfo))
+	if (!ForceCooldown)
 	{
-		return false;
+		// Last chance to fail (maybe we no longer have resources to commit since we after we started this ability activation)
+		if (!CheckCooldown(Handle, ActorInfo))
+		{
+			return false;
+		}
 	}
 
 	ApplyCooldown(Handle, ActorInfo, ActivationInfo);
@@ -920,14 +924,14 @@ bool UGameplayAbility::K2_CommitAbility()
 	return CommitAbility(CurrentSpecHandle, CurrentActorInfo, CurrentActivationInfo);
 }
 
-bool UGameplayAbility::K2_CommitAbilityCooldown(bool BroadcastCommitEvent)
+bool UGameplayAbility::K2_CommitAbilityCooldown(bool BroadcastCommitEvent, bool ForceCooldown)
 {
 	check(CurrentActorInfo);
 	if (BroadcastCommitEvent)
 	{
 		CurrentActorInfo->AbilitySystemComponent->NotifyAbilityCommit(this);
 	}
-	return CommitAbilityCooldown(CurrentSpecHandle, CurrentActorInfo, CurrentActivationInfo);
+	return CommitAbilityCooldown(CurrentSpecHandle, CurrentActorInfo, CurrentActivationInfo, ForceCooldown);
 }
 
 bool UGameplayAbility::K2_CommitAbilityCost(bool BroadcastCommitEvent)
@@ -1574,4 +1578,39 @@ void UGameplayAbility::ConvertDeprecatedGameplayEffectReferencesToBlueprintRefer
 float UGameplayAbility::GetCooldownTimeRemaining() const
 {
 	return IsInstantiated() ? GetCooldownTimeRemaining(CurrentActorInfo) : 0.f;
+}
+
+void UGameplayAbility::SetRemoteInstanceHasEnded()
+{
+	// This could potentially happen in shutdown corner cases
+	if (IsPendingKill() || CurrentActorInfo == nullptr || CurrentActorInfo->AbilitySystemComponent.IsValid() == false)
+	{
+		return;
+	}
+
+	RemoteInstanceEnded = true;
+	for (UGameplayTask* Task : ActiveTasks)
+	{
+		if (Task && Task->IsPendingKill() == false && Task->IsWaitingOnRemotePlayerdata())
+		{
+			// We have a task that is waiting for player input, but the remote player has ended the ability, so he will not send it.
+			// Kill the ability to avoid getting stuck active.
+			
+			ABILITY_LOG(Warning, TEXT("Ability %s is force cancelling because Task %s is waiting on remote player input and the  remote player has just ended the ability."), *GetName(), *Task->GetDebugString());
+			CurrentActorInfo->AbilitySystemComponent->ForceCancelAbilityDueToReplication(this);
+			break;
+		}
+	}
+}
+
+void UGameplayAbility::NotifyAbilityTaskWaitingOnPlayerData(class UAbilityTask* AbilityTask)
+{
+	// This should never happen since it will only be called from actively running ability tasks
+	check(CurrentActorInfo && CurrentActorInfo->AbilitySystemComponent.IsValid());
+
+	if (RemoteInstanceEnded)
+	{
+		ABILITY_LOG(Warning, TEXT("Ability %s is force cancelling because Task %s has started after the remote player has ended the ability."), *GetName(), *AbilityTask->GetDebugString());
+		CurrentActorInfo->AbilitySystemComponent->ForceCancelAbilityDueToReplication(this);
+	}
 }
\ No newline at end of file
diff --git a/Engine/Source/Runtime/GameplayAbilities/Private/Abilities/Tasks/AbilityTask.cpp b/Engine/Source/Runtime/GameplayAbilities/Private/Abilities/Tasks/AbilityTask.cpp
--- a/Engine/Source/Runtime/GameplayAbilities/Private/Abilities/Tasks/AbilityTask.cpp
+++ b/Engine/Source/Runtime/GameplayAbilities/Private/Abilities/Tasks/AbilityTask.cpp
@@ -7,6 +7,7 @@
 UAbilityTask::UAbilityTask(const FObjectInitializer& ObjectInitializer)
 	: Super(ObjectInitializer)
 {
+	WaitState = EAbilityTaskWaitState::WaitingOnGame;
 }
 
 FGameplayAbilitySpecHandle UAbilityTask::GetAbilitySpecHandle() const
@@ -53,5 +54,25 @@ bool UAbilityTask::IsLocallyControlled() const
 
 bool UAbilityTask::CallOrAddReplicatedDelegate(EAbilityGenericReplicatedEvent::Type Event, FSimpleMulticastDelegate::FDelegate Delegate)
 {
-	return AbilitySystemComponent->CallOrAddReplicatedDelegate(Event, GetAbilitySpecHandle(), GetActivationPredictionKey(), Delegate);
+	if (!AbilitySystemComponent->CallOrAddReplicatedDelegate(Event, GetAbilitySpecHandle(), GetActivationPredictionKey(), Delegate))
+	{
+		SetWaitingOnRemotePlayerData();
+		return false;
+	}
+	return true;
 }
+
+void UAbilityTask::SetWaitingOnRemotePlayerData()
+{
+	UGameplayAbility* MyAbility = Ability.Get();
+	if (MyAbility && IsPendingKill() == false && AbilitySystemComponent.IsValid())
+	{
+		WaitState = EAbilityTaskWaitState::WaitingOnUser;
+		MyAbility->NotifyAbilityTaskWaitingOnPlayerData(this);
+	}
+}
+
+bool UAbilityTask::IsWaitingOnRemotePlayerdata() const
+{
+	return (WaitState == EAbilityTaskWaitState::WaitingOnUser);
+}
\ No newline at end of file
diff --git a/Engine/Source/Runtime/GameplayAbilities/Private/Abilities/Tasks/AbilityTask_WaitInputPress.cpp b/Engine/Source/Runtime/GameplayAbilities/Private/Abilities/Tasks/AbilityTask_WaitInputPress.cpp
--- a/Engine/Source/Runtime/GameplayAbilities/Private/Abilities/Tasks/AbilityTask_WaitInputPress.cpp
+++ b/Engine/Source/Runtime/GameplayAbilities/Private/Abilities/Tasks/AbilityTask_WaitInputPress.cpp
@@ -62,7 +62,10 @@ void UAbilityTask_WaitInputPress::Activate()
 		DelegateHandle = AbilitySystemComponent->AbilityReplicatedEventDelegate(EAbilityGenericReplicatedEvent::InputPressed, GetAbilitySpecHandle(), GetActivationPredictionKey()).AddUObject(this, &UAbilityTask_WaitInputPress::OnPressCallback);
 		if (IsForRemoteClient())
 		{
-			AbilitySystemComponent->CallReplicatedEventDelegateIfSet(EAbilityGenericReplicatedEvent::InputPressed, GetAbilitySpecHandle(), GetActivationPredictionKey());
+			if (!AbilitySystemComponent->CallReplicatedEventDelegateIfSet(EAbilityGenericReplicatedEvent::InputPressed, GetAbilitySpecHandle(), GetActivationPredictionKey()))
+			{
+				SetWaitingOnRemotePlayerData();
+			}
 		}
 	}
 }
diff --git a/Engine/Source/Runtime/GameplayAbilities/Private/Abilities/Tasks/AbilityTask_WaitInputRelease.cpp b/Engine/Source/Runtime/GameplayAbilities/Private/Abilities/Tasks/AbilityTask_WaitInputRelease.cpp
--- a/Engine/Source/Runtime/GameplayAbilities/Private/Abilities/Tasks/AbilityTask_WaitInputRelease.cpp
+++ b/Engine/Source/Runtime/GameplayAbilities/Private/Abilities/Tasks/AbilityTask_WaitInputRelease.cpp
@@ -62,7 +62,10 @@ void UAbilityTask_WaitInputRelease::Activate()
 		DelegateHandle = AbilitySystemComponent->AbilityReplicatedEventDelegate(EAbilityGenericReplicatedEvent::InputReleased, GetAbilitySpecHandle(), GetActivationPredictionKey()).AddUObject(this, &UAbilityTask_WaitInputRelease::OnReleaseCallback);
 		if (IsForRemoteClient())
 		{
-			AbilitySystemComponent->CallReplicatedEventDelegateIfSet(EAbilityGenericReplicatedEvent::InputReleased, GetAbilitySpecHandle(), GetActivationPredictionKey());
+			if (!AbilitySystemComponent->CallReplicatedEventDelegateIfSet(EAbilityGenericReplicatedEvent::InputReleased, GetAbilitySpecHandle(), GetActivationPredictionKey()))
+			{
+				SetWaitingOnRemotePlayerData();
+			}
 		}
 	}
 }
diff --git a/Engine/Source/Runtime/GameplayAbilities/Private/Abilities/Tasks/AbilityTask_WaitTargetData.cpp b/Engine/Source/Runtime/GameplayAbilities/Private/Abilities/Tasks/AbilityTask_WaitTargetData.cpp
--- a/Engine/Source/Runtime/GameplayAbilities/Private/Abilities/Tasks/AbilityTask_WaitTargetData.cpp
+++ b/Engine/Source/Runtime/GameplayAbilities/Private/Abilities/Tasks/AbilityTask_WaitTargetData.cpp
@@ -198,6 +198,8 @@ void UAbilityTask_WaitTargetData::RegisterTargetDataCallbacks()
 			AbilitySystemComponent->AbilityTargetDataCancelledDelegate(SpecHandle, ActivationPredictionKey ).AddUObject(this, &UAbilityTask_WaitTargetData::OnTargetDataReplicatedCancelledCallback);
 
 			AbilitySystemComponent->CallReplicatedTargetDataDelegatesIfSet(SpecHandle, ActivationPredictionKey );
+
+			SetWaitingOnRemotePlayerData();
 		}
 	}
 }
diff --git a/Engine/Source/Runtime/GameplayAbilities/Private/AbilitySystemComponent.cpp b/Engine/Source/Runtime/GameplayAbilities/Private/AbilitySystemComponent.cpp
--- a/Engine/Source/Runtime/GameplayAbilities/Private/AbilitySystemComponent.cpp
+++ b/Engine/Source/Runtime/GameplayAbilities/Private/AbilitySystemComponent.cpp
@@ -1389,66 +1389,152 @@ void UAbilitySystemComponent::DebugCyclicAggregatorBroadcasts(FAggregator* Aggre
 
 // ------------------------------------------------------------------------
 
-FString ASC_CleanupName(FString Str)
+bool UAbilitySystemComponent::ServerPrintDebug_Request_Validate()
+{
+	return true;
+}
+
+void UAbilitySystemComponent::ServerPrintDebug_Request_Implementation()
+{
+	FAbilitySystemComponentDebugInfo DebugInfo;
+	DebugInfo.bShowAbilities = true;
+	DebugInfo.bShowAttributes = true;
+	DebugInfo.bShowGameplayEffects = true;
+	DebugInfo.Accumulate = true;
+
+	Debug_Internal(DebugInfo);
+
+	ClientPrintDebug_Response(DebugInfo.Strings);
+}
+
+void UAbilitySystemComponent::ClientPrintDebug_Response_Implementation(const TArray<FString>& Strings)
+{
+	ABILITY_LOG(Warning, TEXT(" "));
+	ABILITY_LOG(Warning, TEXT("Server State: "));
+
+	for (FString Str : Strings)
+	{
+		ABILITY_LOG(Warning, TEXT("%s"), *Str);
+	}
+}
+
+FString UAbilitySystemComponent::CleanupName(FString Str)
 {
 	Str.RemoveFromStart(TEXT("Default__"));
 	Str.RemoveFromEnd(TEXT("_c"));
 	return Str;
 }
 
-void AccumulateScreenPos(float& XPos, float& YPos, float AdditionalY, float OriginalY, float MaxY, const float NewColumnYPadding, UCanvas* Canvas)
+void UAbilitySystemComponent::AccumulateScreenPos(FAbilitySystemComponentDebugInfo& Info)
 {
-	const float ColumnWidth = Canvas->ClipX * 0.4f;
+	const float ColumnWidth = Info.Canvas ? Info.Canvas->ClipX * 0.4f : 0.f;
 
-	float NewY = YPos + AdditionalY;
-	if (NewY > MaxY)
+	float NewY = Info.YPos + Info.YL;
+	if (NewY > Info.MaxY)
 	{
 		// Need new column, reset Y to original height
-		NewY = NewColumnYPadding;
-		XPos += ColumnWidth;
+		NewY = Info.NewColumnYPadding;
+		Info.XPos += ColumnWidth;
+	}
+	Info.YPos = NewY;
+}
+
+void UAbilitySystemComponent::DebugLine(FAbilitySystemComponentDebugInfo& Info, FString Str, float XOffset, float YOffset)
+{
+	if (Info.Canvas)
+	{
+		Info.YL = Info.Canvas->DrawText(GEngine->GetTinyFont(), Str, Info.XPos + XOffset, Info.YPos );
+		AccumulateScreenPos(Info);
+	}
+
+	if (Info.bPrintToLog)
+	{
+		FString LogStr;
+		for (int32 i=0; i < (int32)XOffset; ++i)
+		{
+			LogStr += TEXT(" ");
+		}
+		LogStr += Str;
+		ABILITY_LOG(Warning, TEXT("%s"), *LogStr);
+	}
+
+	if (Info.Accumulate)
+	{
+		FString LogStr;
+		for (int32 i=0; i < (int32)XOffset; ++i)
+		{
+			LogStr += TEXT(" ");
+		}
+		LogStr += Str;
+		Info.Strings.Add(Str);
 	}
-	YPos = NewY;
 }
 
 void UAbilitySystemComponent::DisplayDebug(class UCanvas* Canvas, const class FDebugDisplayInfo& DebugDisplay, float& YL, float& YPos)
 {
-	bool bShowAttributes = true;
-	bool bShowGameplayEffects = true;
-	bool bShowAbilities = true;
+	FAbilitySystemComponentDebugInfo DebugInfo;
 
 	if (DebugDisplay.IsDisplayOn(FName(TEXT("Attributes"))))
 	{
-		bShowAbilities = false;
-		bShowAttributes = true;
-		bShowGameplayEffects = false;
+		DebugInfo.bShowAbilities = false;
+		DebugInfo.bShowAttributes = true;
+		DebugInfo.bShowGameplayEffects = false;
 	}
 	if (DebugDisplay.IsDisplayOn(FName(TEXT("Ability"))))
 	{
-		bShowAbilities = true;
-		bShowAttributes = false;
-		bShowGameplayEffects = false;
+		DebugInfo.bShowAbilities = true;
+		DebugInfo.bShowAttributes = false;
+		DebugInfo.bShowGameplayEffects = false;
 	}
 	else if (DebugDisplay.IsDisplayOn(FName(TEXT("GameplayEffects"))))
 	{
-		bShowAbilities = false;
-		bShowAttributes = false;
-		bShowGameplayEffects = true;
+		DebugInfo.bShowAbilities = false;
+		DebugInfo.bShowAttributes = false;
+		DebugInfo.bShowGameplayEffects = true;
 	}
 
-	float XPos = 0.f;
-	const float OriginalX = XPos;
-	const float OriginalY = YPos;
-	const float MaxY = Canvas->ClipY - 150.f; // Give some padding for any non-columnizing debug output following this output
-	float NewColumnYPadding = 30.f;
+	DebugInfo.bPrintToLog = false;
+	DebugInfo.Canvas = Canvas;
+	DebugInfo.XPos = 0.f;
+	DebugInfo.YPos = YPos;
+	DebugInfo.OriginalX = 0.f;
+	DebugInfo.OriginalY = YPos;
+	DebugInfo.MaxY = Canvas->ClipY - 150.f; // Give some padding for any non-columnizing debug output following this output
+	DebugInfo.NewColumnYPadding = 30.f;
+
+	Debug_Internal(DebugInfo);
+
+	YPos = DebugInfo.YPos;
+	YL = DebugInfo.YL;
+}
+
+void UAbilitySystemComponent::PrintDebug()
+{
+	FAbilitySystemComponentDebugInfo DebugInfo;
+	DebugInfo.bShowAbilities = true;
+	DebugInfo.bShowAttributes = true;
+	DebugInfo.bShowGameplayEffects = true;
+	DebugInfo.bPrintToLog = true;
 
+	Debug_Internal(DebugInfo);
+
+	if (IsOwnerActorAuthoritative() == false)
+	{
+		// See what the server thinks
+		ServerPrintDebug_Request();
+	}
+}
+
+void UAbilitySystemComponent::Debug_Internal(FAbilitySystemComponentDebugInfo& Info)
+{
 	// Draw title at top of screen (default HUD debug text starts at 50 ypos, we can position this on top)*
 	//   *until someone changes it unknowingly
 	{
 		FString DebugTitle("");
 		// Category
-		if (bShowAbilities) DebugTitle += TEXT("ABILITIES ");
-		if (bShowAttributes) DebugTitle += TEXT("ATTRIBUTES ");
-		if (bShowGameplayEffects) DebugTitle += TEXT("GAMEPLAYEFFECTS ");
+		if (Info.bShowAbilities) DebugTitle += TEXT("ABILITIES ");
+		if (Info.bShowAttributes) DebugTitle += TEXT("ATTRIBUTES ");
+		if (Info.bShowGameplayEffects) DebugTitle += TEXT("GAMEPLAYEFFECTS ");
 		// Avatar info
 		if (AvatarActor)
 		{
@@ -1466,31 +1552,41 @@ void UAbilitySystemComponent::DisplayDebug(class UCanvas* Canvas, const class FD
 			else if (OwnerActor->Role == ROLE_Authority) DebugTitle += TEXT("(authority) ");
 		}
 
-		Canvas->SetDrawColor(FColor::White);
-		Canvas->DrawText(GEngine->GetLargeFont(), DebugTitle, XPos + 4.f, 10.f, 1.5f, 1.5f);
+		if (Info.Canvas)
+		{
+			Info.Canvas->SetDrawColor(FColor::White);
+			Info.Canvas->DrawText(GEngine->GetLargeFont(), DebugTitle, Info.XPos + 4.f, 10.f, 1.5f, 1.5f);
+		}
+		else
+		{
+			DebugLine(Info, DebugTitle, 0.f, 0.f);
+		}
 	}
 
 	FGameplayTagContainer OwnerTags;
 	GetOwnedGameplayTags(OwnerTags);
 
-	Canvas->SetDrawColor(FColor::White);
-	YL = Canvas->DrawText(GEngine->GetTinyFont(), FString::Printf(TEXT("Owned Tags: %s"), *OwnerTags.ToStringSimple()), XPos + 4.f, YPos);
-	AccumulateScreenPos(XPos, YPos, YL, OriginalY, MaxY, NewColumnYPadding, Canvas);
+	if (Info.Canvas) Info.Canvas->SetDrawColor(FColor::White);
+
+	DebugLine(Info, FString::Printf(TEXT("Owned Tags: %s"), *OwnerTags.ToStringSimple()), 4.f, 0.f);
 
 	if (BlockedAbilityTags.GetExplicitGameplayTags().Num() > 0)
 	{
-		YL = Canvas->DrawText(GEngine->GetTinyFont(), FString::Printf(TEXT("BlockedAbilityTags: %s"), *BlockedAbilityTags.GetExplicitGameplayTags().ToStringSimple()), XPos + 4.f, YPos);
-		AccumulateScreenPos(XPos, YPos, YL, OriginalY, MaxY, NewColumnYPadding, Canvas);
+		DebugLine(Info, FString::Printf(TEXT("BlockedAbilityTags: %s"), *BlockedAbilityTags.GetExplicitGameplayTags().ToStringSimple()), 4.f, 0.f);
+		
 	}
 
 	TSet<FGameplayAttribute> DrawAttributes;
 
-	const float MaxCharHeight = GEngine->GetTinyFont()->GetMaxCharHeight();
+	float MaxCharHeight = 10;
+	if (GetOwner()->GetNetMode() != NM_DedicatedServer)
+	{
+		MaxCharHeight = GEngine->GetTinyFont()->GetMaxCharHeight();
+	}
 
 	// -------------------------------------------------------------
 
-
-	if (bShowAttributes)
+	if (Info.bShowAttributes)
 	{
 		// Draw the attribute aggregator map.
 		for (auto It = ActiveGameplayEffects.AttributeAggregatorMap.CreateConstIterator(); It; ++It)
@@ -1521,9 +1617,10 @@ void UAbilitySystemComponent::DisplayDebug(class UCanvas* Canvas, const class FD
 					AttributeString += FString::Printf(TEXT(" (Base: %.2f)"), BaseValue);
 				}
 
-				Canvas->SetDrawColor(FColor::White);
-				YL = Canvas->DrawText(GEngine->GetTinyFont(), AttributeString, XPos + 4.f, YPos);
-				AccumulateScreenPos(XPos, YPos, YL, OriginalY, MaxY, NewColumnYPadding, Canvas);
+				if (Info.Canvas) Info.Canvas->SetDrawColor(FColor::White);
+				
+
+				DebugLine(Info, AttributeString, 4.f, 0.f);
 
 				DrawAttributes.Add(Attribute);
 
@@ -1533,7 +1630,7 @@ void UAbilitySystemComponent::DisplayDebug(class UCanvas* Canvas, const class FD
 					{
 						FAggregatorEvaluateParameters EmptyParams;
 						bool IsActivelyModifyingAttribute = Mod.Qualifies(EmptyParams);
-						Canvas->SetDrawColor(IsActivelyModifyingAttribute ? FColor::Yellow : FColor(128, 128, 128));
+						if (Info.Canvas) Info.Canvas->SetDrawColor(IsActivelyModifyingAttribute ? FColor::Yellow : FColor(128, 128, 128));
 
 						FActiveGameplayEffect* ActiveGE = ActiveGameplayEffects.GetActiveGameplayEffect(Mod.ActiveHandle);
 						FString SrcName = ActiveGE ? ActiveGE->Spec.Def->GetName() : FString(TEXT(""));
@@ -1544,24 +1641,23 @@ void UAbilitySystemComponent::DisplayDebug(class UCanvas* Canvas, const class FD
 							if (Mod.TargetTagReqs) SrcName += FString::Printf(TEXT("TargetTags: [%s]"), *Mod.TargetTagReqs->ToString());
 						}
 
-						YL = Canvas->DrawText(GEngine->GetTinyFont(), FString::Printf(TEXT("   %s\t %.2f - %s"), *EGameplayModOpToString(ModOpIdx), Mod.EvaluatedMagnitude, *SrcName), XPos + 7.f, YPos);
-						AccumulateScreenPos(XPos, YPos, YL, OriginalY, MaxY, NewColumnYPadding, Canvas);
-						NewColumnYPadding = FMath::Max<float>(NewColumnYPadding, YPos + YL);
+						DebugLine(Info, FString::Printf(TEXT("   %s\t %.2f - %s"), *EGameplayModOpToString(ModOpIdx), Mod.EvaluatedMagnitude, *SrcName), 7.f, 0.f);
+						Info.NewColumnYPadding = FMath::Max<float>(Info.NewColumnYPadding, Info.YPos + Info.YL);
 					}
 				}
-				AccumulateScreenPos(XPos, YPos, MaxCharHeight, OriginalY, MaxY, NewColumnYPadding, Canvas);
+				AccumulateScreenPos(Info);
 			}
 		}
 	}
 
 	// -------------------------------------------------------------
 
-	if (bShowGameplayEffects)
+	if (Info.bShowGameplayEffects)
 	{
 		for (FActiveGameplayEffect& ActiveGE : &ActiveGameplayEffects)
 		{
 
-			Canvas->SetDrawColor(FColor::White);
+			if (Info.Canvas) Info.Canvas->SetDrawColor(FColor::White);
 
 			FString DurationStr = TEXT("Infinite Duration ");
 			if (ActiveGE.GetDuration() > 0.f)
@@ -1606,17 +1702,15 @@ void UAbilitySystemComponent::DisplayDebug(class UCanvas* Canvas, const class FD
 				}
 			}
 
-			Canvas->SetDrawColor(ActiveGE.bIsInhibited ? FColor(128, 128, 128) : FColor::White);
+			if (Info.Canvas) Info.Canvas->SetDrawColor(ActiveGE.bIsInhibited ? FColor(128, 128, 128) : FColor::White);
 
-			YL = Canvas->DrawText(GEngine->GetTinyFont(), FString::Printf(TEXT("%s %s %s %s %s"), *ASC_CleanupName(GetNameSafe(ActiveGE.Spec.Def)), *DurationStr, *StackString, *LevelString, *PredictionString), XPos + 4.f, YPos);
-			AccumulateScreenPos(XPos, YPos, YL, OriginalY, MaxY, NewColumnYPadding, Canvas);
+			DebugLine(Info, FString::Printf(TEXT("%s %s %s %s %s"), *CleanupName(GetNameSafe(ActiveGE.Spec.Def)), *DurationStr, *StackString, *LevelString, *PredictionString), 4.f, 0.f);
 
 			FGameplayTagContainer GrantedTags;
 			ActiveGE.Spec.GetAllGrantedTags(GrantedTags);
 			if (GrantedTags.Num() > 0)
 			{
-				YL = Canvas->DrawText(GEngine->GetTinyFont(), FString::Printf(TEXT("Granted Tags: %s"), *GrantedTags.ToStringSimple()), XPos + 7.f, YPos);
-				AccumulateScreenPos(XPos, YPos, YL, OriginalY, MaxY, NewColumnYPadding, Canvas);
+				DebugLine(Info, FString::Printf(TEXT("Granted Tags: %s"), *GrantedTags.ToStringSimple()), 7.f, 0.f);
 			}
 
 			for (int32 ModIdx = 0; ModIdx < ActiveGE.Spec.Modifiers.Num(); ++ModIdx)
@@ -1635,24 +1729,23 @@ void UAbilitySystemComponent::DisplayDebug(class UCanvas* Canvas, const class FD
 
 				if (IsActivelyModifyingAttribute == false)
 				{
-					Canvas->SetDrawColor(FColor(128, 128, 128));
+					if (Info.Canvas) Info.Canvas->SetDrawColor(FColor(128, 128, 128));
 				}
 
-				YL = Canvas->DrawText(GEngine->GetTinyFont(), FString::Printf(TEXT("Mod: %s. %s. %.2f"), *ModInfo.Attribute.GetName(), *EGameplayModOpToString(ModInfo.ModifierOp), ModSpec.GetEvaluatedMagnitude()), XPos + 7.f, YPos);
-				AccumulateScreenPos(XPos, YPos, YL, OriginalY, MaxY, NewColumnYPadding, Canvas);
+				DebugLine(Info, FString::Printf(TEXT("Mod: %s. %s. %.2f"), *ModInfo.Attribute.GetName(), *EGameplayModOpToString(ModInfo.ModifierOp), ModSpec.GetEvaluatedMagnitude()), 7.f, 0.f);
 
-				Canvas->SetDrawColor(ActiveGE.bIsInhibited ? FColor(128, 128, 128) : FColor::White);
+				if (Info.Canvas) Info.Canvas->SetDrawColor(ActiveGE.bIsInhibited ? FColor(128, 128, 128) : FColor::White);
 			}
 
-			AccumulateScreenPos(XPos, YPos, MaxCharHeight, OriginalY, MaxY, NewColumnYPadding, Canvas);
+			AccumulateScreenPos(Info);
 		}
 	}
 
 	// -------------------------------------------------------------
 
-	if (bShowAttributes)
+	if (Info.bShowAttributes)
 	{
-		Canvas->SetDrawColor(FColor::White);
+		if (Info.Canvas) Info.Canvas->SetDrawColor(FColor::White);
 		for (UAttributeSet* Set : SpawnedAttributes)
 		{
 			for (TFieldIterator<UProperty> It(Set->GetClass()); It; ++It)
@@ -1665,17 +1758,17 @@ void UAbilitySystemComponent::DisplayDebug(class UCanvas* Canvas, const class FD
 				if (Attribute.IsValid())
 				{
 					float Value = GetNumericAttribute(Attribute);
-					YL = Canvas->DrawText(GEngine->GetTinyFont(), FString::Printf(TEXT("%s %.2f"), *Attribute.GetName(), Value), XPos + 4.f, YPos);
-					AccumulateScreenPos(XPos, YPos, YL, OriginalY, MaxY, NewColumnYPadding, Canvas);
+
+					DebugLine(Info, FString::Printf(TEXT("%s %.2f"), *Attribute.GetName(), Value), 4.f, 0.f);
 				}
 			}
 		}
-		AccumulateScreenPos(XPos, YPos, MaxCharHeight, OriginalY, MaxY, NewColumnYPadding, Canvas);
+		AccumulateScreenPos(Info);
 	}
 
 	// -------------------------------------------------------------
 
-	if (bShowAbilities)
+	if (Info.bShowAbilities)
 	{
 		for (const FGameplayAbilitySpec& AbilitySpec : GetActivatableAbilities())
 		{
@@ -1707,9 +1800,9 @@ void UAbilitySystemComponent::DisplayDebug(class UCanvas* Canvas, const class FD
 
 			FString InputPressedStr = AbilitySpec.InputPressed ? TEXT("(InputPressed)") : TEXT("");
 
-			Canvas->SetDrawColor(AbilityTextColor);
-			YL = Canvas->DrawText(GEngine->GetTinyFont(), FString::Printf(TEXT("%s %s %s"), *ASC_CleanupName(GetNameSafe(AbilitySpec.Ability)), *StatusText, *InputPressedStr), XPos + 4.f, YPos);
-			AccumulateScreenPos(XPos, YPos, YL, OriginalY, MaxY, NewColumnYPadding, Canvas);
+			if (Info.Canvas) Info.Canvas->SetDrawColor(AbilityTextColor);
+
+			DebugLine(Info, FString::Printf(TEXT("%s %s %s"), *CleanupName(GetNameSafe(AbilitySpec.Ability)), *StatusText, *InputPressedStr), 4.f, 0.f);
 
 			if (AbilitySpec.IsActive())
 			{
@@ -1720,35 +1813,33 @@ void UAbilitySystemComponent::DisplayDebug(class UCanvas* Canvas, const class FD
 					if (!Instance)
 						continue;
 
-					Canvas->SetDrawColor(FColor::White);
+					if (Info.Canvas) Info.Canvas->SetDrawColor(FColor::White);
 					for (auto TaskPtr : Instance->ActiveTasks)
 					{
 						UGameplayTask* Task = TaskPtr;
 						if (Task)
 						{
-							YL = Canvas->DrawText(GEngine->GetTinyFont(), FString::Printf(TEXT("%s"), *Task->GetDebugString()), XPos + 7.f, YPos);
-							AccumulateScreenPos(XPos, YPos, YL, OriginalY, MaxY, NewColumnYPadding, Canvas);
+							DebugLine(Info, FString::Printf(TEXT("%s"), *Task->GetDebugString()), 7.f, 0.f);
 						}
 					}
 
 					if (InstanceIdx < Instances.Num() - 2)
 					{
-						Canvas->SetDrawColor(FColor(128, 128, 128));
-						YL = Canvas->DrawText(GEngine->GetTinyFont(), FString::Printf(TEXT("--------")), XPos + 7.f, YPos);
-						AccumulateScreenPos(XPos, YPos, YL, OriginalY, MaxY, NewColumnYPadding, Canvas);
+						if (Info.Canvas) Info.Canvas->SetDrawColor(FColor(128, 128, 128));
+						DebugLine(Info, FString::Printf(TEXT("--------")), 7.f, 0.f);
 					}
 				}
 			}
 		}
-		AccumulateScreenPos(XPos, YPos, MaxCharHeight, OriginalY, MaxY, NewColumnYPadding, Canvas);
+		AccumulateScreenPos(Info);
 	}
 
-	if (XPos > OriginalX)
+	if (Info.XPos > Info.OriginalX)
 	{
 		// We flooded to new columns, returned YPos should be max Y (and some padding)
-		YPos = MaxY + MaxCharHeight*2.f;
+		Info.YPos = Info.MaxY + MaxCharHeight*2.f;
 	}
-	YL = MaxCharHeight;
+	Info.YL = MaxCharHeight;
 }
 
 
diff --git a/Engine/Source/Runtime/GameplayAbilities/Private/AbilitySystemComponent_Abilities.cpp b/Engine/Source/Runtime/GameplayAbilities/Private/AbilitySystemComponent_Abilities.cpp
--- a/Engine/Source/Runtime/GameplayAbilities/Private/AbilitySystemComponent_Abilities.cpp
+++ b/Engine/Source/Runtime/GameplayAbilities/Private/AbilitySystemComponent_Abilities.cpp
@@ -1390,7 +1390,7 @@ void UAbilitySystemComponent::ReplicateEndOrCancelAbility(FGameplayAbilitySpecHa
 	}
 }
 
-//This is only called when ending or canceling an ability in response to a remote instruction.
+// This is only called when ending or canceling an ability in response to a remote instruction.
 void UAbilitySystemComponent::RemoteEndOrCancelAbility(FGameplayAbilitySpecHandle AbilityToEnd, FGameplayAbilityActivationInfo ActivationInfo, bool bWasCanceled)
 {
 	FGameplayAbilitySpec* AbilitySpec = FindAbilitySpecFromHandle(AbilityToEnd);
@@ -1416,18 +1416,22 @@ void UAbilitySystemComponent::RemoteEndOrCancelAbility(FGameplayAbilitySpecHandl
 			for (auto Instance : Instances)
 			{
 				// Check if the ability is the same prediction key (can both by 0) and has been confirmed. If so cancel it.
-				if (Instance->GetCurrentActivationInfoRef().GetActivationPredictionKey() == ActivationInfo.GetActivationPredictionKey() && Instance->GetCurrentActivationInfoRef().bCanBeEndedByOtherInstance)
+				if (Instance->GetCurrentActivationInfoRef().GetActivationPredictionKey() == ActivationInfo.GetActivationPredictionKey())
 				{
-					// End/Cancel the ability but don't replicate it back to whoever called us
-					if (bWasCanceled)
-					{
-						// Since this was a remote cancel, we should force it through. We do not support 'server says ability was cancelled but client disagrees that it can be'.
-						Instance->SetCanBeCanceled(true);
-						Instance->CancelAbility(Instance->CurrentSpecHandle, Instance->CurrentActorInfo, Instance->CurrentActivationInfo, false);
-					}
-					else
+					// Let the ability know that the remote instance has ended, even if we aren't about to end it here.
+					Instance->SetRemoteInstanceHasEnded();
+
+					if (Instance->GetCurrentActivationInfoRef().bCanBeEndedByOtherInstance)
 					{
-						Instance->EndAbility(Instance->CurrentSpecHandle, Instance->CurrentActorInfo, Instance->CurrentActivationInfo, false);
+						// End/Cancel the ability but don't replicate it back to whoever called us
+						if (bWasCanceled)
+						{
+							ForceCancelAbilityDueToReplication(Instance);
+						}
+						else
+						{
+							Instance->EndAbility(Instance->CurrentSpecHandle, Instance->CurrentActorInfo, Instance->CurrentActivationInfo, false);
+						}
 					}
 				}
 			}
@@ -1435,6 +1439,16 @@ void UAbilitySystemComponent::RemoteEndOrCancelAbility(FGameplayAbilitySpecHandl
 	}
 }
 
+/** Force cancels the ability and does not replicate this to the other side. This should be called when the ability is cancelled by the other side. */
+void UAbilitySystemComponent::ForceCancelAbilityDueToReplication(UGameplayAbility* Instance)
+{
+	check(Instance);
+
+	// Since this was a remote cancel, we should force it through. We do not support 'server says ability was cancelled but client disagrees that it can be'.
+	Instance->SetCanBeCanceled(true);
+	Instance->CancelAbility(Instance->CurrentSpecHandle, Instance->CurrentActorInfo, Instance->CurrentActivationInfo, false);
+}
+
 void UAbilitySystemComponent::ServerEndAbility_Implementation(FGameplayAbilitySpecHandle AbilityToEnd, FGameplayAbilityActivationInfo ActivationInfo, FPredictionKey PredictionKey)
 {
 	FScopedPredictionWindow ScopedPrediction(this, PredictionKey);
diff --git a/Engine/Source/Runtime/GameplayAbilities/Public/Abilities/GameplayAbility.h b/Engine/Source/Runtime/GameplayAbilities/Public/Abilities/GameplayAbility.h
--- a/Engine/Source/Runtime/GameplayAbilities/Public/Abilities/GameplayAbility.h
+++ b/Engine/Source/Runtime/GameplayAbilities/Public/Abilities/GameplayAbility.h
@@ -364,7 +364,7 @@ public:
 
 	/** Attempts to commit the ability's cooldown only. If BroadcastCommitEvent is true, it will broadcast the commit event that tasks like WaitAbilityCommit are listening for. */
 	UFUNCTION(BlueprintCallable, Category = Ability, DisplayName = "CommitAbilityCooldown")
-	virtual bool K2_CommitAbilityCooldown(bool BroadcastCommitEvent=false);
+	virtual bool K2_CommitAbilityCooldown(bool BroadcastCommitEvent=false, bool ForceCooldown=false);
 
 	/** Attempts to commit the ability's cost only. If BroadcastCommitEvent is true, it will broadcast the commit event that tasks like WaitAbilityCommit are listening for. */
 	UFUNCTION(BlueprintCallable, Category = Ability, DisplayName = "CommitAbilityCost")
@@ -379,7 +379,7 @@ public:
 	virtual bool K2_CheckAbilityCost();
 
 	virtual bool CommitAbility(const FGameplayAbilitySpecHandle Handle, const FGameplayAbilityActorInfo* ActorInfo, const FGameplayAbilityActivationInfo ActivationInfo);
-	virtual bool CommitAbilityCooldown(const FGameplayAbilitySpecHandle Handle, const FGameplayAbilityActorInfo* ActorInfo, const FGameplayAbilityActivationInfo ActivationInfo);
+	virtual bool CommitAbilityCooldown(const FGameplayAbilitySpecHandle Handle, const FGameplayAbilityActorInfo* ActorInfo, const FGameplayAbilityActivationInfo ActivationInfo, const bool ForceCooldown);
 	virtual bool CommitAbilityCost(const FGameplayAbilitySpecHandle Handle, const FGameplayAbilityActorInfo* ActorInfo, const FGameplayAbilityActivationInfo ActivationInfo);
 
 	/**
@@ -623,6 +623,14 @@ public:
 	/** Movement Sync */
 	virtual void SetMovementSyncPoint(FName SyncName);
 
+	/** Set when the remote instance of this ability has ended (but the local instance may still be running or finishing up */
+	bool RemoteInstanceEnded;
+
+	/** Called by ability system component to inform this ability instance the remote instance was ended */
+	void SetRemoteInstanceHasEnded();
+
+	void NotifyAbilityTaskWaitingOnPlayerData(class UAbilityTask* AbilityTask);
+
 protected:
 
 	// -----------------------------------------------	
diff --git a/Engine/Source/Runtime/GameplayAbilities/Public/Abilities/Tasks/AbilityTask.h b/Engine/Source/Runtime/GameplayAbilities/Public/Abilities/Tasks/AbilityTask.h
--- a/Engine/Source/Runtime/GameplayAbilities/Public/Abilities/Tasks/AbilityTask.h
+++ b/Engine/Source/Runtime/GameplayAbilities/Public/Abilities/Tasks/AbilityTask.h
@@ -59,6 +59,20 @@
  *	
  */
 
+/**
+ *	Latent tasks are waiting on something. This is to differeniate waiting on the user to do something vs waiting on the game to do something.
+ *	Tasks start WaitingOnGame, and are set to WaitingOnUser when appropriate (see WaitTargetData, WaitIiputPress, etc)
+ */
+UENUM()
+enum class EAbilityTaskWaitState : uint8
+{
+	/** Task is waiting for the game to do something */
+	WaitingOnGame,
+
+	/** Waiting for the user to do something */
+	WaitingOnUser,
+};
+
 UCLASS(Abstract)
 class GAMEPLAYABILITIES_API UAbilityTask : public UGameplayTask
 {
@@ -114,6 +128,14 @@ class GAMEPLAYABILITIES_API UAbilityTask : public UGameplayTask
 		static_assert(DelayedFalse<T>(), "UAbilityTask::NewTask should never be used. Use NewAbilityTask instead");
 	}
 
+	/** Called when the ability task is waiting on remote player data. IF the remote player ends the ability prematurely, and a task with this set is still running, the ability is killed. */
+	void SetWaitingOnRemotePlayerData();
+
+	virtual bool IsWaitingOnRemotePlayerdata() const override;
+
+	/** What we are waiting on */
+	EAbilityTaskWaitState WaitState;
+
 protected:
 	/** Helper method for registering client replicated callbacks */
 	bool CallOrAddReplicatedDelegate(EAbilityGenericReplicatedEvent::Type Event, FSimpleMulticastDelegate::FDelegate Delegate);
diff --git a/Engine/Source/Runtime/GameplayAbilities/Public/Abilities/Tasks/AbilityTask_WaitInputPress.h b/Engine/Source/Runtime/GameplayAbilities/Public/Abilities/Tasks/AbilityTask_WaitInputPress.h
--- a/Engine/Source/Runtime/GameplayAbilities/Public/Abilities/Tasks/AbilityTask_WaitInputPress.h
+++ b/Engine/Source/Runtime/GameplayAbilities/Public/Abilities/Tasks/AbilityTask_WaitInputPress.h
@@ -11,8 +11,8 @@ DECLARE_DYNAMIC_MULTICAST_DELEGATE_OneParam(FInputPressDelegate, float, TimeWait
  *	Waits until the input is pressed from activating an ability. This should be true immediately upon starting the ability, since the key was pressed to activate it.
  *	We expect server to execute this task in parallel and keep its own time. We do not keep track of 
  */
-UCLASS(MinimalAPI)
-class UAbilityTask_WaitInputPress : public UAbilityTask
+UCLASS()
+class GAMEPLAYABILITIES_API UAbilityTask_WaitInputPress : public UAbilityTask
 {
 	GENERATED_UCLASS_BODY()
 
diff --git a/Engine/Source/Runtime/GameplayAbilities/Public/AbilitySystemComponent.h b/Engine/Source/Runtime/GameplayAbilities/Public/AbilitySystemComponent.h
--- a/Engine/Source/Runtime/GameplayAbilities/Public/AbilitySystemComponent.h
+++ b/Engine/Source/Runtime/GameplayAbilities/Public/AbilitySystemComponent.h
@@ -153,8 +153,6 @@ class GAMEPLAYABILITIES_API UAbilitySystemComponent : public UGameplayTasksCompo
 	float GetNumericAttribute(const FGameplayAttribute &Attribute) const;
 	float GetNumericAttributeChecked(const FGameplayAttribute &Attribute) const;
 
-	virtual void DisplayDebug(class UCanvas* Canvas, const class FDebugDisplayInfo& DebugDisplay, float& YL, float& YPos);
-
 	// -- Replication -------------------------------------------------------------------------------------------------
 
 	virtual bool ReplicateSubobjects(class UActorChannel *Channel, class FOutBunch *Bunch, FReplicationFlags *RepFlags) override;
@@ -730,6 +728,55 @@ class GAMEPLAYABILITIES_API UAbilitySystemComponent : public UGameplayTasksCompo
 
 	void IncrementAbilityListLock();
 	void DecrementAbilityListLock();
+
+	// --------------------------------------------
+	// Debugging
+	// --------------------------------------------
+
+	struct FAbilitySystemComponentDebugInfo
+	{
+		FAbilitySystemComponentDebugInfo()
+		{
+			FMemory::Memzero(*this);
+		}
+
+		class UCanvas* Canvas;
+
+		bool bPrintToLog;
+
+		bool bShowAttributes;
+		bool bShowGameplayEffects;;
+		bool bShowAbilities;
+
+		float XPos;
+		float YPos;
+		float OriginalX;
+		float OriginalY;
+		float MaxY;
+		float NewColumnYPadding;
+		float YL;
+
+		bool Accumulate;
+		TArray<FString>	Strings;
+	};
+
+	virtual void DisplayDebug(class UCanvas* Canvas, const class FDebugDisplayInfo& DebugDisplay, float& YL, float& YPos);
+	virtual void PrintDebug();
+
+	void AccumulateScreenPos(FAbilitySystemComponentDebugInfo& Info);
+	virtual void Debug_Internal(struct FAbilitySystemComponentDebugInfo& Info);
+	void DebugLine(struct FAbilitySystemComponentDebugInfo& Info, FString Str, float XOffset, float YOffset);
+	FString CleanupName(FString Str);
+
+	UFUNCTION(Server, reliable, WithValidation)
+	void ServerPrintDebug_Request();
+
+	UFUNCTION(Client, reliable)
+	void ClientPrintDebug_Response(const TArray<FString>& Strings);
+
+	/** Called when the ability is forced cancelled due to replication */
+	void	ForceCancelAbilityDueToReplication(UGameplayAbility* Instance);
+
 protected:
 
 	/**
diff --git a/Engine/Source/Runtime/GameplayTasks/Classes/GameplayTask.h b/Engine/Source/Runtime/GameplayTasks/Classes/GameplayTask.h
--- a/Engine/Source/Runtime/GameplayTasks/Classes/GameplayTask.h
+++ b/Engine/Source/Runtime/GameplayTasks/Classes/GameplayTask.h
@@ -256,6 +256,8 @@ public:
 
 	ETaskResourceOverlapPolicy GetResourceOverlapPolicy() const { return ResourceOverlapPolicy; }
 
+	virtual bool IsWaitingOnRemotePlayerdata() const { return false; }
+
 protected:
 	/** End and CleanUp the task - may be called by the task itself or by the task owner if the owner is ending. 
 	 *	IMPORTANT! Do NOT call directly! Call EndTask() or TaskOwnerEnded() 
diff --git a/Engine/Source/Runtime/GameplayTasks/Private/GameplayTask.cpp b/Engine/Source/Runtime/GameplayTasks/Private/GameplayTask.cpp
--- a/Engine/Source/Runtime/GameplayTasks/Private/GameplayTask.cpp
+++ b/Engine/Source/Runtime/GameplayTasks/Private/GameplayTask.cpp
@@ -200,7 +200,7 @@ void UGameplayTask::OnDestroy(bool bOwnerFinished)
 
 FString UGameplayTask::GetDebugString() const
 {
-	return FString::Printf(TEXT("Generic %s"), *GetName());
+	return FString::Printf(TEXT("%s (%s)"), *GetName(), *InstanceName.ToString());
 }
 
 void UGameplayTask::AddRequiredResource(TSubclassOf<UGameplayTaskResource> RequiredResource)
diff --git a/Engine/Source/Runtime/LevelSequence/Private/LevelSequence.cpp b/Engine/Source/Runtime/LevelSequence/Private/LevelSequence.cpp
--- a/Engine/Source/Runtime/LevelSequence/Private/LevelSequence.cpp
+++ b/Engine/Source/Runtime/LevelSequence/Private/LevelSequence.cpp
@@ -145,7 +145,7 @@ FGuid ULevelSequence::FindObjectId(UObject& Object) const
 {
 	for (auto& Pair : CachedObjectBindings)
 	{
-		if (Pair.Value == &Object)
+		if (Pair.Value.Get() == &Object)
 		{
 			return Pair.Key;
 		}
diff --git a/Engine/Source/Runtime/LevelSequence/Private/LevelSequencePlayer.cpp b/Engine/Source/Runtime/LevelSequence/Private/LevelSequencePlayer.cpp
--- a/Engine/Source/Runtime/LevelSequence/Private/LevelSequencePlayer.cpp
+++ b/Engine/Source/Runtime/LevelSequence/Private/LevelSequencePlayer.cpp
@@ -174,6 +174,17 @@ void ULevelSequencePlayer::SetPlayRate(float PlayRate)
 	PlaybackSettings.PlayRate = PlayRate;
 }
 
+void ULevelSequencePlayer::SetPlaybackRange( const float NewStartTime, const float NewEndTime )
+{
+	if( LevelSequence != nullptr )
+	{
+		UMovieScene* MovieScene = LevelSequence->GetMovieScene();
+		if( MovieScene != nullptr )
+		{
+			MovieScene->SetPlaybackRange( NewStartTime, NewEndTime );
+		}
+	}
+}
 
 void ULevelSequencePlayer::OnCursorPositionChanged()
 {
@@ -206,8 +217,6 @@ void ULevelSequencePlayer::Initialize(ULevelSequence* InLevelSequence, UWorld* I
 	LevelSequence = InLevelSequence;
 	LevelSequence->BindToContext(InWorld);
 
-	SequenceStartOffset = LevelSequence->GetMovieScene()->GetPlaybackRange().GetLowerBoundValue();
-
 	World = InWorld;
 	PlaybackSettings = Settings;
 
@@ -310,6 +319,7 @@ void ULevelSequencePlayer::UpdateMovieSceneInstance(float CurrentPosition, float
 {
 	if(RootMovieSceneInstance.IsValid())
 	{
+		const float SequenceStartOffset = LevelSequence->GetMovieScene()->GetPlaybackRange().GetLowerBoundValue();
 		RootMovieSceneInstance->Update(CurrentPosition + SequenceStartOffset, PreviousPosition + SequenceStartOffset, *this);
 #if WITH_EDITOR
 		OnLevelSequencePlayerUpdate.Broadcast(*this, CurrentPosition, PreviousPosition);
diff --git a/Engine/Source/Runtime/LevelSequence/Public/LevelSequencePlayer.h b/Engine/Source/Runtime/LevelSequence/Public/LevelSequencePlayer.h
--- a/Engine/Source/Runtime/LevelSequence/Public/LevelSequencePlayer.h
+++ b/Engine/Source/Runtime/LevelSequence/Public/LevelSequencePlayer.h
@@ -121,6 +121,15 @@ public:
 	UFUNCTION(BlueprintCallable, Category="Game|Cinematic")
 	void SetPlayRate(float PlayRate);
 
+	/**
+	 * Sets the range in time to be played back by this player, overriding the default range stored in the asset
+	 *
+	 * @param	NewStartTime	The new starting time for playback
+	 * @param	NewEndTime		The new ending time for playback.  Must be larger than the start time.
+	 */
+	UFUNCTION(BlueprintCallable, Category="Game|Cinematic")
+	void SetPlaybackRange( const float NewStartTime, const float NewEndTime );
+
 protected:
 
 	// IMovieScenePlayer interface
@@ -171,9 +180,6 @@ private:
 
 private:
 
-	/** The offset from 0 that the sequence we're playing starts */
-	float SequenceStartOffset;
-
 	/** The root movie scene instance to update when playing. */
 	TSharedPtr<FMovieSceneSequenceInstance> RootMovieSceneInstance;
 
diff --git a/Engine/Source/Runtime/MovieScene/Private/MovieScene.cpp b/Engine/Source/Runtime/MovieScene/Private/MovieScene.cpp
--- a/Engine/Source/Runtime/MovieScene/Private/MovieScene.cpp
+++ b/Engine/Source/Runtime/MovieScene/Private/MovieScene.cpp
@@ -238,6 +238,8 @@ void UMovieScene::SetPlaybackRange(float Start, float End)
 {
 	if (ensure(End >= Start))
 	{
+		Modify();
+		
 		PlaybackRange = TRange<float>(Start, TRangeBound<float>::Inclusive(End));
 
 #if WITH_EDITORONLY_DATA
diff --git a/Engine/Source/Runtime/MovieScene/Private/MovieSceneCommonHelpers.cpp b/Engine/Source/Runtime/MovieScene/Private/MovieSceneCommonHelpers.cpp
--- a/Engine/Source/Runtime/MovieScene/Private/MovieSceneCommonHelpers.cpp
+++ b/Engine/Source/Runtime/MovieScene/Private/MovieSceneCommonHelpers.cpp
@@ -126,6 +126,42 @@ void MovieSceneHelpers::SetRuntimeObjectMobility(UObject* Object, EComponentMobi
 	}
 }
 
+void MovieSceneHelpers::SetKeyInterpolation(FRichCurve& InCurve, FKeyHandle InKeyHandle, EMovieSceneKeyInterpolation InKeyInterpolation)
+{
+	switch (InKeyInterpolation)
+	{
+		case EMovieSceneKeyInterpolation::Auto:
+			InCurve.SetKeyInterpMode(InKeyHandle, RCIM_Cubic);
+			InCurve.SetKeyTangentMode(InKeyHandle, RCTM_Auto);
+			break;
+
+		case EMovieSceneKeyInterpolation::User:
+			InCurve.SetKeyInterpMode(InKeyHandle, RCIM_Cubic);
+			InCurve.SetKeyTangentMode(InKeyHandle, RCTM_User);
+			break;
+
+		case EMovieSceneKeyInterpolation::Break:
+			InCurve.SetKeyInterpMode(InKeyHandle, RCIM_Cubic);
+			InCurve.SetKeyTangentMode(InKeyHandle, RCTM_Break);
+			break;
+
+		case EMovieSceneKeyInterpolation::Linear:
+			InCurve.SetKeyInterpMode(InKeyHandle, RCIM_Linear);
+			InCurve.SetKeyTangentMode(InKeyHandle, RCTM_Auto);
+			break;
+
+		case EMovieSceneKeyInterpolation::Constant:
+			InCurve.SetKeyInterpMode(InKeyHandle, RCIM_Constant);
+			InCurve.SetKeyTangentMode(InKeyHandle, RCTM_Auto);
+			break;
+
+		default:
+			InCurve.SetKeyInterpMode(InKeyHandle, RCIM_Cubic);
+			InCurve.SetKeyTangentMode(InKeyHandle, RCTM_Auto);
+			break;
+	}
+}
+
 FTrackInstancePropertyBindings::FTrackInstancePropertyBindings( FName InPropertyName, const FString& InPropertyPath, const FName& InFunctionName )
     : PropertyPath( InPropertyPath )
 	, PropertyName( InPropertyName )
diff --git a/Engine/Source/Runtime/MovieScene/Private/MovieScenePrivatePCH.h b/Engine/Source/Runtime/MovieScene/Private/MovieScenePrivatePCH.h
--- a/Engine/Source/Runtime/MovieScene/Private/MovieScenePrivatePCH.h
+++ b/Engine/Source/Runtime/MovieScene/Private/MovieScenePrivatePCH.h
@@ -14,5 +14,6 @@
 #include "MovieSceneSequenceInstance.h"
 #include "IMovieSceneTrackInstance.h"
 #include "MovieSceneSequence.h"
+#include "MovieSceneCommonHelpers.h"
 
 #endif
\ No newline at end of file
diff --git a/Engine/Source/Runtime/MovieScene/Private/MovieSceneSection.cpp b/Engine/Source/Runtime/MovieScene/Private/MovieSceneSection.cpp
--- a/Engine/Source/Runtime/MovieScene/Private/MovieSceneSection.cpp
+++ b/Engine/Source/Runtime/MovieScene/Private/MovieSceneSection.cpp
@@ -126,34 +126,7 @@ void UMovieSceneSection::AddKeyToCurve(FRichCurve& InCurve, float Time, float Va
 
 		if (!InCurve.IsKeyHandleValid(ExistingKeyHandle) && InCurve.IsKeyHandleValid(NewKeyHandle))
 		{
-			// Set the default key interpolation only if a new key was created
-			switch (Interpolation)
-			{
-				case EMovieSceneKeyInterpolation::Auto:
-					InCurve.SetKeyInterpMode(NewKeyHandle, RCIM_Cubic);
-					InCurve.SetKeyTangentMode(NewKeyHandle, RCTM_Auto);
-					break;
-				case EMovieSceneKeyInterpolation::User:
-					InCurve.SetKeyInterpMode(NewKeyHandle, RCIM_Cubic);
-					InCurve.SetKeyTangentMode(NewKeyHandle, RCTM_User);
-					break;
-				case EMovieSceneKeyInterpolation::Break:
-					InCurve.SetKeyInterpMode(NewKeyHandle, RCIM_Cubic);
-					InCurve.SetKeyTangentMode(NewKeyHandle, RCTM_Break);
-					break;
-				case EMovieSceneKeyInterpolation::Linear:
-					InCurve.SetKeyInterpMode(NewKeyHandle, RCIM_Linear);
-					InCurve.SetKeyTangentMode(NewKeyHandle, RCTM_Auto);
-					break;
-				case EMovieSceneKeyInterpolation::Constant:
-					InCurve.SetKeyInterpMode(NewKeyHandle, RCIM_Constant);
-					InCurve.SetKeyTangentMode(NewKeyHandle, RCTM_Auto);
-					break;
-				default:
-					InCurve.SetKeyInterpMode(NewKeyHandle, RCIM_Cubic);
-					InCurve.SetKeyTangentMode(NewKeyHandle, RCTM_Auto);
-					break;
-			}
+			MovieSceneHelpers::SetKeyInterpolation(InCurve, NewKeyHandle, Interpolation);
 		}
 	}
 }
diff --git a/Engine/Source/Runtime/MovieScene/Public/MovieSceneCommonHelpers.h b/Engine/Source/Runtime/MovieScene/Public/MovieSceneCommonHelpers.h
--- a/Engine/Source/Runtime/MovieScene/Public/MovieSceneCommonHelpers.h
+++ b/Engine/Source/Runtime/MovieScene/Public/MovieSceneCommonHelpers.h
@@ -2,6 +2,8 @@
 
 #pragma once
 
+#include "KeyParams.h"
+
 class MOVIESCENE_API MovieSceneHelpers
 {
 public:
@@ -41,6 +43,15 @@ public:
 	 * @param Mobility The mobility of the runtime object
 	 */
 	static void SetRuntimeObjectMobility(UObject* Object, EComponentMobility::Type ComponentMobility = EComponentMobility::Movable);
+
+	/*
+	 * Set the key interpolation
+	 *
+	 * @param InCurve The curve that contains the key handle to set
+	 * @param InKeyHandle The key handle to set
+	 * @param InInterpolation The interpolation to set
+	 */
+	static void SetKeyInterpolation(FRichCurve& InCurve, FKeyHandle InKeyHandle, EMovieSceneKeyInterpolation InKeyInterpolation);
 };
 
 /**
diff --git a/Engine/Source/Runtime/MovieSceneCapture/Private/AutomatedLevelSequenceCapture.cpp b/Engine/Source/Runtime/MovieSceneCapture/Private/AutomatedLevelSequenceCapture.cpp
--- a/Engine/Source/Runtime/MovieSceneCapture/Private/AutomatedLevelSequenceCapture.cpp
+++ b/Engine/Source/Runtime/MovieSceneCapture/Private/AutomatedLevelSequenceCapture.cpp
@@ -35,15 +35,70 @@ void UAutomatedLevelSequenceCapture::SetLevelSequenceActor(ALevelSequenceActor*
 
 void UAutomatedLevelSequenceCapture::Initialize(TWeakPtr<FSceneViewport> InViewport)
 {
+	// Apply command-line overrides
+	{
+		FString LevelSequenceAssetPath;
+		if( FParse::Value( FCommandLine::Get(), TEXT( "-LevelSequence=" ), LevelSequenceAssetPath ) )
+		{
+			LevelSequenceAsset.SetPath( LevelSequenceAssetPath );
+			LevelSequenceActorId = FGuid();
+		}
+
+		float PrerollOverride;
+		if( FParse::Value( FCommandLine::Get(), TEXT( "-MoviePrerollSeconds=" ), PrerollOverride ) )
+		{
+			PrerollAmount = PrerollOverride;
+		}
+
+		bool bStageSequenceOverride;
+		if( FParse::Bool( FCommandLine::Get(), TEXT( "-MoviePreStage=" ), bStageSequenceOverride ) )
+		{
+			bStageSequence = bStageSequenceOverride;
+		}
+	}
+
 	// Ensure the LevelSequence is up to date with the LevelSequenceActorId (Level sequence is only there for a nice UI)
 	LevelSequenceActor = FUniqueObjectGuid(LevelSequenceActorId);
 
-	// Attempt to find the level sequence actor
 	ALevelSequenceActor* Actor = LevelSequenceActor.Get();
+
+	// If we don't have a valid actor, attempt to find a level sequence actor in the world that references this asset
+	if( Actor == nullptr )
+	{
+		if( LevelSequenceAsset.IsValid() )
+		{
+			ULevelSequence* Asset = Cast<ULevelSequence>( LevelSequenceAsset.TryLoad() );
+			if( Asset != nullptr )
+			{
+				for( auto It = TActorIterator<ALevelSequenceActor>( InViewport.Pin()->GetClient()->GetWorld() ); It; ++It )
+				{
+					if( It->LevelSequence == LevelSequenceAsset )
+					{
+						// Found it!
+						Actor = *It;
+						this->LevelSequenceActor = Actor;
+
+						break;
+					}
+				}
+			}
+		}
+	}
+
 	if (Actor)
 	{
-		// Ensure that the sequence cannot play by itself. This is safe because this code is called before BeginPlay()
+		// Make sure we're not playing yet (in case AutoPlay was called from BeginPlay)
+		if( Actor->SequencePlayer != nullptr && Actor->SequencePlayer->IsPlaying() )
+		{
+			Actor->SequencePlayer->Stop();
+		}
 		Actor->bAutoPlay = false;
+
+		// Bind to the event so we know when to capture a frame
+		if( Actor->SequencePlayer != nullptr )
+		{
+			OnPlayerUpdatedBinding = Actor->SequencePlayer->OnSequenceUpdated().AddUObject( this, &UAutomatedLevelSequenceCapture::SequenceUpdated );
+		}
 	}
 	
 	PrerollTime = 0.f;
@@ -52,8 +107,56 @@ void UAutomatedLevelSequenceCapture::Initialize(TWeakPtr<FSceneViewport> InViewp
 	Super::Initialize(InViewport);
 }
 
+
+void UAutomatedLevelSequenceCapture::SetupFrameRange()
+{
+	ALevelSequenceActor* Actor = LevelSequenceActor.Get();
+	if( Actor )
+	{
+		ULevelSequence* LevelSequence = Cast<ULevelSequence>( Actor->LevelSequence.TryLoad() );
+		if( LevelSequence != nullptr )
+		{
+			UMovieScene* MovieScene = LevelSequence->GetMovieScene();
+			if( MovieScene != nullptr )
+			{
+				// We always add 1 to the number of frames we want to capture, because we want to capture both the start and end frames (which if the play range is 0, would still yield a single frame)
+				const int32 SequenceStartFrame = FMath::RoundToInt( MovieScene->GetPlaybackRange().GetLowerBoundValue() * Settings.FrameRate );
+				const int32 SequenceEndFrame = FMath::Max( SequenceStartFrame, FMath::RoundToInt( MovieScene->GetPlaybackRange().GetUpperBoundValue() * Settings.FrameRate ) );
+
+				// Default to playing back the sequence's stored playback range
+				int32 PlaybackStartFrame = SequenceStartFrame;
+				int32 PlaybackEndFrame = SequenceEndFrame;
+
+				if( Settings.bUseCustomStartFrame )
+				{
+					PlaybackStartFrame = Settings.bUseRelativeFrameNumbers ? ( SequenceStartFrame + Settings.StartFrame ) : Settings.StartFrame;
+				}
+
+				if( Settings.bUseCustomEndFrame )
+				{
+					PlaybackEndFrame = FMath::Max( PlaybackStartFrame, Settings.bUseRelativeFrameNumbers ? ( SequenceEndFrame + Settings.EndFrame ) : Settings.EndFrame );
+				}
+
+				this->FrameCount = ( PlaybackEndFrame - PlaybackStartFrame ) + 1;
+
+				if( !Settings.bUseRelativeFrameNumbers )
+				{
+					this->FrameNumberOffset = PlaybackStartFrame;
+				}
+
+				// Override the movie scene's playback range
+				Actor->SequencePlayer->SetPlaybackRange(
+					(float)PlaybackStartFrame / (float)Settings.FrameRate,
+					(float)PlaybackEndFrame / (float)Settings.FrameRate );
+			}
+		}
+	}
+}
+
+
 void UAutomatedLevelSequenceCapture::Tick(float DeltaSeconds)
 {
+	const bool bAnyFramesToCapture = OutstandingFrameCount > 0;
 	Super::Tick(DeltaSeconds);
 
 	ALevelSequenceActor* Actor = LevelSequenceActor.Get();
@@ -94,40 +197,49 @@ void UAutomatedLevelSequenceCapture::Tick(float DeltaSeconds)
 			Value += DeltaSeconds;
 			if (Value >= PrerollAmount)
 			{
-				if (!OnPlayerUpdatedBinding.IsValid())
-				{
-					// Bind to the event so we know when to capture a frame
-					OnPlayerUpdatedBinding = Actor->SequencePlayer->OnSequenceUpdated().AddUObject(this, &UAutomatedLevelSequenceCapture::SequenceUpdated);
-				}
+				Actor->SequencePlayer->Play();
 
 				StartCapture();
-				Actor->SequencePlayer->Play();
 
 				PrerollTime.Reset();
 			}
 		}
 
-		// Otherwise we can quit if we've processed all the frames
-		else if (OutstandingFrameCount == 0)
+		if( !PrerollTime.IsSet() )
 		{
-			StopCapture();
+			if( bAnyFramesToCapture && OutstandingFrameCount == 0 )
+			{
+				// If we hit this, then we've rendered out the last frame and can stop!
+				StopCapture();
+			}
 		}
 	}
 }
 
 void UAutomatedLevelSequenceCapture::OnCaptureStopped()
 {
+	ALevelSequenceActor* Actor = LevelSequenceActor.Get();
+	if( Actor != nullptr )
+	{
+		if( Actor->SequencePlayer != nullptr )
+		{
+			Actor->SequencePlayer->OnSequenceUpdated().Remove( OnPlayerUpdatedBinding );
+		}
+	}
+
 	Close();
 	FPlatformMisc::RequestExit(0);
 }
 
 void UAutomatedLevelSequenceCapture::SequenceUpdated(const ULevelSequencePlayer& Player, float CurrentTime, float PreviousTime)
 {
-	// Grab this frame
-	PrepareForScreenshot();
-	if (!Player.IsPlaying())
+	if( bCapturing )
 	{
-		Player.OnSequenceUpdated().Remove(OnPlayerUpdatedBinding);
+		// Save the previous rendered frame
+		CaptureFrame( LastFrameDelta );
+
+		// Get ready to capture the next rendered frame
+		PrepareForScreenshot();
 	}
 }
 
diff --git a/Engine/Source/Runtime/MovieSceneCapture/Private/LevelCapture.cpp b/Engine/Source/Runtime/MovieSceneCapture/Private/LevelCapture.cpp
--- a/Engine/Source/Runtime/MovieSceneCapture/Private/LevelCapture.cpp
+++ b/Engine/Source/Runtime/MovieSceneCapture/Private/LevelCapture.cpp
@@ -30,12 +30,15 @@ void ULevelCapture::Tick(float DeltaSeconds)
 	AActor* PrerequisiteActorPtr = PrerequisiteActor.Get();
 	if (!PrerequisiteActorId.IsValid() || (PrerequisiteActorPtr && PrerequisiteActorPtr->HasActorBegunPlay()))
 	{
-		if (!bCapturing)
+		if( !Settings.bUseCustomStartFrame || GetMetrics().Frame >= Settings.StartFrame )
 		{
-			StartCapture();
-		}
+			if (!bCapturing)
+			{
+				StartCapture();
+			}
 
-		PrepareForScreenshot();
+			PrepareForScreenshot();
+		}
 	}
 }
 
diff --git a/Engine/Source/Runtime/MovieSceneCapture/Private/MovieSceneCapture.cpp b/Engine/Source/Runtime/MovieSceneCapture/Private/MovieSceneCapture.cpp
--- a/Engine/Source/Runtime/MovieSceneCapture/Private/MovieSceneCapture.cpp
+++ b/Engine/Source/Runtime/MovieSceneCapture/Private/MovieSceneCapture.cpp
@@ -26,18 +26,26 @@ FMovieSceneCaptureSettings::FMovieSceneCaptureSettings()
 	: Resolution(1280, 720)
 {
 	OutputDirectory.Path = FPaths::VideoCaptureDir();
+	FPaths::MakePlatformFilename( OutputDirectory.Path );
+
+	bCreateTemporaryCopiesOfLevels = false;
+	bUseRelativeFrameNumbers = false;
+	GameModeOverride = nullptr;
 	OutputFormat = NSLOCTEXT("MovieCapture", "DefaultFormat", "{world}_{width}x{height}").ToString();
 	FrameRate = 24;
-	FrameCount = 0;
+	bUseCustomStartFrame = false;
+	StartFrame = 0;
+	bUseCustomEndFrame = false;
+	EndFrame = 1;
 	CaptureType = EMovieCaptureType::AVI;
 	bUseCompression = true;
 	CompressionQuality = 1.f;
 	bEnableTextureStreaming = true;
 	bCinematicMode = true;
-	bAllowMovement = true;
-	bAllowTurning = true;
-	bShowPlayer = true;
-	bShowHUD = true;
+	bAllowMovement = false;
+	bAllowTurning = false;
+	bShowPlayer = false;
+	bShowHUD = false;
 }
 
 UMovieSceneCapture::UMovieSceneCapture(const FObjectInitializer& Initializer)
@@ -54,21 +62,94 @@ UMovieSceneCapture::UMovieSceneCapture(const FObjectInitializer& Initializer)
 		InheritedCommandLineArguments.AppendChar(' ');
 	}
 
-	// the PIEVIACONSOLE parameter tells UGameEngine to add the auto-save dir to the paths array and repopulate the package file cache
-	// this is needed in order to support streaming levels as the streaming level packages will be loaded only when needed (thus
-	// their package names need to be findable by the package file caching system)
-	// (we add to EditorCommandLine because the URL is ignored by WindowsTools)
-	AdditionalCommandLineArguments += TEXT("-PIEVIACONSOLE -NoLoadingScreen ");
+	AdditionalCommandLineArguments += TEXT("-NoLoadingScreen");
 
 	Handle = FUniqueMovieSceneCaptureHandle();
 
+	FrameCount = 0;
 	OutstandingFrameCount = 0;
 	LastFrameDelta = 0.f;
 	bCapturing = false;
+	FrameNumberOffset = 0;
 }
 
 void UMovieSceneCapture::Initialize(TWeakPtr<FSceneViewport> InSceneViewport)
 {
+	// Apply command-line overrides
+	{
+		FString OutputPathOverride;
+		if( FParse::Value( FCommandLine::Get(), TEXT( "-MovieFolder=" ), OutputPathOverride ) )
+		{
+			Settings.OutputDirectory.Path = OutputPathOverride;
+		}
+
+		FString OutputNameOverride;
+		if( FParse::Value( FCommandLine::Get(), TEXT( "-MovieName=" ), OutputNameOverride ) )
+		{
+			Settings.OutputFormat = OutputNameOverride;
+		}
+
+		bool bOverrideRelativeFrameNumbers;
+		if( FParse::Bool( FCommandLine::Get(), TEXT( "-MovieRelativeFrames=" ), bOverrideRelativeFrameNumbers ) )
+		{
+			Settings.bUseRelativeFrameNumbers = bOverrideRelativeFrameNumbers;
+		}
+
+		bool bOverrideCinematicMode;
+		if( FParse::Bool( FCommandLine::Get(), TEXT( "-MovieCinematicMode=" ), bOverrideCinematicMode ) )
+		{
+			Settings.bCinematicMode = bOverrideCinematicMode;
+		}
+
+		FString FormatOverride;
+		if( FParse::Value( FCommandLine::Get(), TEXT( "-MovieFormat=" ), FormatOverride ) )
+		{
+			if( FormatOverride.Equals( TEXT( "AVI" ), ESearchCase::IgnoreCase ) )
+			{
+				Settings.CaptureType = EMovieCaptureType::AVI;
+			}
+			else if( FormatOverride.Equals( TEXT( "BMP" ), ESearchCase::IgnoreCase ) )
+			{
+				Settings.CaptureType = EMovieCaptureType::BMP;
+			}
+			else if( FormatOverride.Equals( TEXT( "PNG" ), ESearchCase::IgnoreCase ) )
+			{
+				Settings.CaptureType = EMovieCaptureType::PNG;
+			}
+			else if( FormatOverride.Equals( TEXT( "JPG" ), ESearchCase::IgnoreCase ) || FormatOverride.Equals( TEXT( "JPEG" ), ESearchCase::IgnoreCase ) )
+			{
+				Settings.CaptureType = EMovieCaptureType::JPEG;
+			}
+		}
+
+		int32 StartFrameOverride;
+		if( FParse::Value( FCommandLine::Get(), TEXT( "-MovieStartFrame=" ), StartFrameOverride ) )
+		{
+			Settings.bUseCustomStartFrame = true;
+			Settings.StartFrame = StartFrameOverride;
+		}
+
+		int32 EndFrameOverride;
+		if( FParse::Value( FCommandLine::Get(), TEXT( "-MovieEndFrame=" ), EndFrameOverride ) )
+		{
+			Settings.bUseCustomEndFrame = true;
+			Settings.EndFrame = EndFrameOverride;
+		}
+
+		int32 FrameRateOverride;
+		if( FParse::Value( FCommandLine::Get(), TEXT( "-MovieFrameRate=" ), FrameRateOverride ) )
+		{
+			Settings.FrameRate = FrameRateOverride;
+		}
+
+		int32 QualityOverride;
+		if( FParse::Value( FCommandLine::Get(), TEXT( "-MovieQuality=" ), QualityOverride ) )
+		{
+			Settings.CompressionQuality = FMath::TruncToFloat( FMath::Clamp( QualityOverride, 0, 100 ) ) / 100.0f;
+			Settings.bUseCompression = QualityOverride <= 0.99f;
+		}
+	}
+
 	SceneViewport = InSceneViewport;
 
 	if (!HasAnyFlags(RF_ClassDefaultObject))
@@ -126,11 +207,7 @@ void UMovieSceneCapture::PrepareForScreenshot()
 
 void UMovieSceneCapture::Tick(float DeltaSeconds)
 {
-	if (OutstandingFrameCount > 0)
-	{
-		CaptureFrame(LastFrameDelta);
-		--OutstandingFrameCount;
-	}
+	CaptureFrame(LastFrameDelta);
 	LastFrameDelta = DeltaSeconds;
 }
 
@@ -186,10 +263,12 @@ void UMovieSceneCapture::StartCapture()
 void UMovieSceneCapture::CaptureFrame(float DeltaSeconds)
 {
 	auto Viewport = SceneViewport.Pin();
-	if (!CaptureStrategy.IsValid() || !Viewport.IsValid() || !ScratchBuffer.Num())
+	if (OutstandingFrameCount == 0 || !CaptureStrategy.IsValid() || !Viewport.IsValid() || !ScratchBuffer.Num())
 	{
 		return;
 	}
+	
+	--OutstandingFrameCount;
 
 	CachedMetrics.ElapsedSeconds += DeltaSeconds;
 
@@ -198,12 +277,12 @@ void UMovieSceneCapture::CaptureFrame(float DeltaSeconds)
 	{
 		TArray<FColor> ThisFrameBuffer;
 		Swap(ThisFrameBuffer, ScratchBuffer);
+		ScratchBuffer.Reset();
 
 		uint32 NumDroppedFrames = CaptureStrategy->GetDroppedFrames(CachedMetrics.ElapsedSeconds, CachedMetrics.Frame);
 		CachedMetrics.Frame += NumDroppedFrames;
 
 		CaptureStrategy->OnPresent(CachedMetrics.ElapsedSeconds, CachedMetrics.Frame);
-		++CachedMetrics.Frame;
 		
 		if (AVIWriter)
 		{
@@ -216,7 +295,8 @@ void UMovieSceneCapture::CaptureFrame(float DeltaSeconds)
 			SaveFrameToFile(MoveTemp(ThisFrameBuffer));
 		}
 #endif
-		if (Settings.FrameCount != 0 && CachedMetrics.Frame >= Settings.FrameCount)
+		++CachedMetrics.Frame;
+		if (FrameCount != 0 && CachedMetrics.Frame >= FrameCount)
 		{
 			StopCapture();
 		}
@@ -313,7 +393,7 @@ const TCHAR* UMovieSceneCapture::GetDefaultFileExtension() const
 	{
 		case EMovieCaptureType::BMP:	return TEXT(".bmp");
 		case EMovieCaptureType::PNG:	return TEXT(".png");
-		case EMovieCaptureType::JPEG:	return TEXT(".jpeg");
+		case EMovieCaptureType::JPEG:	return TEXT(".jpg");
 #if PLATFORM_MAC
 		default:						return TEXT(".mov");
 #else
@@ -326,7 +406,7 @@ FString UMovieSceneCapture::ResolveFileFormat(const FString& Folder, const FStri
 {
 	TMap<FString, FStringFormatArg> Mappings;
 	Mappings.Add(TEXT("fps"), FString::Printf(TEXT("%d"), Settings.FrameRate));
-	Mappings.Add(TEXT("frame"), FString::Printf(TEXT("%04d"), CachedMetrics.Frame));
+	Mappings.Add(TEXT("frame"), FString::Printf(TEXT("%04d"), Settings.bUseRelativeFrameNumbers ? CachedMetrics.Frame : CachedMetrics.Frame + FrameNumberOffset));
 	Mappings.Add(TEXT("width"), FString::Printf(TEXT("%d"), CachedMetrics.Width));
 	Mappings.Add(TEXT("height"), FString::Printf(TEXT("%d"), CachedMetrics.Height));
 	Mappings.Add(TEXT("world"), GWorld->GetName());
@@ -345,14 +425,15 @@ FString UMovieSceneCapture::ResolveFileFormat(const FString& Folder, const FStri
 
 FString UMovieSceneCapture::ResolveUniqueFilename()
 {
-	FString BaseFilename = ResolveFileFormat(Settings.OutputDirectory.Path, Settings.OutputFormat);
-	FString ThisTry = BaseFilename + GetDefaultFileExtension();
-
-	if (!IFileManager::Get().DirectoryExists(*Settings.OutputDirectory.Path))
+	if( !IFileManager::Get().DirectoryExists( *Settings.OutputDirectory.Path ) )
 	{
-		IFileManager::Get().MakeDirectory(*Settings.OutputDirectory.Path);
+		IFileManager::Get().MakeDirectory( *Settings.OutputDirectory.Path );
 	}
 
+	const FString BaseFilename = ResolveFileFormat( Settings.OutputDirectory.Path, Settings.OutputFormat );
+
+	FString ThisTry = BaseFilename + GetDefaultFileExtension();
+
 	if (Settings.bOverwriteExisting)
 	{
 		// Try and delete it first
@@ -370,10 +451,10 @@ FString UMovieSceneCapture::ResolveUniqueFilename()
 		return ThisTry;
 	}
 
-	uint32 Index = 2;
+	int32 DuplicateIndex = 1;
 	for (;;)
 	{
-		ThisTry = BaseFilename + FString::Printf(TEXT(" %02d"), Index) + GetDefaultFileExtension();
+		ThisTry = BaseFilename + FString::Printf(TEXT("_(%d)"), DuplicateIndex) + GetDefaultFileExtension();
 
 		// If the file doesn't exist, we can use that, else, increment the index and try again
 		if (IFileManager::Get().FileSize(*ThisTry) == -1)
@@ -381,7 +462,7 @@ FString UMovieSceneCapture::ResolveUniqueFilename()
 			return ThisTry;
 		}
 
-		++Index;
+		++DuplicateIndex;
 	}
 
 	return ThisTry;
diff --git a/Engine/Source/Runtime/MovieSceneCapture/Private/MovieSceneCaptureModule.cpp b/Engine/Source/Runtime/MovieSceneCapture/Private/MovieSceneCaptureModule.cpp
--- a/Engine/Source/Runtime/MovieSceneCapture/Private/MovieSceneCaptureModule.cpp
+++ b/Engine/Source/Runtime/MovieSceneCapture/Private/MovieSceneCaptureModule.cpp
@@ -19,7 +19,7 @@ private:
 	virtual void StartupModule() override
 	{
 		FCoreDelegates::OnPreExit.AddRaw(this, &FMovieSceneCaptureModule::PreExit);
-		FWorldDelegates::OnPostWorldInitialization.AddRaw(this, &FMovieSceneCaptureModule::OnPostWorldInitialization);
+		FCoreUObjectDelegates::PostLoadMap.AddRaw(this, &FMovieSceneCaptureModule::OnPostLoadMap );
 	}
 
 	void PreExit()
@@ -27,22 +27,18 @@ private:
 		FActiveMovieSceneCaptures::Get().Shutdown();
 	}
 
-	void OnPostWorldInitialization(UWorld* World, const UWorld::InitializationValues IVS)
+	void OnPostLoadMap()
 	{
-		if (!World)
-		{
-			return;
-		}
-
 		UGameEngine* GameEngine = Cast<UGameEngine>(GEngine);
 		if (GameEngine && StartupMovieCaptureHandle.IsValid())
 		{
 			IMovieSceneCaptureInterface* StartupCaptureInterface = RetrieveMovieSceneInterface(StartupMovieCaptureHandle);
 			StartupCaptureInterface->Initialize(GameEngine->SceneViewport.ToSharedRef());
+			StartupCaptureInterface->SetupFrameRange();
 		}
 
 		StartupMovieCaptureHandle = FMovieSceneCaptureHandle();
-		FWorldDelegates::OnPostWorldInitialization.RemoveAll(this);
+		FCoreUObjectDelegates::PostLoadMap.RemoveAll(this);
 	}
 
 	virtual void PreUnloadCallback() override
@@ -58,56 +54,83 @@ private:
 			return nullptr;
 		}
 
+		FString TypeName;
+		if( FParse::Value( FCommandLine::Get(), TEXT( "-MovieSceneCaptureType=" ), TypeName ) )
+		{
+			// OK, they specified the type of capture they want on the command-line.  Manifests are now optional!
+		}
+
 		FString ManifestPath;
 		if (!FParse::Value(FCommandLine::Get(), TEXT("-MovieSceneCaptureManifest="), ManifestPath) || ManifestPath.IsEmpty())
 		{
-			return nullptr;
+			// Allow capturing without a manifest.  Command-line parameters for individual options will be necessary.
+			if( TypeName.IsEmpty() )
+			{
+				return nullptr;
+			}
 		}
 
 		UMovieSceneCapture* Capture = nullptr;
-		FString Json;
-		if (FFileHelper::LoadFileToString(Json, *ManifestPath))
+		if( !ManifestPath.IsEmpty() )
 		{
-			TSharedPtr<FJsonObject> RootObject;
-			TSharedRef<TJsonReader<> > JsonReader = TJsonReaderFactory<>::Create(Json);
-			if (FJsonSerializer::Deserialize(JsonReader, RootObject) && RootObject.IsValid())
+			FString Json;
+			if (FFileHelper::LoadFileToString(Json, *ManifestPath))
 			{
-				auto TypeField = RootObject->TryGetField(TEXT("Type"));
-				if (!TypeField.IsValid())
-				{
-					return nullptr;
-				}
-
-				UClass* Class = FindObject<UClass>(nullptr, *TypeField->AsString());
-				if (!Class)
-				{
-					return nullptr;
-				}
-
-				Capture = NewObject<UMovieSceneCapture>(GetTransientPackage(), Class);
-				if (!Capture)
-				{
-					return nullptr;
-				}
-
-				auto DataField = RootObject->TryGetField(TEXT("Data"));
-				if (!DataField.IsValid())
+				TSharedPtr<FJsonObject> RootObject;
+				TSharedRef<TJsonReader<> > JsonReader = TJsonReaderFactory<>::Create(Json);
+				if (FJsonSerializer::Deserialize(JsonReader, RootObject) && RootObject.IsValid())
 				{
-					return nullptr;
+					auto TypeField = RootObject->TryGetField(TEXT("Type"));
+					if (!TypeField.IsValid())
+					{
+						return nullptr;
+					}
+
+					TypeName = TypeField->AsString();
+					UClass* Class = FindObject<UClass>(nullptr, *TypeName);
+					if (!Class)
+					{
+						return nullptr;
+					}
+
+					Capture = NewObject<UMovieSceneCapture>(GetTransientPackage(), Class);
+					if (!Capture)
+					{
+						return nullptr;
+					}
+
+					auto DataField = RootObject->TryGetField(TEXT("Data"));
+					if (!DataField.IsValid())
+					{
+						return nullptr;
+					}
+
+					if (!FJsonObjectConverter::JsonAttributesToUStruct(DataField->AsObject()->Values, Class, Capture, 0, 0))
+					{
+						return nullptr;
+					}
 				}
+			}
+		}
+		else if( !TypeName.IsEmpty() )
+		{
+			UClass* Class = FindObject<UClass>( nullptr, *TypeName );
+			if( !Class )
+			{
+				return nullptr;
+			}
 
-				if (!FJsonObjectConverter::JsonAttributesToUStruct(DataField->AsObject()->Values, Class, Capture, 0, 0))
-				{
-					return nullptr;
-				}
-				
-				FActiveMovieSceneCaptures::Get().Add(Capture);
-				StartupMovieCaptureHandle = Capture->GetHandle();
-				return Capture;
+			Capture = NewObject<UMovieSceneCapture>( GetTransientPackage(), Class );
+			if( !Capture )
+			{
+				return nullptr;
 			}
 		}
 
-		return nullptr;
+		check( Capture != nullptr );
+		FActiveMovieSceneCaptures::Get().Add( Capture );
+		StartupMovieCaptureHandle = Capture->GetHandle();
+		return Capture;
 	}
 
 	virtual IMovieSceneCaptureInterface* CreateMovieSceneCapture(TWeakPtr<FSceneViewport> InSceneViewport) override
@@ -115,6 +138,7 @@ private:
 		UMovieSceneCapture* Capture = NewObject<UMovieSceneCapture>(GetTransientPackage());
 		FActiveMovieSceneCaptures::Get().Add(Capture);
 		Capture->Initialize(InSceneViewport);
+		Capture->SetupFrameRange();
 		Capture->StartCapture();
 		return Capture;
 	}
diff --git a/Engine/Source/Runtime/MovieSceneCapture/Public/AutomatedLevelSequenceCapture.h b/Engine/Source/Runtime/MovieSceneCapture/Public/AutomatedLevelSequenceCapture.h
--- a/Engine/Source/Runtime/MovieSceneCapture/Public/AutomatedLevelSequenceCapture.h
+++ b/Engine/Source/Runtime/MovieSceneCapture/Public/AutomatedLevelSequenceCapture.h
@@ -36,6 +36,7 @@ public:
 	// UMovieSceneCapture interface
 	virtual void Initialize(TWeakPtr<FSceneViewport> InViewport) override;
 	virtual void PostEditChangeProperty(FPropertyChangedEvent& PropertyChangedEvent) override;
+	virtual void SetupFrameRange() override;
 
 private:
 
diff --git a/Engine/Source/Runtime/MovieSceneCapture/Public/IMovieSceneCapture.h b/Engine/Source/Runtime/MovieSceneCapture/Public/IMovieSceneCapture.h
--- a/Engine/Source/Runtime/MovieSceneCapture/Public/IMovieSceneCapture.h
+++ b/Engine/Source/Runtime/MovieSceneCapture/Public/IMovieSceneCapture.h
@@ -24,6 +24,11 @@ public:
 	/** Initialize this capture object by binding it to the specified viewport */
 	virtual void Initialize(TWeakPtr<FSceneViewport> Viewport) = 0;
 
+	/** Called after initialization to allow the movie scene capture object to override the range of frames to 
+	    capture.  It's called after initialization so you can be sure that the FrameRate and other variables are
+		fully initialized first. */
+	virtual void SetupFrameRange() = 0;
+
 	/** Shut down this movie capture */
 	virtual void Close() = 0;
 
diff --git a/Engine/Source/Runtime/MovieSceneCapture/Public/MovieSceneCapture.h b/Engine/Source/Runtime/MovieSceneCapture/Public/MovieSceneCapture.h
--- a/Engine/Source/Runtime/MovieSceneCapture/Public/MovieSceneCapture.h
+++ b/Engine/Source/Runtime/MovieSceneCapture/Public/MovieSceneCapture.h
@@ -51,6 +51,9 @@ public:
 
 	// Begin IMovieSceneCaptureInterface
 	virtual void Initialize(TWeakPtr<FSceneViewport> InSceneViewport) override;
+	virtual void SetupFrameRange() override
+	{
+	}
 	virtual void Close() override;
 	virtual FMovieSceneCaptureHandle GetHandle() const override { return Handle; }
 	const FMovieSceneCaptureSettings& GetSettings() const override { return Settings; }
@@ -151,6 +154,8 @@ protected:
 	TSharedPtr<ICaptureStrategy> CaptureStrategy;
 	/** Scratch space for per-frame screenshots */
 	TArray<FColor> ScratchBuffer;
+	/** The number of frames to capture.  If this is zero, we'll capture the entire sequence. */
+	int32 FrameCount;
 	/** Running count of how many frames we're waiting to capture from slate */
 	int32 OutstandingFrameCount;
 	/** The delta of the last frame */
@@ -159,6 +164,10 @@ protected:
 	bool bCapturing;
 	/** Class responsible for ticking this instance */
 	TUniquePtr<FTicker> Ticker;
+	/** Frame number index offset when saving out frames.  This is used to allow the frame numbers on disk to match
+	    what they would be in the authoring application, rather than a simple 0-based sequential index */
+	int32 FrameNumberOffset;
+
 };
 
 /** A strategy that employs a fixed frame time-step, and as such never drops a frame. Potentially accelerated.  */
