Copying //UE4/Dev-Framework to //UE4/Main

diff --git a/Engine/Source/Runtime/GameplayAbilities/Public/AbilitySystemComponent.h b/Engine/Source/Runtime/GameplayAbilities/Public/AbilitySystemComponent.h
--- a/Engine/Source/Runtime/GameplayAbilities/Public/AbilitySystemComponent.h
+++ b/Engine/Source/Runtime/GameplayAbilities/Public/AbilitySystemComponent.h
@@ -45,7 +45,7 @@
 DECLARE_MULTICAST_DELEGATE_OneParam(FTargetingRejectedConfirmation, int32);
 
 /** Called when ability fails to activate, passes along the failed ability and a tag explaining why */
-DECLARE_MULTICAST_DELEGATE_TwoParams(FAbilityFailedDelegate, const UGameplayAbility*, FGameplayTagContainer);
+DECLARE_MULTICAST_DELEGATE_TwoParams(FAbilityFailedDelegate, const UGameplayAbility*, const FGameplayTagContainer&);
 
 /** Called when ability ends */
 DECLARE_MULTICAST_DELEGATE_OneParam(FAbilityEnded, UGameplayAbility*);
@@ -608,10 +608,10 @@ class GAMEPLAYABILITIES_API UAbilitySystemComponent : public UGameplayTasksCompo
 	 */
 
 	/** Grants Ability. Returns handle that can be used in TryActivateAbility, etc. */
-	FGameplayAbilitySpecHandle GiveAbility(FGameplayAbilitySpec AbilitySpec);
+	FGameplayAbilitySpecHandle GiveAbility(const FGameplayAbilitySpec& AbilitySpec);
 
 	/** Grants an ability and attempts to activate it exactly one time, which will cause it to be removed. Only valid on the server! */
-	FGameplayAbilitySpecHandle GiveAbilityAndActivateOnce(FGameplayAbilitySpec AbilitySpec);
+	FGameplayAbilitySpecHandle GiveAbilityAndActivateOnce(const FGameplayAbilitySpec& AbilitySpec);
 
 	/** Wipes all 'given' abilities. */
 	void ClearAllAbilities();
@@ -932,7 +932,7 @@ public:
 
 	virtual void NotifyAbilityCommit(UGameplayAbility* Ability);
 	virtual void NotifyAbilityActivated(const FGameplayAbilitySpecHandle Handle, UGameplayAbility* Ability);
-	virtual void NotifyAbilityFailed(const FGameplayAbilitySpecHandle Handle, UGameplayAbility* Ability, FGameplayTagContainer FailureReason);
+	virtual void NotifyAbilityFailed(const FGameplayAbilitySpecHandle Handle, UGameplayAbility* Ability, const FGameplayTagContainer& FailureReason);
 
 	UPROPERTY()
 	TArray<AGameplayAbilityTargetActor*>	SpawnedTargetActors;
@@ -1098,7 +1098,7 @@ public:
 	
 	/**  */
 	UFUNCTION(Server, reliable, WithValidation)
-	void ServerSetReplicatedTargetData(FGameplayAbilitySpecHandle AbilityHandle, FPredictionKey AbilityOriginalPredictionKey, FGameplayAbilityTargetDataHandle ReplicatedTargetDataHandle, FGameplayTag ApplicationTag, FPredictionKey CurrentPredictionKey);
+	void ServerSetReplicatedTargetData(FGameplayAbilitySpecHandle AbilityHandle, FPredictionKey AbilityOriginalPredictionKey, const FGameplayAbilityTargetDataHandle& ReplicatedTargetDataHandle, FGameplayTag ApplicationTag, FPredictionKey CurrentPredictionKey);
 
 	/** Replicates to the server that targeting has been cancelled */
 	UFUNCTION(Server, reliable, WithValidation)
diff --git a/Engine/Source/Runtime/GameplayAbilities/Public/GameplayAbilitySpec.h b/Engine/Source/Runtime/GameplayAbilities/Public/GameplayAbilitySpec.h
--- a/Engine/Source/Runtime/GameplayAbilities/Public/GameplayAbilitySpec.h
+++ b/Engine/Source/Runtime/GameplayAbilities/Public/GameplayAbilitySpec.h
@@ -163,7 +163,7 @@ struct GAMEPLAYABILITIES_API FGameplayAbilityActivationInfo
 
 	/** An ability that runs on multiple game instances can be canceled by a remote instance, but only if that remote instance has already confirmed starting it. */
 	UPROPERTY()
-	bool bCanBeEndedByOtherInstance;
+	uint8 bCanBeEndedByOtherInstance:1;
 
 	void SetActivationConfirmed();
 
@@ -173,7 +173,7 @@ struct GAMEPLAYABILITIES_API FGameplayAbilityActivationInfo
 	/** Called on the server to set the key used by the client to predict this ability */
 	void ServerSetActivationPredictionKey(FPredictionKey PredictionKey);
 
-	FPredictionKey GetActivationPredictionKey() { return PredictionKeyWhenActivated; }
+	const FPredictionKey& GetActivationPredictionKey() const { return PredictionKeyWhenActivated; }
 
 private:
 
@@ -194,13 +194,13 @@ struct GAMEPLAYABILITIES_API FGameplayAbilitySpec : public FFastArraySerializerI
 	GENERATED_USTRUCT_BODY()
 
 	FGameplayAbilitySpec()
-	: Ability(nullptr), Level(1), InputID(INDEX_NONE), SourceObject(nullptr), InputPressed(false), ActiveCount(0), RemoveAfterActivation(false), PendingRemove(false)
+	: Ability(nullptr), Level(1), InputID(INDEX_NONE), SourceObject(nullptr), ActiveCount(0), InputPressed(false), RemoveAfterActivation(false), PendingRemove(false)
 	{
 		
 	}
 
 	FGameplayAbilitySpec(UGameplayAbility* InAbility, int32 InLevel=1, int32 InInputID=INDEX_NONE, UObject* InSourceObject=nullptr)
-		: Ability(InAbility), Level(InLevel), InputID(InInputID), SourceObject(InSourceObject), InputPressed(false), ActiveCount(0), RemoveAfterActivation(false),  PendingRemove(false)
+		: Ability(InAbility), Level(InLevel), InputID(InInputID), SourceObject(InSourceObject), ActiveCount(0), InputPressed(false), RemoveAfterActivation(false), PendingRemove(false)
 	{
 		Handle.GenerateNewHandle();
 	}
@@ -227,17 +227,21 @@ struct GAMEPLAYABILITIES_API FGameplayAbilitySpec : public FFastArraySerializerI
 	UPROPERTY()
 	UObject* SourceObject;
 
-	/** Is input currently pressed. Set to false when input is released */
-	UPROPERTY(NotReplicated)
-	bool	InputPressed;
-
 	/** A count of the number of times this ability has been activated minus the number of times it has been ended. For instanced abilities this will be the number of currently active instances. Can't replicate until prediction accurately handles this.*/
 	UPROPERTY(NotReplicated)
 	uint8	ActiveCount;
 
+	/** Is input currently pressed. Set to false when input is released */
+	UPROPERTY(NotReplicated)
+	uint8 InputPressed:1;
+
 	/** If true, this ability should be removed as soon as it finishes executing */
 	UPROPERTY(NotReplicated)
-	bool	RemoveAfterActivation;
+	uint8 RemoveAfterActivation:1;
+
+	/** Pending removal due to scope lock */
+	UPROPERTY(NotReplicated)
+	uint8 PendingRemove:1;
 
 	/** Activation state of this ability. This is not replicated since it needs to be overwritten locally on clients during prediction. */
 	UPROPERTY(NotReplicated)
@@ -255,10 +259,6 @@ struct GAMEPLAYABILITIES_API FGameplayAbilitySpec : public FFastArraySerializerI
 	UPROPERTY(NotReplicated)
 	FActiveGameplayEffectHandle	GameplayEffectHandle;
 
-	/** Pending removal due to scope lock */
-	UPROPERTY(NotReplicated)
-	bool	PendingRemove;
-
 	/** Returns the primary instance, used for instance once abilities */
 	UGameplayAbility* GetPrimaryInstance() const;
 
diff --git a/Engine/Source/Runtime/GameplayAbilities/Public/GameplayCueNotify_Actor.h b/Engine/Source/Runtime/GameplayAbilities/Public/GameplayCueNotify_Actor.h
--- a/Engine/Source/Runtime/GameplayAbilities/Public/GameplayCueNotify_Actor.h
+++ b/Engine/Source/Runtime/GameplayAbilities/Public/GameplayCueNotify_Actor.h
@@ -37,7 +37,7 @@ class GAMEPLAYABILITIES_API AGameplayCueNotify_Actor : public AActor
 
 	virtual void Serialize(FArchive& Ar) override;
 
-	virtual void HandleGameplayCue(AActor* MyTarget, EGameplayCueEvent::Type EventType, FGameplayCueParameters Parameters);
+	virtual void HandleGameplayCue(AActor* MyTarget, EGameplayCueEvent::Type EventType, const FGameplayCueParameters& Parameters);
 
 	virtual void GameplayCueFinishedCallback();
 
@@ -52,19 +52,19 @@ class GAMEPLAYABILITIES_API AGameplayCueNotify_Actor : public AActor
 
 	/** Generic Event Graph event that will get called for every event type */
 	UFUNCTION(BlueprintImplementableEvent, Category = "GameplayCueNotify", DisplayName = "HandleGameplayCue")
-	void K2_HandleGameplayCue(AActor* MyTarget, EGameplayCueEvent::Type EventType, FGameplayCueParameters Parameters);
+	void K2_HandleGameplayCue(AActor* MyTarget, EGameplayCueEvent::Type EventType, const FGameplayCueParameters& Parameters);
 
 	UFUNCTION(BlueprintNativeEvent, Category = "GameplayCueNotify")
-	bool OnExecute(AActor* MyTarget, FGameplayCueParameters Parameters);
+	bool OnExecute(AActor* MyTarget, const FGameplayCueParameters& Parameters);
 
 	UFUNCTION(BlueprintNativeEvent, Category = "GameplayCueNotify")
-	bool OnActive(AActor* MyTarget, FGameplayCueParameters Parameters);
+	bool OnActive(AActor* MyTarget, const FGameplayCueParameters& Parameters);
 
 	UFUNCTION(BlueprintNativeEvent, Category = "GameplayCueNotify")
-	bool WhileActive(AActor* MyTarget, FGameplayCueParameters Parameters);
+	bool WhileActive(AActor* MyTarget, const FGameplayCueParameters& Parameters);
 
 	UFUNCTION(BlueprintNativeEvent, Category = "GameplayCueNotify")
-	bool OnRemove(AActor* MyTarget, FGameplayCueParameters Parameters);
+	bool OnRemove(AActor* MyTarget, const FGameplayCueParameters& Parameters);
 
 	UPROPERTY(EditDefaultsOnly, Category=GameplayCue, meta=(Categories="GameplayCue"))
 	FGameplayTag	GameplayCueTag;
diff --git a/Engine/Source/Runtime/GameplayAbilities/Public/GameplayCueNotify_HitImpact.h b/Engine/Source/Runtime/GameplayAbilities/Public/GameplayCueNotify_HitImpact.h
--- a/Engine/Source/Runtime/GameplayAbilities/Public/GameplayCueNotify_HitImpact.h
+++ b/Engine/Source/Runtime/GameplayAbilities/Public/GameplayCueNotify_HitImpact.h
@@ -20,7 +20,7 @@ class GAMEPLAYABILITIES_API UGameplayCueNotify_HitImpact : public UGameplayCueNo
 	/** Does this GameplayCueNotify handle this type of GameplayCueEvent? */
 	virtual bool HandlesEvent(EGameplayCueEvent::Type EventType) const override;
 	
-	virtual void HandleGameplayCue(AActor* MyTarget, EGameplayCueEvent::Type EventType, FGameplayCueParameters Parameters) override;
+	virtual void HandleGameplayCue(AActor* MyTarget, EGameplayCueEvent::Type EventType, const FGameplayCueParameters& Parameters) override;
 
 	UPROPERTY(EditAnywhere, BlueprintReadWrite, Category = GameplayCue)
 	USoundBase* Sound;
diff --git a/Engine/Source/Runtime/GameplayAbilities/Public/GameplayCueNotify_Static.h b/Engine/Source/Runtime/GameplayAbilities/Public/GameplayCueNotify_Static.h
--- a/Engine/Source/Runtime/GameplayAbilities/Public/GameplayCueNotify_Static.h
+++ b/Engine/Source/Runtime/GameplayAbilities/Public/GameplayCueNotify_Static.h
@@ -33,7 +33,7 @@ class GAMEPLAYABILITIES_API UGameplayCueNotify_Static : public UObject
 
 	virtual void Serialize(FArchive& Ar) override;
 
-	virtual void HandleGameplayCue(AActor* MyTarget, EGameplayCueEvent::Type EventType, FGameplayCueParameters Parameters);
+	virtual void HandleGameplayCue(AActor* MyTarget, EGameplayCueEvent::Type EventType, const FGameplayCueParameters& Parameters);
 
 	UWorld* GetWorld() const override;
 
@@ -43,19 +43,19 @@ class GAMEPLAYABILITIES_API UGameplayCueNotify_Static : public UObject
 
 	/** Generic Event Graph event that will get called for every event type */
 	UFUNCTION(BlueprintImplementableEvent, Category = "GameplayCueNotify", DisplayName = "HandleGameplayCue")
-	void K2_HandleGameplayCue(AActor* MyTarget, EGameplayCueEvent::Type EventType, FGameplayCueParameters Parameters) const;
+	void K2_HandleGameplayCue(AActor* MyTarget, EGameplayCueEvent::Type EventType, const FGameplayCueParameters& Parameters) const;
 
 	UFUNCTION(BlueprintNativeEvent, BlueprintPure, Category = "GameplayCueNotify")
-	bool OnExecute(AActor* MyTarget, FGameplayCueParameters Parameters) const;
+	bool OnExecute(AActor* MyTarget, const FGameplayCueParameters& Parameters) const;
 
 	UFUNCTION(BlueprintNativeEvent, BlueprintPure, Category = "GameplayCueNotify")
-	bool OnActive(AActor* MyTarget, FGameplayCueParameters Parameters) const;
+	bool OnActive(AActor* MyTarget, const FGameplayCueParameters& Parameters) const;
 
 	UFUNCTION(BlueprintNativeEvent, BlueprintPure, Category = "GameplayCueNotify")
-	bool WhileActive(AActor* MyTarget, FGameplayCueParameters Parameters) const;
+	bool WhileActive(AActor* MyTarget, const FGameplayCueParameters& Parameters) const;
 
 	UFUNCTION(BlueprintNativeEvent, BlueprintPure, Category = "GameplayCueNotify")
-	bool OnRemove(AActor* MyTarget, FGameplayCueParameters Parameters) const;
+	bool OnRemove(AActor* MyTarget, const FGameplayCueParameters& Parameters) const;
 
 	UPROPERTY(EditDefaultsOnly, Category = GameplayCue, meta=(Categories="GameplayCue"))
 	FGameplayTag	GameplayCueTag;
diff --git a/Engine/Source/Runtime/GameplayAbilities/Public/GameplayCueSet.h b/Engine/Source/Runtime/GameplayAbilities/Public/GameplayCueSet.h
--- a/Engine/Source/Runtime/GameplayAbilities/Public/GameplayCueSet.h
+++ b/Engine/Source/Runtime/GameplayAbilities/Public/GameplayCueSet.h
@@ -48,7 +48,7 @@ class GAMEPLAYABILITIES_API UGameplayCueSet : public UDataAsset
 	GENERATED_UCLASS_BODY()
 
 	/** Handles the cue event by spawning the cue actor. Returns true if the event was handled. */
-	virtual bool HandleGameplayCue(AActor* TargetActor, FGameplayTag GameplayCueTag, EGameplayCueEvent::Type EventType, FGameplayCueParameters Parameters);
+	virtual bool HandleGameplayCue(AActor* TargetActor, FGameplayTag GameplayCueTag, EGameplayCueEvent::Type EventType, const FGameplayCueParameters& Parameters);
 
 	/** Adds a list of cues to the set */
 	virtual void AddCues(const TArray<FGameplayCueReferencePair>& CuesToAdd);
diff --git a/Engine/Source/Runtime/GameplayAbilities/Public/GameplayEffectTypes.h b/Engine/Source/Runtime/GameplayAbilities/Public/GameplayEffectTypes.h
--- a/Engine/Source/Runtime/GameplayAbilities/Public/GameplayEffectTypes.h
+++ b/Engine/Source/Runtime/GameplayAbilities/Public/GameplayEffectTypes.h
@@ -332,11 +332,11 @@ struct GAMEPLAYABILITIES_API FGameplayEffectContext
 		return SourceObject.Get();
 	}
 
-	virtual void AddActors(TArray<TWeakObjectPtr<AActor>> InActor, bool bReset = false);
+	virtual void AddActors(const TArray<TWeakObjectPtr<AActor>>& IActor, bool bReset = false);
 
-	virtual void AddHitResult(const FHitResult InHitResult, bool bReset = false);
+	virtual void AddHitResult(const FHitResult& InHitResult, bool bReset = false);
 
-	virtual const TArray<TWeakObjectPtr<AActor>> GetActors() const
+	virtual const TArray<TWeakObjectPtr<AActor>>& GetActors() const
 	{
 		return Actors;
 	}
@@ -584,7 +584,7 @@ struct FGameplayEffectContextHandle
 		return false;
 	}
 
-	void AddActors(TArray<TWeakObjectPtr<AActor>> InActors, bool bReset = false)
+	void AddActors(const TArray<TWeakObjectPtr<AActor>>& InActors, bool bReset = false)
 	{
 		if (IsValid())
 		{
@@ -592,7 +592,7 @@ struct FGameplayEffectContextHandle
 		}
 	}
 
-	void AddHitResult(const FHitResult InHitResult, bool bReset = false)
+	void AddHitResult(const FHitResult& InHitResult, bool bReset = false)
 	{
 		if (IsValid())
 		{
@@ -600,7 +600,7 @@ struct FGameplayEffectContextHandle
 		}
 	}
 
-	TArray<TWeakObjectPtr<AActor>> GetActors()
+	const TArray<TWeakObjectPtr<AActor>> GetActors()
 	{
 		return Data->GetActors();
 	}
@@ -611,7 +611,7 @@ struct FGameplayEffectContextHandle
 		{
 			return Data->GetHitResult();
 		}
-		return NULL;
+		return nullptr;
 	}
 
 	void AddOrigin(FVector InOrigin)
diff --git a/Engine/Source/Runtime/Landscape/Private/Landscape.cpp b/Engine/Source/Runtime/Landscape/Private/Landscape.cpp
--- a/Engine/Source/Runtime/Landscape/Private/Landscape.cpp
+++ b/Engine/Source/Runtime/Landscape/Private/Landscape.cpp
@@ -1971,7 +1971,7 @@ void ULandscapeInfo::UpdateComponentLayerWhitelist()
 // This handles legacy behavior of landscapes created under world composition
 // We adjust landscape section offsets and set a flag inside landscape that sections offset are static and should never be touched again
 //
-void AdjustLandscapeSectionOffsets(UWorld* InWorld, const TArray<ALandscapeProxy*> InLandscapeList)
+void AdjustLandscapeSectionOffsets(UWorld* InWorld, const TArray<ALandscapeProxy*>& InLandscapeList)
 {
 	// We interested only in registered actors
 	TArray<ALandscapeProxy*> RegisteredLandscapeList = InLandscapeList.FilterByPredicate([](ALandscapeProxy* Proxy) {
diff --git a/Engine/Source/Runtime/Landscape/Private/LandscapeGrass.cpp b/Engine/Source/Runtime/Landscape/Private/LandscapeGrass.cpp
--- a/Engine/Source/Runtime/Landscape/Private/LandscapeGrass.cpp
+++ b/Engine/Source/Runtime/Landscape/Private/LandscapeGrass.cpp
@@ -374,7 +374,7 @@ public:
 	,	LandscapeProxy(InLandscapeProxy)
 	,	ComponentSizeQuads(InLandscapeProxy->ComponentSizeQuads)
 	,	ComponentSizeVerts(ComponentSizeQuads + 1)
-	,	GrassTypes(InGrassTypes)
+	,	GrassTypes(MoveTemp(InGrassTypes))
 	,	RenderTargetTexture(nullptr)
 	{
 		check(InLandscapeComponents.Num());
diff --git a/Engine/Source/Runtime/Renderer/Private/SceneVisibility.cpp b/Engine/Source/Runtime/Renderer/Private/SceneVisibility.cpp
--- a/Engine/Source/Runtime/Renderer/Private/SceneVisibility.cpp
+++ b/Engine/Source/Runtime/Renderer/Private/SceneVisibility.cpp
@@ -2446,10 +2446,7 @@ void FSceneRenderer::ComputeViewVisibility(FRHICommandListImmediate& RHICmdList)
 			// now iterate through turn off visibility if hidden by LOD
 			for(FSceneSetBitIterator BitIt(PrimitiveHiddenByLODMap); BitIt; ++BitIt)
 			{
-				if(PrimitiveHiddenByLODMap.AccessCorrespondingBit(BitIt))
-				{
-					View.PrimitiveVisibilityMap.AccessCorrespondingBit(BitIt) = false;
-				}
+				View.PrimitiveVisibilityMap.AccessCorrespondingBit(BitIt) = false;
 			}
 		}
 
@@ -2462,7 +2459,7 @@ void FSceneRenderer::ComputeViewVisibility(FRHICommandListImmediate& RHICmdList)
 
 			if (bLODSceneTreeActive)
 			{
-				for (FSceneBitArray::FIterator BitIt(PrimitiveHiddenByLODMap); BitIt; ++BitIt)
+				for (FSceneSetBitIterator BitIt(PrimitiveHiddenByLODMap); BitIt; ++BitIt)
 				{
 					View.PrimitiveViewRelevanceMap[BitIt.GetIndex()].bInitializedThisFrame = true;
 				}
diff --git a/Engine/Source/Runtime/Windows/XAudio2/Private/XAudio2Device.cpp b/Engine/Source/Runtime/Windows/XAudio2/Private/XAudio2Device.cpp
--- a/Engine/Source/Runtime/Windows/XAudio2/Private/XAudio2Device.cpp
+++ b/Engine/Source/Runtime/Windows/XAudio2/Private/XAudio2Device.cpp
@@ -57,6 +57,10 @@ XAUDIO2_DEVICE_DETAILS FXAudioDeviceProperties::DeviceDetails;
 
 #define DEBUG_XAUDIO2 0
 
+void FXAudio2Device::GetAudioDeviceList(TArray<FString>& OutAudioDeviceNames) const
+{
+	DeviceProperties->GetAudioDeviceList(OutAudioDeviceNames);
+}
 
 bool FXAudio2Device::InitializeHardware()
 {
@@ -122,13 +126,36 @@ bool FXAudio2Device::InitializeHardware()
 		return( false );		
 	}
 
-	// Get the details of the default device 0
-	if( !ValidateAPICall(TEXT("GetDeviceDetails"),
-		DeviceProperties->XAudio2->GetDeviceDetails(0, &FXAudioDeviceProperties::DeviceDetails)))
+	// Initialize the audio device index to 0, which is the windows default device index
+	int32 DeviceIndex = 0;
+
+	FString WindowsAudioDeviceName;
+	GConfig->GetString(TEXT("/Script/WindowsTargetPlatform.WindowsTargetSettings"), TEXT("AudioDevice"), WindowsAudioDeviceName, GEngineIni);
+
+	// If the user specified one, try to load the user audio device
+	if (WindowsAudioDeviceName.Len() > 0)
 	{
-		UE_LOG(LogInit, Log, TEXT( "Failed to get DeviceDetails for XAudio2" ) );
+		TArray<FString> AudioDevices;
+		DeviceProperties->GetAudioDeviceList(AudioDevices);
+
+		for (int32 i = 0; i < AudioDevices.Num(); ++i)
+		{
+			// Find the device index corresponding to the given device name
+			if (AudioDevices[i] == WindowsAudioDeviceName)
+			{
+				DeviceIndex = i;
+				break;
+			}
+		}
+	}
+
+	// Get the details of the desired device index (0 is default)
+	if (!ValidateAPICall(TEXT("GetDeviceDetails"),
+		DeviceProperties->XAudio2->GetDeviceDetails(DeviceIndex, &FXAudioDeviceProperties::DeviceDetails)))
+	{
+		UE_LOG(LogInit, Log, TEXT("Failed to get DeviceDetails for XAudio2"));
 		DeviceProperties->XAudio2 = nullptr;
-		return( false );
+		return(false);
 	}
 
 #if DEBUG_XAUDIO2
@@ -164,7 +191,7 @@ bool FXAudio2Device::InitializeHardware()
 
 	// Create the final output voice with either 2 or 6 channels
 	if (!ValidateAPICall(TEXT("CreateMasteringVoice"), 
-		DeviceProperties->XAudio2->CreateMasteringVoice(&DeviceProperties->MasteringVoice, FXAudioDeviceProperties::NumSpeakers, SampleRate, 0, 0, NULL)))
+		DeviceProperties->XAudio2->CreateMasteringVoice(&DeviceProperties->MasteringVoice, FXAudioDeviceProperties::NumSpeakers, SampleRate, 0, DeviceIndex, nullptr)))
 	{
 		UE_LOG(LogInit, Warning, TEXT( "Failed to create the mastering voice for XAudio2" ) );
 		DeviceProperties->XAudio2 = nullptr;
diff --git a/Engine/Source/Runtime/Windows/XAudio2/Private/XAudio2Support.h b/Engine/Source/Runtime/Windows/XAudio2/Private/XAudio2Support.h
--- a/Engine/Source/Runtime/Windows/XAudio2/Private/XAudio2Support.h
+++ b/Engine/Source/Runtime/Windows/XAudio2/Private/XAudio2Support.h
@@ -639,6 +639,25 @@ struct FXAudioDeviceProperties
 #endif
 	}
 
+	void GetAudioDeviceList(TArray<FString>& OutAudioDeviceList) const
+	{
+		OutAudioDeviceList.Reset();
+
+#if XAUDIO_SUPPORTS_DEVICE_DETAILS
+		uint32 NumDevices = 0;
+		if (XAudio2)
+		{
+			XAudio2->GetDeviceCount(&NumDevices);
+			for (uint32 i = 0; i < NumDevices; ++i)
+			{
+				XAUDIO2_DEVICE_DETAILS Details;
+				XAudio2->GetDeviceDetails(i, &Details);
+				OutAudioDeviceList.Add(FString(Details.DisplayName));
+			}
+		}
+#endif
+	}
+
 	/** Returns either a new IXAudio2SourceVoice or a recycled IXAudio2SourceVoice according to the sound format and max channel count in the voice's effect chain*/
 	void GetFreeSourceVoice(IXAudio2SourceVoice** Voice, const FPCMBufferInfo& BufferInfo, const XAUDIO2_EFFECT_CHAIN* EffectChain = nullptr, int32 MaxEffectChainChannels = 0)
 	{
diff --git a/Engine/Source/Runtime/Windows/XAudio2/Public/XAudio2Device.h b/Engine/Source/Runtime/Windows/XAudio2/Public/XAudio2Device.h
--- a/Engine/Source/Runtime/Windows/XAudio2/Public/XAudio2Device.h
+++ b/Engine/Source/Runtime/Windows/XAudio2/Public/XAudio2Device.h
@@ -83,6 +83,8 @@ class FSpatializationHelper;
  */
 class FXAudio2Device : public FAudioDevice
 {
+	virtual void GetAudioDeviceList(TArray<FString>& OutAudioDeviceNames) const override;
+
 	/** Starts up any platform specific hardware/APIs */
 	virtual bool InitializeHardware() override;
 
