Copying //UE4/Orion-Staging to //UE4/Main (Origin: //Orion/Dev-General [REDACTED_HANDLE])

diff --git a/Engine/Source/Runtime/Engine/Public/ParticleEmitterInstances.h b/Engine/Source/Runtime/Engine/Public/ParticleEmitterInstances.h
--- a/Engine/Source/Runtime/Engine/Public/ParticleEmitterInstances.h
+++ b/Engine/Source/Runtime/Engine/Public/ParticleEmitterInstances.h
@@ -215,6 +215,8 @@ public:
 	FMatrix EmitterToSimulation;
 	/** Transform from simulation space to world space.					*/
 	FMatrix SimulationToWorld;
+	/** Component can disable Tick and Rendering of this emitter. */
+	uint32 bEnabled : 1;
 	/** If true, kill this emitter instance when it is deactivated.		*/
 	uint32 bKillOnDeactivate:1;
 	/** if true, kill this emitter instance when it has completed.		*/
@@ -713,6 +715,59 @@ protected:
 
 };
 
+#if STATS
+struct FScopeCycleCounterEmitter : public FCycleCounter
+{
+	/**
+	* Constructor, starts timing
+	*/
+	template<class T>
+	FORCEINLINE_STATS FScopeCycleCounterEmitter(const T *Object)
+	{
+		if (Object)
+		{
+			TStatId StatId = Object->SpriteTemplate->GetStatID();
+			if (FThreadStats::IsCollectingData(StatId))
+			{
+				Start(StatId);
+			}
+		}
+	}
+	/**
+	* Constructor, starts timing with an alternate enable stat to use high performance disable for only SOME UObject stats
+	*/
+	template<class T>
+	FORCEINLINE_STATS FScopeCycleCounterEmitter(const T *Object, TStatId OtherStat)
+	{
+		if (FThreadStats::IsCollectingData(OtherStat) && Object)
+		{
+			TStatId StatId = Object->SpriteTemplate->GetStatID();
+			if (!StatId.IsNone())
+			{
+				Start(StatId);
+			}
+		}
+	}
+	/**
+	* Updates the stat with the time spent
+	*/
+	FORCEINLINE_STATS ~FScopeCycleCounterEmitter()
+	{
+		Stop();
+	}
+};
+#else
+struct FScopeCycleCounterEmitter
+{
+	FORCEINLINE_STATS FScopeCycleCounterEmitter(const FParticleEmitterInstance *Object)
+	{
+	}
+	FORCEINLINE_STATS FScopeCycleCounterEmitter(const FParticleEmitterInstance *Object, TStatId OtherStat)
+	{
+	}
+};
+#endif
+
 /*-----------------------------------------------------------------------------
 	ParticleSpriteEmitterInstance
 -----------------------------------------------------------------------------*/
@@ -1145,21 +1200,23 @@ struct FParticleTrailsEmitterInstance_Base : public FParticleEmitterInstance
 
 	void CheckIndices(int32 TrailIdx)
 	{
-#if !(UE_BUILD_SHIPPING || UE_BUILD_TEST)
-		if (CurrentEndIndices[TrailIdx] != INDEX_NONE)
-		{
-			DECLARE_PARTICLE_PTR(EndParticle, ParticleData + ParticleStride * CurrentEndIndices[TrailIdx]);
-			FRibbonTypeDataPayload* EndTrailData = ((FRibbonTypeDataPayload*)((uint8*)EndParticle + TypeDataOffset));
-			check(TRAIL_EMITTER_IS_END(EndTrailData->Flags));
-		}
-
-		if (CurrentStartIndices[TrailIdx] != INDEX_NONE)
-		{
-			DECLARE_PARTICLE_PTR(StartParticle, ParticleData + ParticleStride * CurrentStartIndices[TrailIdx]);
-			FRibbonTypeDataPayload* StartTrailData = ((FRibbonTypeDataPayload*)((uint8*)StartParticle + TypeDataOffset));
-			check(TRAIL_EMITTER_IS_START(StartTrailData->Flags));
-		}
-#endif
+		//Disabled until bug with this optimization is found. Usage of the optimization is alredy disabled but these check were still being hit.
+
+// #if !(UE_BUILD_SHIPPING || UE_BUILD_TEST)
+// 		if (CurrentEndIndices[TrailIdx] != INDEX_NONE)
+// 		{
+// 			DECLARE_PARTICLE_PTR(EndParticle, ParticleData + ParticleStride * CurrentEndIndices[TrailIdx]);
+// 			FRibbonTypeDataPayload* EndTrailData = ((FRibbonTypeDataPayload*)((uint8*)EndParticle + TypeDataOffset));
+// 			check(TRAIL_EMITTER_IS_END(EndTrailData->Flags));
+// 		}
+// 
+// 		if (CurrentStartIndices[TrailIdx] != INDEX_NONE)
+// 		{
+// 			DECLARE_PARTICLE_PTR(StartParticle, ParticleData + ParticleStride * CurrentStartIndices[TrailIdx]);
+// 			FRibbonTypeDataPayload* StartTrailData = ((FRibbonTypeDataPayload*)((uint8*)StartParticle + TypeDataOffset));
+// 			check(TRAIL_EMITTER_IS_START(StartTrailData->Flags));
+// 		}
+// #endif
 	}
 
 
diff --git a/Engine/Source/Runtime/Engine/Public/ParticleHelper.h b/Engine/Source/Runtime/Engine/Public/ParticleHelper.h
--- a/Engine/Source/Runtime/Engine/Public/ParticleHelper.h
+++ b/Engine/Source/Runtime/Engine/Public/ParticleHelper.h
@@ -1422,6 +1422,8 @@ struct FDynamicEmitterDataBase
 	/** Returns the current macro uv override. Specialized by FGPUSpriteDynamicEmitterData  */
 	virtual const FMacroUVOverride& GetMacroUVOverride() const { return GetSource().MacroUVOverride; }
 
+	/** Stat id of this object, 0 if nobody asked for it yet */
+	mutable TStatId StatID;
 	/** true if this emitter is currently selected */
 	uint32	bSelected:1;
 	/** true if this emitter has valid rendering data */
@@ -2708,3 +2710,20 @@ FORCEINLINE FVector2D GetParticleSizeWithUVFlipInSign(const FBaseParticle& Parti
 		Particle.BaseSize.X >= 0.0f ? ScaledSize.X : -ScaledSize.X,
 		Particle.BaseSize.Y >= 0.0f ? ScaledSize.Y : -ScaledSize.Y);
 }
+
+
+/** A level of significance for a particle system. Used by game code to enable/disable emitters progressively as they move away from the camera or are occluded/off screen. */
+UENUM()
+enum class EParticleSignificanceLevel : uint8
+{
+	/** Low significance emitter. Culled first. */
+	Low,
+	/** Medium significance emitter. */
+	Medium,
+	/** High significance emitter. Culled last. */
+	High,
+	/** Critical emitter. Never culled. */
+	Critical,
+
+	Num UMETA(Hidden),
+};
\ No newline at end of file
diff --git a/Engine/Source/Runtime/Engine/Public/PlatformFeatures.h b/Engine/Source/Runtime/Engine/Public/PlatformFeatures.h
--- a/Engine/Source/Runtime/Engine/Public/PlatformFeatures.h
+++ b/Engine/Source/Runtime/Engine/Public/PlatformFeatures.h
@@ -37,5 +37,7 @@ public:
 
 	ENGINE_API virtual class ISaveGameSystem* GetSaveGameSystem();
 
-	ENGINE_API virtual class IDVRStreamingSystem *GetStreamingSystem();
+	ENGINE_API virtual class IDVRStreamingSystem* GetStreamingSystem();
+
+	ENGINE_API virtual TSharedPtr<const class FJsonObject> GetTitleSettings();
 };
diff --git a/Engine/Source/Runtime/Engine/Public/PrimitiveSceneProxy.h b/Engine/Source/Runtime/Engine/Public/PrimitiveSceneProxy.h
--- a/Engine/Source/Runtime/Engine/Public/PrimitiveSceneProxy.h
+++ b/Engine/Source/Runtime/Engine/Public/PrimitiveSceneProxy.h
@@ -171,6 +171,9 @@ public:
 	/** Gathers shadow shapes from this proxy. */
 	virtual void GetShadowShapes(TArray<FCapsuleShape>& CapsuleShapes) const {}
 
+	/** Returns a list of valid bone indices referenced by this proxy's shadow shapes. */
+	virtual void GetShadowShapeBoneIndices(TArray<uint16>& BoneIndices) const {}
+
 	/** 
 	 * Gathers the primitive's dynamic mesh elements.  This will only be called if GetViewRelevance declares dynamic relevance.
 	 * This is called from the rendering thread for each set of views that might be rendered.  
diff --git a/Engine/Source/Runtime/Engine/Public/Scalability.h b/Engine/Source/Runtime/Engine/Public/Scalability.h
--- a/Engine/Source/Runtime/Engine/Public/Scalability.h
+++ b/Engine/Source/Runtime/Engine/Public/Scalability.h
@@ -22,7 +22,12 @@ namespace Scalability
 		int32 TextureQuality;
 		int32 EffectsQuality;
 
+		float CPUBenchmarkResults;
+		float GPUBenchmarkResults;
+
 		FQualityLevels()
+			: CPUBenchmarkResults(-1.0f)
+			, GPUBenchmarkResults(-1.0f)
 		{
 			SetDefaults();
 		}
@@ -80,7 +85,7 @@ namespace Scalability
 	ENGINE_API void ProcessCommand(const TCHAR* Cmd, FOutputDevice& Ar);
 
 	/** Minimum single axis scale for render resolution */
-	static const int32 MinResolutionScale = 50;
+	static const int32 MinResolutionScale = 10;
 
 	/** Maximum single axis scale for render resolution */
 	static const int32 MaxResolutionScale = 100;
diff --git a/Engine/Source/Runtime/Engine/Public/SceneUtils.h b/Engine/Source/Runtime/Engine/Public/SceneUtils.h
--- a/Engine/Source/Runtime/Engine/Public/SceneUtils.h
+++ b/Engine/Source/Runtime/Engine/Public/SceneUtils.h
@@ -10,12 +10,8 @@
 // Each event type will be displayed using the defined color
 #pragma once
 
-// Disable draw mesh events for final builds
-#if !(UE_BUILD_SHIPPING || UE_BUILD_TEST) && PLATFORM_SUPPORTS_DRAW_MESH_EVENTS
-#define WANTS_DRAW_MESH_EVENTS 1
-#else
-#define WANTS_DRAW_MESH_EVENTS 0
-#endif
+// Note:  WITH_PROFILEGPU should be 0 for final builds
+#define WANTS_DRAW_MESH_EVENTS (WITH_PROFILEGPU && PLATFORM_SUPPORTS_DRAW_MESH_EVENTS)
 
 #if WANTS_DRAW_MESH_EVENTS
 
diff --git a/Engine/Source/Runtime/Engine/Public/SkeletalMeshTypes.h b/Engine/Source/Runtime/Engine/Public/SkeletalMeshTypes.h
--- a/Engine/Source/Runtime/Engine/Public/SkeletalMeshTypes.h
+++ b/Engine/Source/Runtime/Engine/Public/SkeletalMeshTypes.h
@@ -1753,6 +1753,7 @@ public:
 	
 	virtual bool HasDistanceFieldRepresentation() const override;
 	virtual void GetShadowShapes(TArray<FCapsuleShape>& CapsuleShapes) const override;
+	virtual void GetShadowShapeBoneIndices(TArray<uint16>& BoneIndices) const override;
 
 	/**
 	 * Returns the world transform to use for drawing.
diff --git a/Engine/Source/Runtime/Engine/Public/TextureResource.h b/Engine/Source/Runtime/Engine/Public/TextureResource.h
--- a/Engine/Source/Runtime/Engine/Public/TextureResource.h
+++ b/Engine/Source/Runtime/Engine/Public/TextureResource.h
@@ -320,6 +320,8 @@ private:
 	/** Number of file I/O requests for current request														*/
 	int32					IORequestCount;
 
+	FThreadSafeCounter AsyncReallocateCounter;
+
 #if WITH_EDITORONLY_DATA
 	/** Pending async derived data streaming tasks															*/
 	TIndirectArray<FAsyncStreamDerivedMipTask> PendingAsyncStreamDerivedMipTasks;
diff --git a/Engine/Source/Runtime/Engine/Public/TimerManager.h b/Engine/Source/Runtime/Engine/Public/TimerManager.h
--- a/Engine/Source/Runtime/Engine/Public/TimerManager.h
+++ b/Engine/Source/Runtime/Engine/Public/TimerManager.h
@@ -17,10 +17,13 @@ struct FTimerUnifiedDelegate
 	FTimerDelegate FuncDelegate;
 	/** Holds the dynamic delegate to call. */
 	FTimerDynamicDelegate FuncDynDelegate;
+	/** Holds the tfunction callback to call. */
+	TFunction<void(void)> FuncCallback;
 
 	FTimerUnifiedDelegate() {};
 	FTimerUnifiedDelegate(FTimerDelegate const& D) : FuncDelegate(D) {};
 	FTimerUnifiedDelegate(FTimerDynamicDelegate const& D) : FuncDynDelegate(D) {};
+	FTimerUnifiedDelegate(TFunction<void(void)>&& Callback) : FuncCallback(MoveTemp(Callback)) {}
 	
 	inline void Execute()
 	{
@@ -41,11 +44,15 @@ struct FTimerUnifiedDelegate
 		{
 			FuncDynDelegate.ProcessDelegate<UObject>(nullptr);
 		}
+		else if ( FuncCallback )
+		{
+			FuncCallback();
+		}
 	}
 
 	inline bool IsBound() const
 	{
-		return ( FuncDelegate.IsBound() || FuncDynDelegate.IsBound() );
+		return ( FuncDelegate.IsBound() || FuncDynDelegate.IsBound() || FuncCallback );
 	}
 
 	inline bool IsBoundToObject(void const* Object) const
@@ -66,6 +73,7 @@ struct FTimerUnifiedDelegate
 	{
 		FuncDelegate.Unbind();
 		FuncDynDelegate.Unbind();
+		FuncCallback = nullptr;
 	}
 
 	/** Utility to output info about delegate as a string. */
@@ -189,6 +197,11 @@ public:
 	{
 		InternalSetTimer(InOutHandle, FTimerUnifiedDelegate(), InRate, InbLoop, InFirstDelay);
 	}
+	/** Version that takes a TFunction */
+	FORCEINLINE void SetTimer(FTimerHandle& InOutHandle, TFunction<void(void)>&& Callback, float InRate, bool InbLoop, float InFirstDelay = -1.f )
+	{
+		InternalSetTimer(InOutHandle, FTimerUnifiedDelegate(MoveTemp(Callback)), InRate, InbLoop, InFirstDelay);
+	}
 
 	/**
 	 * Sets a timer to call the given native function on the next tick.
@@ -217,7 +230,11 @@ public:
 	{
 		InternalSetTimerForNextTick(FTimerUnifiedDelegate(InDynDelegate));
 	}
-
+	/** Version that takes a TFunction */
+	FORCEINLINE void SetTimerForNextTick(TFunction<void(void)>&& Callback)
+	{
+		InternalSetTimerForNextTick(FTimerUnifiedDelegate(MoveTemp(Callback)));
+	}
 
 	/**
 	 * Clears a previously set timer, identical to calling SetTimer() with a <= 0.f rate.
diff --git a/Engine/Source/Runtime/Engine/Public/UpdateTextureShaders.h b/Engine/Source/Runtime/Engine/Public/UpdateTextureShaders.h
--- a/Engine/Source/Runtime/Engine/Public/UpdateTextureShaders.h
+++ b/Engine/Source/Runtime/Engine/Public/UpdateTextureShaders.h
@@ -112,3 +112,34 @@ public:
 	FShaderResourceParameter SrcTexture;
 	FShaderResourceParameter DestTexture;
 };
+
+template<uint32 ElementsPerThread>
+class TCopyDataCS : public FGlobalShader
+{
+	DECLARE_SHADER_TYPE(TCopyDataCS, Global);
+public:
+	TCopyDataCS() {}
+	TCopyDataCS(const ShaderMetaType::CompiledShaderInitializerType& Initializer)
+		: FGlobalShader(Initializer)
+	{
+		SrcBuffer.Bind(Initializer.ParameterMap, TEXT("SrcCopyBuffer"), SPF_Mandatory);
+		DestBuffer.Bind(Initializer.ParameterMap, TEXT("DestBuffer"), SPF_Mandatory);		
+	}
+
+	// FShader interface.
+	virtual bool Serialize(FArchive& Ar) override
+	{
+		bool bShaderHasOutdatedParameters = FGlobalShader::Serialize(Ar);
+		Ar << SrcBuffer << DestBuffer;
+		return bShaderHasOutdatedParameters;
+	}
+
+	static bool ShouldCache(EShaderPlatform Platform)
+	{
+		return IsFeatureLevelSupported(Platform, ERHIFeatureLevel::SM5);
+	}
+
+	//protected:
+	FShaderResourceParameter SrcBuffer;
+	FShaderResourceParameter DestBuffer;	
+};
\ No newline at end of file
diff --git a/Engine/Source/Runtime/Engine/Public/UserActivityTracking.h b/Engine/Source/Runtime/Engine/Public/UserActivityTracking.h
new file mode 100644
--- /dev/null
+++ b/Engine/Source/Runtime/Engine/Public/UserActivityTracking.h
@@ -0,0 +1,19 @@
+// Copyright 1998-2015 Epic Games, Inc. All Rights Reserved.
+
+#pragma once
+
+class ENGINE_API FUserActivityTracking : FNoncopyable
+{
+public:
+	DECLARE_MULTICAST_DELEGATE_OneParam(FOnActivityChanged, const FUserActivity&);
+
+	static void SetActivity(const FUserActivity& InUserActivity) { UserActivity = InUserActivity; }
+	static const FUserActivity& GetUserActivity() { return UserActivity; }
+
+	// Called when the UserActivity changes
+	static FOnActivityChanged OnActivityChanged;
+
+private:
+	static FUserActivity UserActivity;
+};
+
diff --git a/Engine/Source/Runtime/GameplayAbilities/Private/AbilitySystemComponent.cpp b/Engine/Source/Runtime/GameplayAbilities/Private/AbilitySystemComponent.cpp
--- a/Engine/Source/Runtime/GameplayAbilities/Private/AbilitySystemComponent.cpp
+++ b/Engine/Source/Runtime/GameplayAbilities/Private/AbilitySystemComponent.cpp
@@ -32,9 +32,8 @@ UAbilitySystemComponent::UAbilitySystemComponent(const FObjectInitializer& Objec
 	bAutoActivate = true;	// Forcing AutoActivate since above we manually force tick enabled.
 							// if we don't have this, UpdateShouldTick() fails to have any effect
 							// because we'll be receiving ticks but bIsActive starts as false
-
-	ActiveGameplayCues.Owner = this;
-
+	
+	CachedIsNetSimulated = false;
 	UserAbilityActivationInhibited = false;
 
 	GenericConfirmInputID = INDEX_NONE;
@@ -127,6 +126,9 @@ void UAbilitySystemComponent::OnRegister()
 {
 	Super::OnRegister();
 
+	// Cached off netrole to avoid constant checking on owning actor
+	CachedIsNetSimulated = IsNetSimulating();
+
 	// Init starting data
 	for (int32 i=0; i < DefaultStartingData.Num(); ++i)
 	{
@@ -139,6 +141,11 @@ void UAbilitySystemComponent::OnRegister()
 
 	ActiveGameplayEffects.RegisterWithOwner(this);
 	ActivatableAbilities.RegisterWithOwner(this);
+	ActiveGameplayCues.Owner = this;
+	ActiveGameplayCues.bMinimalReplication = false;
+	MinimalReplicationGameplayCues.Owner = this;
+	MinimalReplicationGameplayCues.bMinimalReplication = true;
+	MinimalReplicationTags.Owner = this;
 
 	/** Allocate an AbilityActorInfo. Note: this goes through a global function and is a SharedPtr so projects can make their own AbilityActorInfo */
 	AbilityActorInfo = TSharedPtr<FGameplayAbilityActorInfo>(UAbilitySystemGlobals::Get().AllocAbilityActorInfo());
@@ -158,11 +165,6 @@ const FActiveGameplayEffect* UAbilitySystemComponent::GetActiveGameplayEffect(co
 	return ActiveGameplayEffects.GetActiveGameplayEffect(Handle);
 }
 
-bool UAbilitySystemComponent::IsOwnerActorAuthoritative() const
-{
-	return !IsNetSimulating();
-}
-
 bool UAbilitySystemComponent::HasNetworkAuthorityToApplyGameplayEffect(FPredictionKey PredictionKey) const
 {
 	return (IsOwnerActorAuthoritative() || PredictionKey.IsValidForMorePrediction());
@@ -426,36 +428,6 @@ FOnActiveGameplayEffectStackChange* UAbilitySystemComponent::OnGameplayEffectSta
 	return nullptr;
 }
 
-int32 UAbilitySystemComponent::GetNumActiveGameplayEffects() const
-{
-	return ActiveGameplayEffects.GetNumGameplayEffects();
-}
-
-void UAbilitySystemComponent::GetAllActiveGameplayEffectSpecs(TArray<FGameplayEffectSpec>& OutSpecCopies)
-{	
-	ActiveGameplayEffects.GetAllActiveGameplayEffectSpecs(OutSpecCopies);
-}
-
-const FGameplayTagContainer* UAbilitySystemComponent::GetGameplayEffectSourceTagsFromHandle(FActiveGameplayEffectHandle Handle) const
-{
-	return ActiveGameplayEffects.GetGameplayEffectSourceTagsFromHandle(Handle);
-}
-
-const FGameplayTagContainer* UAbilitySystemComponent::GetGameplayEffectTargetTagsFromHandle(FActiveGameplayEffectHandle Handle) const
-{
-	return ActiveGameplayEffects.GetGameplayEffectTargetTagsFromHandle(Handle);
-}
-
-void UAbilitySystemComponent::CaptureAttributeForGameplayEffect(OUT FGameplayEffectAttributeCaptureSpec& OutCaptureSpec)
-{
-	// Verify the capture is happening on an attribute the component actually has a set for; if not, can't capture the value
-	const FGameplayAttribute& AttributeToCapture = OutCaptureSpec.BackingDefinition.AttributeToCapture;
-	if (AttributeToCapture.IsValid() && (AttributeToCapture.IsSystemAttribute() || GetAttributeSubobject(AttributeToCapture.GetAttributeSetClass())))
-	{
-		ActiveGameplayEffects.CaptureAttributeForGameplayEffect(OutCaptureSpec);
-	}
-}
-
 FOnGameplayEffectTagCountChanged& UAbilitySystemComponent::RegisterGameplayTagEvent(FGameplayTag Tag, EGameplayTagEventType::Type EventType)
 {
 	return GameplayTagCountContainer.RegisterGameplayTagEvent(Tag, EventType);
@@ -499,73 +471,6 @@ const FGameplayEffectAttributeCaptureDefinition& UAbilitySystemComponent::GetInc
 
 // ------------------------------------------------------------------------
 
-void UAbilitySystemComponent::GetOwnedGameplayTags(FGameplayTagContainer& TagContainer) const
-{
-	SCOPE_CYCLE_COUNTER(STAT_GameplayEffectsGetOwnedTags);
-	TagContainer.AppendTags(GameplayTagCountContainer.GetExplicitGameplayTags());
-}
-
-bool UAbilitySystemComponent::HasMatchingGameplayTag(FGameplayTag TagToCheck) const
-{
-	SCOPE_CYCLE_COUNTER(STAT_HasMatchingGameplayTag);
-	return GameplayTagCountContainer.HasMatchingGameplayTag(TagToCheck);
-}
-
-bool UAbilitySystemComponent::HasAllMatchingGameplayTags(const FGameplayTagContainer& TagContainer, bool bCountEmptyAsMatch) const
-{
-	SCOPE_CYCLE_COUNTER(STAT_GameplayEffectsHasAllTags);
-	return GameplayTagCountContainer.HasAllMatchingGameplayTags(TagContainer, bCountEmptyAsMatch);
-}
-
-bool UAbilitySystemComponent::HasAnyMatchingGameplayTags(const FGameplayTagContainer& TagContainer, bool bCountEmptyAsMatch) const
-{
-	SCOPE_CYCLE_COUNTER(STAT_GameplayEffectsHasAnyTag);
-	return GameplayTagCountContainer.HasAnyMatchingGameplayTags(TagContainer, bCountEmptyAsMatch);
-}
-
-void UAbilitySystemComponent::AddLooseGameplayTag(const FGameplayTag& GameplayTag, int32 Count)
-{
-	UpdateTagMap(GameplayTag, Count);
-}
-
-void UAbilitySystemComponent::AddLooseGameplayTags(const FGameplayTagContainer& GameplayTags, int32 Count)
-{
-	UpdateTagMap(GameplayTags, Count);
-}
-
-void UAbilitySystemComponent::RemoveLooseGameplayTag(const FGameplayTag& GameplayTag, int32 Count)
-{
-	UpdateTagMap(GameplayTag, -Count);
-}
-
-void UAbilitySystemComponent::RemoveLooseGameplayTags(const FGameplayTagContainer& GameplayTags, int32 Count)
-{
-	UpdateTagMap(GameplayTags, -Count);
-}
-
-// These are functionally redundant but are called by GEs and GameplayCues that add tags that are not 'loose' (but are handled the same way in practice)
-
-void UAbilitySystemComponent::UpdateTagMap(const FGameplayTag& BaseTag, int32 CountDelta)
-{
-	const bool bHadTag = GameplayTagCountContainer.HasMatchingGameplayTag(BaseTag);
-	GameplayTagCountContainer.UpdateTagCount(BaseTag, CountDelta);
-	const bool bHasTag = GameplayTagCountContainer.HasMatchingGameplayTag(BaseTag);
-
-	if (bHadTag != bHasTag)
-	{
-		OnTagUpdated(BaseTag, bHasTag);
-	}
-}
-
-void UAbilitySystemComponent::UpdateTagMap(const FGameplayTagContainer& Container, int32 CountDelta)
-{
-	for (auto TagIt = Container.CreateConstIterator(); TagIt; ++TagIt)
-	{
-		const FGameplayTag& Tag = *TagIt;
-		UpdateTagMap(Tag, CountDelta);
-	}
-}
-
 void UAbilitySystemComponent::ResetTagMap()
 {
 	GameplayTagCountContainer.Reset();
@@ -742,7 +647,7 @@ FActiveGameplayEffectHandle UAbilitySystemComponent::ApplyGameplayEffectSpecToSe
 			// Because PostReplicatedChange will get called from modifying the stack count
 			// (and not PostReplicatedAdd) we won't know which GE was modified.
 			// So instead we need to explicitly RPC the client so it knows the GC needs updating
-			NetMulticast_InvokeGameplayCueAddedAndWhileActive_FromSpec(*OurCopyOfSpec, PredictionKey);
+			UAbilitySystemGlobals::Get().GetGameplayCueManager()->InvokeGameplayCueAddedAndWhileActive_FromSpec(this, *OurCopyOfSpec, PredictionKey);
 		}
 		else
 		{
@@ -1056,14 +961,14 @@ void UAbilitySystemComponent::ExecuteGameplayCue(const FGameplayTag GameplayCueT
 	UAbilitySystemGlobals::Get().GetGameplayCueManager()->InvokeGameplayCueExecuted_WithParams(this, GameplayCueTag, ScopedPredictionKey, GameplayCueParameters);
 }
 
-void UAbilitySystemComponent::AddGameplayCue(const FGameplayTag GameplayCueTag, FGameplayEffectContextHandle EffectContext)
+void UAbilitySystemComponent::AddGameplayCue_Internal(const FGameplayTag GameplayCueTag, const FGameplayEffectContextHandle& EffectContext, FActiveGameplayCueContainer& GameplayCueContainer)
 {
 	if (IsOwnerActorAuthoritative())
 	{
 		bool bWasInList = HasMatchingGameplayTag(GameplayCueTag);
 
 		ForceReplication();
-		ActiveGameplayCues.AddCue(GameplayCueTag, ScopedPredictionKey);
+		GameplayCueContainer.AddCue(GameplayCueTag, ScopedPredictionKey);
 		NetMulticast_InvokeGameplayCueAdded(GameplayCueTag, ScopedPredictionKey, EffectContext);
 
 		if (!bWasInList)
@@ -1074,7 +979,7 @@ void UAbilitySystemComponent::AddGameplayCue(const FGameplayTag GameplayCueTag,
 	}
 	else if (ScopedPredictionKey.IsLocalClientKey())
 	{
-		ActiveGameplayCues.PredictiveAdd(GameplayCueTag, ScopedPredictionKey);
+		GameplayCueContainer.PredictiveAdd(GameplayCueTag, ScopedPredictionKey);
 
 		// Allow for predictive gameplaycue events? Needs more thought
 		InvokeGameplayCueEvent(GameplayCueTag, EGameplayCueEvent::OnActive, EffectContext);
@@ -1082,13 +987,13 @@ void UAbilitySystemComponent::AddGameplayCue(const FGameplayTag GameplayCueTag,
 	}
 }
 
-void UAbilitySystemComponent::RemoveGameplayCue(const FGameplayTag GameplayCueTag)
+void UAbilitySystemComponent::RemoveGameplayCue_Internal(const FGameplayTag GameplayCueTag, FActiveGameplayCueContainer& GameplayCueContainer)
 {
 	if (IsOwnerActorAuthoritative())
 	{
 		bool bWasInList = HasMatchingGameplayTag(GameplayCueTag);
 
-		ActiveGameplayCues.RemoveCue(GameplayCueTag);
+		GameplayCueContainer.RemoveCue(GameplayCueTag);
 
 		if (bWasInList)
 		{
@@ -1099,7 +1004,7 @@ void UAbilitySystemComponent::RemoveGameplayCue(const FGameplayTag GameplayCueTa
 	}
 	else if (ScopedPredictionKey.IsLocalClientKey())
 	{
-		ActiveGameplayCues.PredictiveRemove(GameplayCueTag);
+		GameplayCueContainer.PredictiveRemove(GameplayCueTag);
 	}
 }
 
@@ -1119,6 +1024,8 @@ void UAbilitySystemComponent::NetMulticast_InvokeGameplayCueExecuted_FromSpec_Im
 	}
 }
 
+// -----------
+
 void UAbilitySystemComponent::NetMulticast_InvokeGameplayCueExecuted_Implementation(const FGameplayTag GameplayCueTag, FPredictionKey PredictionKey, FGameplayEffectContextHandle EffectContext)
 {
 	if (IsOwnerActorAuthoritative() || PredictionKey.IsLocalClientKey() == false)
@@ -1127,6 +1034,19 @@ void UAbilitySystemComponent::NetMulticast_InvokeGameplayCueExecuted_Implementat
 	}
 }
 
+void UAbilitySystemComponent::NetMulticast_InvokeGameplayCuesExecuted_Implementation(const FGameplayTagContainer GameplayCueTags, FPredictionKey PredictionKey, FGameplayEffectContextHandle EffectContext)
+{
+	if (IsOwnerActorAuthoritative() || PredictionKey.IsLocalClientKey() == false)
+	{
+		for (const FGameplayTag& GameplayCueTag : GameplayCueTags)
+		{
+			InvokeGameplayCueEvent(GameplayCueTag, EGameplayCueEvent::Executed, EffectContext);
+		}
+	}
+}
+
+// -----------
+
 void UAbilitySystemComponent::NetMulticast_InvokeGameplayCueExecuted_WithParams_Implementation(const FGameplayTag GameplayCueTag, FPredictionKey PredictionKey, FGameplayCueParameters GameplayCueParameters)
 {
 	if (IsOwnerActorAuthoritative() || PredictionKey.IsLocalClientKey() == false)
@@ -1135,6 +1055,19 @@ void UAbilitySystemComponent::NetMulticast_InvokeGameplayCueExecuted_WithParams_
 	}
 }
 
+void UAbilitySystemComponent::NetMulticast_InvokeGameplayCuesExecuted_WithParams_Implementation(const FGameplayTagContainer GameplayCueTags, FPredictionKey PredictionKey, FGameplayCueParameters GameplayCueParameters)
+{
+	if (IsOwnerActorAuthoritative() || PredictionKey.IsLocalClientKey() == false)
+	{
+		for (const FGameplayTag& GameplayCueTag : GameplayCueTags)
+		{
+			InvokeGameplayCueEvent(GameplayCueTag, EGameplayCueEvent::Executed, GameplayCueParameters);
+		}
+	}
+}
+
+// -----------
+
 void UAbilitySystemComponent::NetMulticast_InvokeGameplayCueAdded_Implementation(const FGameplayTag GameplayCueTag, FPredictionKey PredictionKey, FGameplayEffectContextHandle EffectContext)
 {
 	if (IsOwnerActorAuthoritative() || PredictionKey.IsLocalClientKey() == false)
@@ -1143,6 +1076,8 @@ void UAbilitySystemComponent::NetMulticast_InvokeGameplayCueAdded_Implementation
 	}
 }
 
+// -----------
+
 void UAbilitySystemComponent::NetMulticast_InvokeGameplayCueAddedAndWhileActive_FromSpec_Implementation(const FGameplayEffectSpecForRPC& Spec, FPredictionKey PredictionKey)
 {
 	if (IsOwnerActorAuthoritative() || PredictionKey.IsLocalClientKey() == false)
@@ -1152,22 +1087,30 @@ void UAbilitySystemComponent::NetMulticast_InvokeGameplayCueAddedAndWhileActive_
 	}
 }
 
-bool UAbilitySystemComponent::IsGameplayCueActive(const FGameplayTag GameplayCueTag) const
+void UAbilitySystemComponent::NetMulticast_InvokeGameplayCueAddedAndWhileActive_WithParams_Implementation(const FGameplayTag GameplayCueTag, FPredictionKey PredictionKey, FGameplayCueParameters GameplayCueParameters)
 {
-	return HasMatchingGameplayTag(GameplayCueTag);
+	if (IsOwnerActorAuthoritative() || PredictionKey.IsLocalClientKey() == false)
+	{
+		InvokeGameplayCueEvent(GameplayCueTag, EGameplayCueEvent::OnActive, GameplayCueParameters);
+		InvokeGameplayCueEvent(GameplayCueTag, EGameplayCueEvent::WhileActive, GameplayCueParameters);
+	}
 }
 
-// ----------------------------------------------------------------------------------------
-
-void UAbilitySystemComponent::SetBaseAttributeValueFromReplication(float NewValue, FGameplayAttribute Attribute)
+	
+void UAbilitySystemComponent::NetMulticast_InvokeGameplayCuesAddedAndWhileActive_WithParams_Implementation(const FGameplayTagContainer GameplayCueTags, FPredictionKey PredictionKey, FGameplayCueParameters GameplayCueParameters)
 {
-	ActiveGameplayEffects.SetBaseAttributeValueFromReplication(Attribute, NewValue);
+	if (IsOwnerActorAuthoritative() || PredictionKey.IsLocalClientKey() == false)
+	{
+		for (const FGameplayTag& GameplayCueTag : GameplayCueTags)
+		{
+			InvokeGameplayCueEvent(GameplayCueTag, EGameplayCueEvent::OnActive, GameplayCueParameters);
+			InvokeGameplayCueEvent(GameplayCueTag, EGameplayCueEvent::WhileActive, GameplayCueParameters);
+		}
+	}
 }
 
-bool UAbilitySystemComponent::CanApplyAttributeModifiers(const UGameplayEffect *GameplayEffect, float Level, const FGameplayEffectContextHandle& EffectContext)
-{
-	return ActiveGameplayEffects.CanApplyAttributeModifiers(GameplayEffect, Level, EffectContext);
-}
+// ----------------------------------------------------------------------------------------
+
 
 // #deprecated, use FGameplayEffectQuery version
 TArray<float> UAbilitySystemComponent::GetActiveEffectsTimeRemaining(const FActiveGameplayEffectQuery Query) const
@@ -1205,16 +1148,6 @@ TArray<FActiveGameplayEffectHandle> UAbilitySystemComponent::GetActiveEffects(co
 	return ActiveGameplayEffects.GetActiveEffects(Query);
 }
 
-float UAbilitySystemComponent::GetActiveEffectsEndTime(const FGameplayEffectQuery& Query) const
-{
-	return ActiveGameplayEffects.GetActiveEffectsEndTime(Query);
-}
-
-void UAbilitySystemComponent::ModifyActiveEffectStartTime(FActiveGameplayEffectHandle Handle, float StartTimeDiff)
-{
-	ActiveGameplayEffects.ModifyActiveEffectStartTime(Handle, StartTimeDiff);
-}
-
 void UAbilitySystemComponent::RemoveActiveEffectsWithTags(const FGameplayTagContainer Tags)
 {
 	if (IsOwnerActorAuthoritative())
@@ -1270,7 +1203,7 @@ void UAbilitySystemComponent::RemoveActiveEffects(const FGameplayEffectQuery& Qu
 void UAbilitySystemComponent::GetLifetimeReplicatedProps(TArray< FLifetimeProperty > & OutLifetimeProps) const
 {	
 	DOREPLIFETIME(UAbilitySystemComponent, SpawnedAttributes);
-	DOREPLIFETIME(UAbilitySystemComponent, ActiveGameplayEffects);
+	DOREPLIFETIME(UAbilitySystemComponent, ActiveGameplayEffects );
 	DOREPLIFETIME(UAbilitySystemComponent, ActiveGameplayCues);
 	
 	DOREPLIFETIME_CONDITION(UAbilitySystemComponent, ActivatableAbilities, COND_ReplayOrOwner);
@@ -1282,9 +1215,18 @@ void UAbilitySystemComponent::GetLifetimeReplicatedProps(TArray< FLifetimeProper
 	DOREPLIFETIME(UAbilitySystemComponent, ReplicatedPredictionKey);
 	DOREPLIFETIME(UAbilitySystemComponent, RepAnimMontageInfo);
 	
+	DOREPLIFETIME(UAbilitySystemComponent, MinimalReplicationGameplayCues );
+	DOREPLIFETIME_CONDITION(UAbilitySystemComponent, MinimalReplicationTags, COND_Custom);
+	
 	Super::GetLifetimeReplicatedProps(OutLifetimeProps);
 }
 
+void UAbilitySystemComponent::PreReplication(IRepChangedPropertyTracker& ChangedPropertyTracker)
+{	
+	// These only replicate when in minimal replication mode	
+	DOREPLIFETIME_ACTIVE_OVERRIDE(UAbilitySystemComponent, MinimalReplicationTags, bMinimalReplication);
+}
+
 void UAbilitySystemComponent::ForceReplication()
 {
 	AActor *OwningActor = GetOwner();
@@ -1338,11 +1280,6 @@ void UAbilitySystemComponent::PostNetReceive()
 	ActiveGameplayEffects.DecrementLock();
 }
 
-void UAbilitySystemComponent::OnRep_GameplayEffects()
-{
-
-}
-
 void UAbilitySystemComponent::OnRep_PredictionKey()
 {
 	// Every predictive action we've done up to and including the current value of ReplicatedPredictionKey needs to be wiped
@@ -1354,6 +1291,11 @@ bool UAbilitySystemComponent::HasAuthorityOrPredictionKey(const FGameplayAbility
 	return ((ActivationInfo->ActivationMode == EGameplayAbilityActivationMode::Authority) || CanPredict());
 }
 
+void UAbilitySystemComponent::SetMinimalReplication(bool bNewMinimalReplication)
+{
+	bMinimalReplication = true;
+}
+
 // ---------------------------------------------------------------------------------------
 
 void UAbilitySystemComponent::PrintAllGameplayEffects() const
diff --git a/Engine/Source/Runtime/GameplayAbilities/Private/AbilitySystemComponent_Abilities.cpp b/Engine/Source/Runtime/GameplayAbilities/Private/AbilitySystemComponent_Abilities.cpp
--- a/Engine/Source/Runtime/GameplayAbilities/Private/AbilitySystemComponent_Abilities.cpp
+++ b/Engine/Source/Runtime/GameplayAbilities/Private/AbilitySystemComponent_Abilities.cpp
@@ -1280,6 +1280,7 @@ void UAbilitySystemComponent::ClientTryActivateAbility_Implementation(FGameplayA
 
 void UAbilitySystemComponent::InternalServerTryActiveAbility(FGameplayAbilitySpecHandle Handle, bool InputPressed, const FPredictionKey& PredictionKey, const FGameplayEventData* TriggerEventData)
 {
+#if WITH_SERVER_CODE
 #if !(UE_BUILD_SHIPPING || UE_BUILD_TEST)
 	if (DenyClientActivation > 0)
 	{
@@ -1374,6 +1375,7 @@ void UAbilitySystemComponent::InternalServerTryActiveAbility(FGameplayAbilitySpe
 		Spec->InputPressed = false;
 	}
 	MarkAbilitySpecDirty(*Spec);
+#endif
 }
 
 void UAbilitySystemComponent::ReplicateEndOrCancelAbility(FGameplayAbilitySpecHandle Handle, FGameplayAbilityActivationInfo ActivationInfo, UGameplayAbility* Ability, bool bWasCanceled)
diff --git a/Engine/Source/Runtime/GameplayAbilities/Private/AbilitySystemGlobals.cpp b/Engine/Source/Runtime/GameplayAbilities/Private/AbilitySystemGlobals.cpp
--- a/Engine/Source/Runtime/GameplayAbilities/Private/AbilitySystemGlobals.cpp
+++ b/Engine/Source/Runtime/GameplayAbilities/Private/AbilitySystemGlobals.cpp
@@ -38,8 +38,12 @@ void UAbilitySystemGlobals::InitGlobalData()
 	GetGameplayCueManager();
 	GetGameplayTagResponseTable();
 	InitGlobalTags();
+
+	//register for world clean up events
+	FWorldDelegates::OnPostWorldCreation.AddUObject(this, &UAbilitySystemGlobals::HandlePostWorldCreate);
 }
 
+
 UCurveTable * UAbilitySystemGlobals::GetGlobalCurveTable()
 {
 	return InternalGetLoadTable<UCurveTable>(GlobalCurveTable, GlobalCurveTableName);
@@ -198,7 +202,7 @@ UFunction* UAbilitySystemGlobals::GetGameplayCueFunction(const FGameplayTag& Chi
 
 // --------------------------------------------------------------------
 
-void UAbilitySystemGlobals::InitGameplayCueParameters(FGameplayCueParameters& CueParameters, const FGameplayEffectSpecForRPC &Spec)
+void UAbilitySystemGlobals::InitGameplayCueParameters(FGameplayCueParameters& CueParameters, const FGameplayEffectSpecForRPC& Spec)
 {
 	CueParameters.AggregatedSourceTags = Spec.AggregatedSourceTags;
 	CueParameters.AggregatedTargetTags = Spec.AggregatedTargetTags;
@@ -206,6 +210,36 @@ void UAbilitySystemGlobals::InitGameplayCueParameters(FGameplayCueParameters& Cu
 	InitGameplayCueParameters(CueParameters, Spec.GetContext());
 }
 
+void UAbilitySystemGlobals::InitGameplayCueParameters_GESpec(FGameplayCueParameters& CueParameters, const FGameplayEffectSpec& Spec)
+{
+	CueParameters.AggregatedSourceTags = *Spec.CapturedSourceTags.GetAggregatedTags();
+	CueParameters.AggregatedTargetTags = *Spec.CapturedTargetTags.GetAggregatedTags();
+
+	// Look for a modified attribute magnitude to pass to the CueParameters
+	for (const FGameplayEffectCue& CueDef : Spec.Def->GameplayCues)
+	{	
+		bool FoundMatch = false;
+		if (CueDef.MagnitudeAttribute.IsValid())
+		{
+			for (const FGameplayEffectModifiedAttribute& ModifiedAttribute : Spec.ModifiedAttributes)
+			{
+				if (ModifiedAttribute.Attribute == CueDef.MagnitudeAttribute)
+				{
+					CueParameters.RawMagnitude = ModifiedAttribute.TotalMagnitude;
+					FoundMatch = true;
+					break;
+				}
+			}
+			if (FoundMatch)
+			{
+				break;
+			}
+		}
+	}
+
+	InitGameplayCueParameters(CueParameters, Spec.GetContext());
+}
+
 void UAbilitySystemGlobals::InitGameplayCueParameters(FGameplayCueParameters& CueParameters, const FGameplayEffectContextHandle& EffectContext)
 {
 	if (EffectContext.IsValid())
@@ -332,3 +366,8 @@ bool UAbilitySystemGlobals::ShouldIgnoreCosts() const
 	return false;
 #endif // #if !(UE_BUILD_SHIPPING || UE_BUILD_TEST)
 }
+
+void UAbilitySystemGlobals::HandlePostWorldCreate(UWorld* NewWorld)
+{
+	IGameplayCueInterface::ClearTagToFunctionMap();
+}
diff --git a/Engine/Source/Runtime/GameplayAbilities/Private/AbilitySystemStats.cpp b/Engine/Source/Runtime/GameplayAbilities/Private/AbilitySystemStats.cpp
--- a/Engine/Source/Runtime/GameplayAbilities/Private/AbilitySystemStats.cpp
+++ b/Engine/Source/Runtime/GameplayAbilities/Private/AbilitySystemStats.cpp
@@ -29,3 +29,4 @@ DEFINE_STAT(STAT_HandleGameplayCueNotifyActor);
 DEFINE_STAT(STAT_ApplyGameplayEffectToTarget);
 DEFINE_STAT(STAT_OnActiveGameplayEffectAdded);
 DEFINE_STAT(STAT_OnActiveGameplayEffectRemoved);
+DEFINE_STAT(STAT_GameplayCueInterface_HandleGameplayCue);
diff --git a/Engine/Source/Runtime/GameplayAbilities/Private/GameplayCueInterface.cpp b/Engine/Source/Runtime/GameplayAbilities/Private/GameplayCueInterface.cpp
--- a/Engine/Source/Runtime/GameplayAbilities/Private/GameplayCueInterface.cpp
+++ b/Engine/Source/Runtime/GameplayAbilities/Private/GameplayCueInterface.cpp
@@ -6,6 +6,18 @@
 #include "GameplayTagsModule.h"
 #include "GameplayCueSet.h"
 
+
+struct FCueNameAndUFunction
+{
+	FGameplayTag Tag;
+	UFunction* Func;
+};
+typedef TMap<FGameplayTag, TArray<FCueNameAndUFunction> > FGameplayCueTagFunctionList;
+TMap<UClass*, FGameplayCueTagFunctionList > PerClassGameplayTagToFunctionMap;
+
+
+
+
 UGameplayCueInterface::UGameplayCueInterface(const FObjectInitializer& ObjectInitializer)
 : Super(ObjectInitializer)
 {
@@ -20,6 +32,12 @@ void IGameplayCueInterface::DispatchBlueprintCustomHandler(AActor* Actor, UFunct
 	Actor->ProcessEvent(Func, &Parms);
 }
 
+void IGameplayCueInterface::ClearTagToFunctionMap()
+{
+	PerClassGameplayTagToFunctionMap.Empty();
+}
+
+
 void IGameplayCueInterface::HandleGameplayCues(AActor *Self, const FGameplayTagContainer& GameplayCueTags, EGameplayCueEvent::Type EventType, FGameplayCueParameters Parameters)
 {
 	for (auto TagIt = GameplayCueTags.CreateConstIterator(); TagIt; ++TagIt)
@@ -35,57 +53,73 @@ bool IGameplayCueInterface::ShouldAcceptGameplayCue(AActor *Self, FGameplayTag G
 
 void IGameplayCueInterface::HandleGameplayCue(AActor *Self, FGameplayTag GameplayCueTag, EGameplayCueEvent::Type EventType, FGameplayCueParameters Parameters)
 {
-	// Look up a custom function for this gameplay tag. 
-	bool bFoundHandler = false;
-	bool bShouldContinue = true;
+	SCOPE_CYCLE_COUNTER(STAT_GameplayCueInterface_HandleGameplayCue);
 
+	// Look up a custom function for this gameplay tag. 
 	UClass* Class = Self->GetClass();
 	IGameplayTagsModule& GameplayTagsModule = IGameplayTagsModule::Get();
 	FGameplayTagContainer TagAndParentsContainer = GameplayTagsModule.GetGameplayTagsManager().RequestGameplayTagParents(GameplayCueTag);
 
 	Parameters.OriginalTag = GameplayCueTag;
 
-	for (auto InnerTagIt = TagAndParentsContainer.CreateConstIterator(); InnerTagIt && bShouldContinue; ++InnerTagIt)
+	//Find entry for the class
+	FGameplayCueTagFunctionList& GameplayTagFunctionList = PerClassGameplayTagToFunctionMap.FindOrAdd(Class);
+	TArray<FCueNameAndUFunction>* FunctionList = GameplayTagFunctionList.Find(GameplayCueTag);
+	if (FunctionList == NULL)
 	{
-		UFunction* Func = NULL;
-		FName CueName = InnerTagIt->GetTagName();
+		//generate new function list
+		FunctionList = &GameplayTagFunctionList.Add(GameplayCueTag);
 
-		Func = Class->FindFunctionByName(CueName, EIncludeSuperFlag::IncludeSuper);
-		// If the handler calls ForwardGameplayCueToParent, keep calling functions until one consumes the cue and doesn't forward it
-		while (bShouldContinue && Func)
-		{			
-			Parameters.MatchedTagName = CueName;
-
-			// Reset the forward parameter now, so we can check it after function
-			bForwardToParent = false;
-			IGameplayCueInterface::DispatchBlueprintCustomHandler(Self, Func, EventType, Parameters);
-			
-			bShouldContinue = bForwardToParent;
-			bFoundHandler = true;
-			Func = Func->GetSuperFunction();
-		}
-
-		if (bShouldContinue)
+		for (auto InnerTagIt = TagAndParentsContainer.CreateConstIterator(); InnerTagIt; ++InnerTagIt)
 		{
+			UFunction* Func = NULL;
+			FName CueName = InnerTagIt->GetTagName();
+
+			Func = Class->FindFunctionByName(CueName, EIncludeSuperFlag::IncludeSuper);
+			// If the handler calls ForwardGameplayCueToParent, keep calling functions until one consumes the cue and doesn't forward it
+			while (Func)
+			{
+				FCueNameAndUFunction NewCueFunctionPair;
+				NewCueFunctionPair.Tag = *InnerTagIt;
+				NewCueFunctionPair.Func = Func;
+				FunctionList->Add(NewCueFunctionPair);
+
+				Func = Func->GetSuperFunction();
+			}
+
 			// Native functions cant be named with ".", so look for them with _. 
 			FName NativeCueFuncName = *CueName.ToString().Replace(TEXT("."), TEXT("_"));
 			Func = Class->FindFunctionByName(NativeCueFuncName, EIncludeSuperFlag::IncludeSuper);
 
-			while (bShouldContinue && Func)
+			while (Func)
 			{
-				Parameters.MatchedTagName = CueName; // purposefully returning the . qualified name.
+				FCueNameAndUFunction NewCueFunctionPair;
+				NewCueFunctionPair.Tag = *InnerTagIt;
+				NewCueFunctionPair.Func = Func;
+				FunctionList->Add(NewCueFunctionPair);
 
-				// Reset the forward parameter now, so we can check it after function
-				bForwardToParent = false;
-				IGameplayCueInterface::DispatchBlueprintCustomHandler(Self, Func, EventType, Parameters);
-
-				bShouldContinue = bForwardToParent;
-				bFoundHandler = true;
 				Func = Func->GetSuperFunction();
 			}
 		}
 	}
 
+	//Iterate through all functions in the list until we should no longer continue
+	check(FunctionList);
+		
+	bool bShouldContinue = true;
+	for (int32 FunctionIndex = 0; bShouldContinue && (FunctionIndex < FunctionList->Num()); ++FunctionIndex)
+	{
+		FCueNameAndUFunction& CueFunctionPair = FunctionList->GetData()[FunctionIndex];
+		UFunction* Func = CueFunctionPair.Func;
+		Parameters.MatchedTagName = CueFunctionPair.Tag;
+
+		// Reset the forward parameter now, so we can check it after function
+		bForwardToParent = false;
+		IGameplayCueInterface::DispatchBlueprintCustomHandler(Self, Func, EventType, Parameters);
+
+		bShouldContinue = bForwardToParent;
+	}
+
 	if (bShouldContinue)
 	{
 		TArray<UGameplayCueSet*> Sets;
@@ -102,7 +136,7 @@ void IGameplayCueInterface::HandleGameplayCue(AActor *Self, FGameplayTag Gamepla
 
 	if (bShouldContinue)
 	{
-		Parameters.MatchedTagName = GameplayCueTag.GetTagName();
+		Parameters.MatchedTagName = GameplayCueTag;
 		GameplayCueDefaultHandler(EventType, Parameters);
 	}
 }
@@ -125,8 +159,8 @@ void FActiveGameplayCue::PreReplicatedRemove(const struct FActiveGameplayCueCont
 	if (bPredictivelyRemoved == false)
 	{
 		// If predicted ignore the add/remove
-		InArray.Owner->InvokeGameplayCueEvent(GameplayCueTag, EGameplayCueEvent::Removed);
 		InArray.Owner->UpdateTagMap(GameplayCueTag, -1);
+		InArray.Owner->InvokeGameplayCueEvent(GameplayCueTag, EGameplayCueEvent::Removed);
 	}
 }
 
@@ -181,8 +215,8 @@ void FActiveGameplayCueContainer::PredictiveRemove(const FGameplayTag& Tag)
 			// Predictive remove: mark the cue as predictive remove, invoke remove event, update tag map.
 			// DONT remove from the replicated array.
 			Cue.bPredictivelyRemoved = true;
-			Owner->InvokeGameplayCueEvent(Tag, EGameplayCueEvent::Removed);
 			Owner->UpdateTagMap(Tag, -1);
+			Owner->InvokeGameplayCueEvent(Tag, EGameplayCueEvent::Removed);	
 			return;
 		}
 	}
@@ -206,4 +240,14 @@ bool FActiveGameplayCueContainer::HasCue(const FGameplayTag& Tag) const
 	}
 
 	return false;
+}
+
+bool FActiveGameplayCueContainer::NetDeltaSerialize(FNetDeltaSerializeInfo & DeltaParms)
+{
+	if (bMinimalReplication && (Owner && !Owner->bMinimalReplication))
+	{
+		return true;
+	}
+
+	return FastArrayDeltaSerialize<FActiveGameplayCue>(GameplayCues, DeltaParms, *this);
 }
\ No newline at end of file
diff --git a/Engine/Source/Runtime/GameplayAbilities/Private/GameplayCueManager.cpp b/Engine/Source/Runtime/GameplayAbilities/Private/GameplayCueManager.cpp
--- a/Engine/Source/Runtime/GameplayAbilities/Private/GameplayCueManager.cpp
+++ b/Engine/Source/Runtime/GameplayAbilities/Private/GameplayCueManager.cpp
@@ -55,6 +55,7 @@ void UGameplayCueManager::OnCreated()
 {
 	FWorldDelegates::OnPostWorldCreation.AddUObject(this, &UGameplayCueManager::OnWorldCreated);
 	FWorldDelegates::OnWorldCleanup.AddUObject(this, &UGameplayCueManager::OnWorldCleanup);
+	FWorldDelegates::OnPreWorldFinishDestroy.AddUObject(this, &UGameplayCueManager::OnWorldCleanup, true, true);
 
 	FNetworkReplayDelegates::OnPreScrub.AddUObject(this, &UGameplayCueManager::OnPreReplayScrub);
 }
@@ -645,13 +646,106 @@ FScopedGameplayCueSendContext::~FScopedGameplayCueSendContext()
 	UAbilitySystemGlobals::Get().GetGameplayCueManager()->EndGameplayCueSendContext();
 }
 
-void UGameplayCueManager::InvokeGameplayCueExecuted_FromSpec(UAbilitySystemComponent* OwningComponent, const FGameplayEffectSpecForRPC Spec, FPredictionKey PredictionKey)
+template<class AllocatorType>
+void PullGameplayCueTagsFromSpec(const FGameplayEffectSpec& Spec, TArray<FGameplayTag, AllocatorType>& OutArray)
 {
+	// Add all GameplayCue Tags from the GE into the GameplayCueTags PendingCue.list
+	for (const FGameplayEffectCue& EffectCue : Spec.Def->GameplayCues)
+	{
+		for (const FGameplayTag& Tag: EffectCue.GameplayCueTags)
+		{
+			if (Tag.IsValid())
+			{
+				OutArray.Add(Tag);
+			}
+		}
+	}
+}
+
+/**
+ *	Enabling AbilitySystemAlwaysConvertGESpecToGCParams will mean that all calls to gameplay cues with GameplayEffectSpecs will be converted into GameplayCue Parameters server side and then replicated.
+ *	This potentially saved bandwidth but also has less information, depending on how the GESpec is converted to GC Parameters and what your GC's need to know.
+ */
+
+int32 AbilitySystemAlwaysConvertGESpecToGCParams = 0;
+static FAutoConsoleVariableRef CVarAbilitySystemAlwaysConvertGESpecToGCParams(TEXT("AbilitySystem.AlwaysConvertGESpecToGCParams"), AbilitySystemAlwaysConvertGESpecToGCParams, TEXT("Always convert a GameplayCue from GE Spec to GC from GC Parameters on the server"), ECVF_Default );
+
+void UGameplayCueManager::InvokeGameplayCueAddedAndWhileActive_FromSpec(UAbilitySystemComponent* OwningComponent, const FGameplayEffectSpec& Spec, FPredictionKey PredictionKey)
+{
+	if (Spec.Def->GameplayCues.Num() == 0)
+	{
+		return;
+	}
+
+	if (AbilitySystemAlwaysConvertGESpecToGCParams)
+	{
+		// Transform the GE Spec into GameplayCue parmameters here (on the server)
+
+		FGameplayCueParameters Parameters;
+		UAbilitySystemGlobals::Get().InitGameplayCueParameters_GESpec(Parameters, Spec);
+
+		static TArray<FGameplayTag, TInlineAllocator<4> > Tags;
+		Tags.Reset();
+
+		PullGameplayCueTagsFromSpec(Spec, Tags);
+
+		if (Tags.Num() == 1)
+		{
+			OwningComponent->NetMulticast_InvokeGameplayCueAddedAndWhileActive_WithParams(Tags[0], PredictionKey, Parameters);
+			
+		}
+		else if (Tags.Num() > 1)
+		{
+			OwningComponent->NetMulticast_InvokeGameplayCuesAddedAndWhileActive_WithParams(FGameplayTagContainer::CreateFromArray(Tags), PredictionKey, Parameters);
+		}
+		else
+		{
+			ABILITY_LOG(Warning, TEXT("No actual gameplay cue tags found in GameplayEffect %s (despite it having entries in its gameplay cue list!"), *Spec.Def->GetName());
+
+		}
+	}
+	else
+	{
+		OwningComponent->NetMulticast_InvokeGameplayCueAddedAndWhileActive_FromSpec(Spec, PredictionKey);
+
+	}
+}
+
+void UGameplayCueManager::InvokeGameplayCueExecuted_FromSpec(UAbilitySystemComponent* OwningComponent, const FGameplayEffectSpec& Spec, FPredictionKey PredictionKey)
+{	
+	if (Spec.Def->GameplayCues.Num() == 0)
+	{
+		// This spec doesn't have any GCs, so early out
+		ABILITY_LOG(Verbose, TEXT("No GCs in this Spec, so early out: %s"), *Spec.Def->GetName());
+		return;
+	}
+
 	FGameplayCuePendingExecute PendingCue;
-	PendingCue.PayloadType = EGameplayCuePayloadType::FromSpec;
-	PendingCue.OwningComponent = OwningComponent;
-	PendingCue.FromSpec = Spec;
-	PendingCue.PredictionKey = PredictionKey;
+
+	if (AbilitySystemAlwaysConvertGESpecToGCParams)
+	{
+		// Transform the GE Spec into GameplayCue parmameters here (on the server)
+		PendingCue.PayloadType = EGameplayCuePayloadType::CueParameters;
+		PendingCue.OwningComponent = OwningComponent;
+		PendingCue.PredictionKey = PredictionKey;
+
+		PullGameplayCueTagsFromSpec(Spec, PendingCue.GameplayCueTags);
+		if (PendingCue.GameplayCueTags.Num() == 0)
+		{
+			ABILITY_LOG(Warning, TEXT("GE %s has GameplayCues but not valid GameplayCue tag."), *Spec.Def->GetName());			
+			return;
+		}
+		
+		UAbilitySystemGlobals::Get().InitGameplayCueParameters_GESpec(PendingCue.CueParameters, Spec);
+	}
+	else
+	{
+		// Transform the GE Spec into a FGameplayEffectSpecForRPC (holds less information than the GE Spec itself, but more information that the FGamepalyCueParameter)
+		PendingCue.PayloadType = EGameplayCuePayloadType::FromSpec;
+		PendingCue.OwningComponent = OwningComponent;
+		PendingCue.FromSpec = FGameplayEffectSpecForRPC(Spec);
+		PendingCue.PredictionKey = PredictionKey;
+	}
 
 	if (ProcessPendingCueExecute(PendingCue))
 	{
@@ -669,7 +763,7 @@ void UGameplayCueManager::InvokeGameplayCueExecuted(UAbilitySystemComponent* Own
 {
 	FGameplayCuePendingExecute PendingCue;
 	PendingCue.PayloadType = EGameplayCuePayloadType::CueParameters;
-	PendingCue.GameplayCueTag = GameplayCueTag;
+	PendingCue.GameplayCueTags.Add(GameplayCueTag);
 	PendingCue.OwningComponent = OwningComponent;
 	UAbilitySystemGlobals::Get().InitGameplayCueParameters(PendingCue.CueParameters, EffectContext);
 	PendingCue.PredictionKey = PredictionKey;
@@ -690,7 +784,7 @@ void UGameplayCueManager::InvokeGameplayCueExecuted_WithParams(UAbilitySystemCom
 {
 	FGameplayCuePendingExecute PendingCue;
 	PendingCue.PayloadType = EGameplayCuePayloadType::CueParameters;
-	PendingCue.GameplayCueTag = GameplayCueTag;
+	PendingCue.GameplayCueTags.Add(GameplayCueTag);
 	PendingCue.OwningComponent = OwningComponent;
 	PendingCue.CueParameters = GameplayCueParameters;
 	PendingCue.PredictionKey = PredictionKey;
@@ -744,29 +838,55 @@ void UGameplayCueManager::FlushPendingCues()
 			switch (PendingCue.PayloadType)
 			{
 			case EGameplayCuePayloadType::CueParameters:
-				if (bHasAuthority)
+				if (ensure(PendingCue.GameplayCueTags.Num() >= 1))
 				{
-					PendingCue.OwningComponent->ForceReplication();
-					PendingCue.OwningComponent->NetMulticast_InvokeGameplayCueExecuted_WithParams(PendingCue.GameplayCueTag, PendingCue.PredictionKey, PendingCue.CueParameters);
-					static FName NetMulticast_InvokeGameplayCueExecuted_WithParamsName = TEXT("NetMulticast_InvokeGameplayCueExecuted_WithParams");
-					CheckForTooManyRPCs(NetMulticast_InvokeGameplayCueExecuted_WithParamsName, PendingCue, PendingCue.GameplayCueTag.ToString(), nullptr);
-				}
-				else if (bLocalPredictionKey)
-				{
-					PendingCue.OwningComponent->InvokeGameplayCueEvent(PendingCue.GameplayCueTag, EGameplayCueEvent::Executed, PendingCue.CueParameters);
+					if (bHasAuthority)
+					{
+						PendingCue.OwningComponent->ForceReplication();
+						if (PendingCue.GameplayCueTags.Num() > 1)
+						{
+							PendingCue.OwningComponent->NetMulticast_InvokeGameplayCuesExecuted_WithParams(FGameplayTagContainer::CreateFromArray(PendingCue.GameplayCueTags), PendingCue.PredictionKey, PendingCue.CueParameters);
+						}
+						else
+						{
+							PendingCue.OwningComponent->NetMulticast_InvokeGameplayCueExecuted_WithParams(PendingCue.GameplayCueTags[0], PendingCue.PredictionKey, PendingCue.CueParameters);
+							static FName NetMulticast_InvokeGameplayCueExecuted_WithParamsName = TEXT("NetMulticast_InvokeGameplayCueExecuted_WithParams");
+							CheckForTooManyRPCs(NetMulticast_InvokeGameplayCueExecuted_WithParamsName, PendingCue, PendingCue.GameplayCueTags[0].ToString(), nullptr);
+						}
+					}
+					else if (bLocalPredictionKey)
+					{
+						for (const FGameplayTag& Tag : PendingCue.GameplayCueTags)
+						{
+							PendingCue.OwningComponent->InvokeGameplayCueEvent(Tag, EGameplayCueEvent::Executed, PendingCue.CueParameters);
+						}
+					}
 				}
 				break;
 			case EGameplayCuePayloadType::EffectContext:
-				if (bHasAuthority)
+				if (ensure(PendingCue.GameplayCueTags.Num() >= 1))
 				{
-					PendingCue.OwningComponent->ForceReplication();
-					PendingCue.OwningComponent->NetMulticast_InvokeGameplayCueExecuted(PendingCue.GameplayCueTag, PendingCue.PredictionKey, PendingCue.CueParameters.EffectContext);
-					static FName NetMulticast_InvokeGameplayCueExecutedName = TEXT("NetMulticast_InvokeGameplayCueExecuted");
-					CheckForTooManyRPCs(NetMulticast_InvokeGameplayCueExecutedName, PendingCue, PendingCue.GameplayCueTag.ToString(), PendingCue.CueParameters.EffectContext.Get());
-				}
-				else if (bLocalPredictionKey)
-				{
-					PendingCue.OwningComponent->InvokeGameplayCueEvent(PendingCue.GameplayCueTag, EGameplayCueEvent::Executed, PendingCue.CueParameters.EffectContext);
+					if (bHasAuthority)
+					{
+						PendingCue.OwningComponent->ForceReplication();
+						if (PendingCue.GameplayCueTags.Num() > 1)
+						{
+							PendingCue.OwningComponent->NetMulticast_InvokeGameplayCuesExecuted(FGameplayTagContainer::CreateFromArray(PendingCue.GameplayCueTags), PendingCue.PredictionKey, PendingCue.CueParameters.EffectContext);
+						}
+						else
+						{
+							PendingCue.OwningComponent->NetMulticast_InvokeGameplayCueExecuted(PendingCue.GameplayCueTags[0], PendingCue.PredictionKey, PendingCue.CueParameters.EffectContext);
+							static FName NetMulticast_InvokeGameplayCueExecutedName = TEXT("NetMulticast_InvokeGameplayCueExecuted");
+							CheckForTooManyRPCs(NetMulticast_InvokeGameplayCueExecutedName, PendingCue, PendingCue.GameplayCueTags[0].ToString(), PendingCue.CueParameters.EffectContext.Get());
+						}
+					}
+					else if (bLocalPredictionKey)
+					{
+						for (const FGameplayTag& Tag : PendingCue.GameplayCueTags)
+						{
+							PendingCue.OwningComponent->InvokeGameplayCueEvent(Tag, EGameplayCueEvent::Executed, PendingCue.CueParameters.EffectContext);
+						}
+					}
 				}
 				break;
 			case EGameplayCuePayloadType::FromSpec:
@@ -828,7 +948,7 @@ bool UGameplayCueManager::DoesPendingCueExecuteMatch(FGameplayCuePendingExecute&
 	}
 	else
 	{
-		if (PendingCue.GameplayCueTag != ExistingCue.GameplayCueTag)
+		if (PendingCue.GameplayCueTags != ExistingCue.GameplayCueTags)
 		{
 			return false;
 		}
@@ -923,12 +1043,17 @@ FPreallocationInfo& UGameplayCueManager::GetPreallocationInfo(UWorld* World)
 
 void UGameplayCueManager::OnWorldCreated(UWorld* NewWorld)
 {
+	// Attempting to track down rare GC error where PreallocationInfo_Internal.OwningWorld is not cleaned up.
+	ABILITY_LOG(Display, TEXT("UGameplayCueManager::OnWorldCreated %s. Current PreallocationInfo_Internal: %s"), *GetNameSafe(NewWorld), *GetNameSafe(PreallocationInfo_Internal.OwningWorld));
+
 	PreallocationInfo_Internal.PreallocatedInstances.Reset();
 	PreallocationInfo_Internal.OwningWorld = NewWorld;
 }
 
 void UGameplayCueManager::OnWorldCleanup(UWorld* World, bool bSessionEnded, bool bCleanupResources)
 {
+	// Attempting to track down rare GC error where PreallocationInfo_Internal.OwningWorld is not cleaned up.
+	ABILITY_LOG(Display, TEXT("UGameplayCueManager::OnWorldCleanup %s. Current PreallocationInfo_Internal: %s"), *GetNameSafe(World), *GetNameSafe(PreallocationInfo_Internal.OwningWorld));
 
 #if !(UE_BUILD_SHIPPING || UE_BUILD_TEST)
 	DumpPreallocationStats(World);
@@ -970,7 +1095,7 @@ void UGameplayCueManager::DumpPreallocationStats(UWorld* World)
 				TArray<AGameplayCueNotify_Actor*>& List = It.Value;
 				if (List.Num() > CDO->NumPreallocatedInstances)
 				{
-					ABILITY_LOG(Warning, TEXT("Notify class: %s was used simultaneously %d times. The CDO default is %d preallocated instanced."), *ThisClass->GetName(), List.Num(),  CDO->NumPreallocatedInstances); 
+					ABILITY_LOG(Display, TEXT("Notify class: %s was used simultaneously %d times. The CDO default is %d preallocated instanced."), *ThisClass->GetName(), List.Num(),  CDO->NumPreallocatedInstances); 
 				}
 			}
 		}
diff --git a/Engine/Source/Runtime/GameplayAbilities/Private/GameplayCueNotify_Actor.cpp b/Engine/Source/Runtime/GameplayAbilities/Private/GameplayCueNotify_Actor.cpp
--- a/Engine/Source/Runtime/GameplayAbilities/Private/GameplayCueNotify_Actor.cpp
+++ b/Engine/Source/Runtime/GameplayAbilities/Private/GameplayCueNotify_Actor.cpp
@@ -18,8 +18,6 @@ AGameplayCueNotify_Actor::AGameplayCueNotify_Actor(const FObjectInitializer& Obj
 	bUniqueInstancePerInstigator = false;
 
 	NumPreallocatedInstances = 0;
-
-	StackCount = 0;
 }
 
 #if WITH_EDITOR
@@ -79,40 +77,32 @@ bool AGameplayCueNotify_Actor::HandlesEvent(EGameplayCueEvent::Type EventType) c
 	return true;
 }
 
-int32 GameplayCueNotifyActorStacking = 0;
-static FAutoConsoleVariableRef CVarGameplayCueNotifyActorStacking(TEXT("AbilitySystem.GameplayCueNotifyActorStacking"), GameplayCueNotifyActorStacking, TEXT("Enable simple stacking rules for gameplaycue actors"), ECVF_Default );
+int32 GameplayCueNotifyTagCheckOnRemove = 1;
+static FAutoConsoleVariableRef CVarGameplayCueNotifyActorStacking(TEXT("AbilitySystem.GameplayCueNotifyTagCheckOnRemove"), GameplayCueNotifyTagCheckOnRemove, TEXT("Check that target no longer has tag when removing GamepalyCues"), ECVF_Default );
 
 void AGameplayCueNotify_Actor::HandleGameplayCue(AActor* MyTarget, EGameplayCueEvent::Type EventType, const FGameplayCueParameters& Parameters)
 {
 	SCOPE_CYCLE_COUNTER(STAT_HandleGameplayCueNotifyActor);
 
-	if (MyTarget && !MyTarget->IsPendingKill())
+	if (Parameters.MatchedTagName.IsValid() == false)
+	{
+		ABILITY_LOG(Warning, TEXT("GameplayCue parameter is none for %s"), *GetNameSafe(this));
+	}
+
+	// If cvar is enabled, check that the target no longer has the matched tag before doing remove logic. This is a simple way of supporting stacking, such that if an actor has two sources giving him the same GC tag, it will not be removed when the first one is removed.
+	if (GameplayCueNotifyTagCheckOnRemove > 0 && EventType == EGameplayCueEvent::Removed)
 	{
-		if (GameplayCueNotifyActorStacking)
+		if (IGameplayTagAssetInterface* TagInterface = Cast<IGameplayTagAssetInterface>(MyTarget))
 		{
-			if (EventType == EGameplayCueEvent::WhileActive)
+			if (TagInterface->HasMatchingGameplayTag(Parameters.MatchedTagName))
 			{
-				StackCount++;
-				if (StackCount > 1)
-				{
-					return;
-				}
-			}
-			else if (EventType == EGameplayCueEvent::Removed)
-			{
-				StackCount--;
-				if (StackCount > 0)
-				{
-					return;
-				}
-				if (!ensureMsgf(StackCount == 0, TEXT("GameplayCue %s has negative StackCount."), *GetName()))
-				{
-					StackCount = 0;
-				}
-			}
+				return;
+			}			
 		}
+	}
 
-
+	if (MyTarget && !MyTarget->IsPendingKill())
+	{
 		K2_HandleGameplayCue(MyTarget, EventType, Parameters);
 
 		// Clear any pending auto-destroy that may have occurred from a previous OnRemove
@@ -151,7 +141,12 @@ void AGameplayCueNotify_Actor::HandleGameplayCue(AActor* MyTarget, EGameplayCueE
 	}
 	else
 	{
-		ABILITY_LOG(Warning, TEXT("Null Target"));
+		ABILITY_LOG(Warning, TEXT("Null Target called for event %d on GameplayCueNotifyActor %s"), (int32)EventType, *GetName() );
+		if (EventType == EGameplayCueEvent::Removed)
+		{
+			// Make sure the removed event is handled so that we don't leak GC notify actors
+			GameplayCueFinishedCallback();
+		}
 	}
 }
 
@@ -198,6 +193,5 @@ bool AGameplayCueNotify_Actor::GameplayCuePendingRemove()
 
 bool AGameplayCueNotify_Actor::Recycle()
 {
-	StackCount = 0;
 	return false;
 }
\ No newline at end of file
diff --git a/Engine/Source/Runtime/GameplayAbilities/Private/GameplayCueSet.cpp b/Engine/Source/Runtime/GameplayAbilities/Private/GameplayCueSet.cpp
--- a/Engine/Source/Runtime/GameplayAbilities/Private/GameplayCueSet.cpp
+++ b/Engine/Source/Runtime/GameplayAbilities/Private/GameplayCueSet.cpp
@@ -27,7 +27,10 @@ bool UGameplayCueSet::HandleGameplayCue(AActor* TargetActor, FGameplayTag Gamepl
 	if (Ptr)
 	{
 		int32 DataIdx = *Ptr;
-		return HandleGameplayCueNotify_Internal(TargetActor, DataIdx, EventType, Parameters);
+
+		// TODO - resolve internal handler modifying params before passing them on with new const-ref params.
+		FGameplayCueParameters writableParameters = Parameters;
+		return HandleGameplayCueNotify_Internal(TargetActor, DataIdx, EventType, writableParameters);
 	}
 
 	return false;
@@ -130,7 +133,7 @@ void UGameplayCueSet::PrintCues() const
 	}
 }
 
-bool UGameplayCueSet::HandleGameplayCueNotify_Internal(AActor* TargetActor, int32 DataIdx, EGameplayCueEvent::Type EventType, const FGameplayCueParameters& Parameters)
+bool UGameplayCueSet::HandleGameplayCueNotify_Internal(AActor* TargetActor, int32 DataIdx, EGameplayCueEvent::Type EventType, FGameplayCueParameters& Parameters)
 {	
 	bool bReturnVal = false;
 
@@ -140,6 +143,8 @@ bool UGameplayCueSet::HandleGameplayCueNotify_Internal(AActor* TargetActor, int3
 
 		FGameplayCueNotifyData& CueData = GameplayCueData[DataIdx];
 
+		Parameters.MatchedTagName = CueData.GameplayCueTag;
+
 		// If object is not loaded yet
 		if (CueData.LoadedGameplayCueClass == nullptr)
 		{
@@ -244,7 +249,8 @@ void UGameplayCueSet::BuildAccelerationMap_Internal()
 
 		FGameplayTag Parent = ThisGameplayCueTag.RequestDirectParent();
 
-		GameplayCueDataMap.Add(ThisGameplayCueTag) = GameplayCueDataMap.FindChecked(Parent);
+		int32 ParentValue = GameplayCueDataMap.FindChecked(Parent);
+		GameplayCueDataMap.Add(ThisGameplayCueTag, ParentValue);
 	}
 
 
diff --git a/Engine/Source/Runtime/GameplayAbilities/Private/GameplayEffect.cpp b/Engine/Source/Runtime/GameplayAbilities/Private/GameplayEffect.cpp
--- a/Engine/Source/Runtime/GameplayAbilities/Private/GameplayEffect.cpp
+++ b/Engine/Source/Runtime/GameplayAbilities/Private/GameplayEffect.cpp
@@ -1491,6 +1491,7 @@ void FActiveGameplayEffect::PostReplicatedChange(const struct FActiveGameplayEff
 
 FActiveGameplayEffectsContainer::FActiveGameplayEffectsContainer()
 	: Owner(nullptr)
+	, OwnerIsNetAuthority(false)
 	, ScopedLockCount(0)
 	, PendingRemoves(0)
 	, PendingGameplayEffectHead(nullptr)
@@ -1514,6 +1515,7 @@ void FActiveGameplayEffectsContainer::RegisterWithOwner(UAbilitySystemComponent*
 	if (Owner != InOwner)
 	{
 		Owner = InOwner;
+		OwnerIsNetAuthority = Owner->IsOwnerActorAuthoritative();
 
 		// Binding raw is ok here, since the owner is literally the UObject that owns us. If we are destroyed, its because that uobject is destroyed,
 		// and if that is destroyed, the delegate wont be able to fire.
@@ -2570,8 +2572,12 @@ void FActiveGameplayEffectsContainer::AddActiveGameplayEffectGrantedTagsAndModif
 
 	// Update our owner with the tags this GameplayEffect grants them
 	Owner->UpdateTagMap(Effect.Spec.Def->InheritableOwnedTagsContainer.CombinedTags, 1);
-
 	Owner->UpdateTagMap(Effect.Spec.DynamicGrantedTags, 1);
+	if (IsNetAuthority() && Owner->bMinimalReplication)
+	{
+		Owner->AddMinimalReplicationGameplayTags(Effect.Spec.Def->InheritableOwnedTagsContainer.CombinedTags);
+		Owner->AddMinimalReplicationGameplayTags(Effect.Spec.DynamicGrantedTags);
+	}
 
 	ApplicationImmunityGameplayTagCountContainer.UpdateTagCount(Effect.Spec.Def->GrantedApplicationImmunityTags.RequireTags, 1);
 	ApplicationImmunityGameplayTagCountContainer.UpdateTagCount(Effect.Spec.Def->GrantedApplicationImmunityTags.IgnoreTags, 1);
@@ -2600,6 +2606,17 @@ void FActiveGameplayEffectsContainer::AddActiveGameplayEffectGrantedTagsAndModif
 			Owner->InvokeGameplayCueEvent(Effect.Spec, EGameplayCueEvent::OnActive);
 			Owner->InvokeGameplayCueEvent(Effect.Spec, EGameplayCueEvent::WhileActive);
 		}
+
+		if (IsNetAuthority() && Owner->bMinimalReplication)
+		{
+			for (const FGameplayTag& CueTag : Cue.GameplayCueTags)
+			{
+				// Note: minimal replication does not replicate the effect context with the gameplay cue parameters.
+				// This is just a choice right now. If needed, it may be better to convert the effect context to GC parameters *here*
+				// and pass those into this function
+				Owner->AddGameplayCue_MinimalReplication(CueTag);
+			}
+		}
 	}
 }
 
@@ -2776,12 +2793,17 @@ void FActiveGameplayEffectsContainer::RemoveActiveGameplayEffectGrantedTagsAndMo
 	// Update gameplaytag count and broadcast delegate if we are at 0
 	IGameplayTagsModule& GameplayTagsModule = IGameplayTagsModule::Get();
 	Owner->UpdateTagMap(Effect.Spec.Def->InheritableOwnedTagsContainer.CombinedTags, -1);
+	Owner->UpdateTagMap(Effect.Spec.DynamicGrantedTags, -1);
+
+	if (IsNetAuthority() && Owner->bMinimalReplication)
+	{
+		Owner->RemoveMinimalReplicationGameplayTags(Effect.Spec.Def->InheritableOwnedTagsContainer.CombinedTags);
+		Owner->RemoveMinimalReplicationGameplayTags(Effect.Spec.DynamicGrantedTags);
+	}
 
 	ApplicationImmunityGameplayTagCountContainer.UpdateTagCount(Effect.Spec.Def->GrantedApplicationImmunityTags.RequireTags, -1);
 	ApplicationImmunityGameplayTagCountContainer.UpdateTagCount(Effect.Spec.Def->GrantedApplicationImmunityTags.IgnoreTags, -1);
 
-	Owner->UpdateTagMap(Effect.Spec.DynamicGrantedTags, -1);
-
 	// Cancel/remove granted abilities
 	if (IsNetAuthority())
 	{
@@ -2814,12 +2836,20 @@ void FActiveGameplayEffectsContainer::RemoveActiveGameplayEffectGrantedTagsAndMo
 	// Update GameplayCue tags and events
 	for (const FGameplayEffectCue& Cue : Effect.Spec.Def->GameplayCues)
 	{
+		Owner->UpdateTagMap(Cue.GameplayCueTags, -1);
+
 		if (bInvokeGameplayCueEvents)
 		{
 			Owner->InvokeGameplayCueEvent(Effect.Spec, EGameplayCueEvent::Removed);			
 		}
 
-		Owner->UpdateTagMap(Cue.GameplayCueTags, -1);
+		if (IsNetAuthority() && Owner->bMinimalReplication)
+		{
+			for (const FGameplayTag& CueTag : Cue.GameplayCueTags)
+			{
+				Owner->RemoveGameplayCue_MinimalReplication(CueTag);
+			}
+		}
 	}
 }
 
@@ -2938,14 +2968,14 @@ bool FActiveGameplayEffectsContainer::HasApplicationImmunityToSpec(const FGamepl
 	return false;
 }
 
-bool FActiveGameplayEffectsContainer::IsNetAuthority() const
-{
-	check(Owner);
-	return Owner->IsOwnerActorAuthoritative();
-}
-
 bool FActiveGameplayEffectsContainer::NetDeltaSerialize(FNetDeltaSerializeInfo& DeltaParms)
 {
+	// Have to manually return here, replication conditions don't apply to delta properties
+	if (Owner && Owner->bMinimalReplication)
+	{
+		return true;
+	}
+
 	bool RetVal = FastArrayDeltaSerialize<FActiveGameplayEffect>(GameplayEffects_Internal, DeltaParms, *this);
 
 	// After the array has been replicated, invoke GC events ONLY if the effect is not inhibited
diff --git a/Engine/Source/Runtime/GameplayAbilities/Private/GameplayEffectTypes.cpp b/Engine/Source/Runtime/GameplayAbilities/Private/GameplayEffectTypes.cpp
--- a/Engine/Source/Runtime/GameplayAbilities/Private/GameplayEffectTypes.cpp
+++ b/Engine/Source/Runtime/GameplayAbilities/Private/GameplayEffectTypes.cpp
@@ -286,41 +286,6 @@ FString EGameplayCueEventToString(int32 Type)
 	return e->GetEnumName(Type);
 }
 
-bool FGameplayTagCountContainer::HasMatchingGameplayTag(FGameplayTag TagToCheck) const
-{
-	return ExplicitTags.HasTag(TagToCheck, EGameplayTagMatchType::IncludeParentTags, EGameplayTagMatchType::Explicit);
-}
-
-bool FGameplayTagCountContainer::HasAllMatchingGameplayTags(const FGameplayTagContainer& TagContainer, bool bCountEmptyAsMatch) const
-{
-	return ExplicitTags.MatchesAll(TagContainer, bCountEmptyAsMatch);
-}
-
-
-bool FGameplayTagCountContainer::HasAnyMatchingGameplayTags(const FGameplayTagContainer& TagContainer, bool bCountEmptyAsMatch) const
-{
-	return ExplicitTags.MatchesAny(TagContainer, bCountEmptyAsMatch);
-}
-
-void FGameplayTagCountContainer::UpdateTagCount(const FGameplayTag& Tag, int32 CountDelta)
-{
-	if (CountDelta != 0)
-	{
-		UpdateTagMap_Internal(Tag, CountDelta);
-	}
-}
-
-void FGameplayTagCountContainer::UpdateTagCount(const FGameplayTagContainer& Container, int32 CountDelta)
-{	
-	if (CountDelta != 0)
-	{
-		for (auto TagIt = Container.CreateConstIterator(); TagIt; ++TagIt)
-		{
-			UpdateTagMap_Internal(*TagIt, CountDelta);
-		}
-	}
-}
-
 void FGameplayTagCountContainer::Notify_StackCountChange(const FGameplayTag& Tag)
 {	
 	// The purpose of this function is to let anyone listening on the EGameplayTagEventType::AnyCountChange event know that the 
@@ -351,16 +316,6 @@ FOnGameplayEffectTagCountChanged& FGameplayTagCountContainer::RegisterGameplayTa
 	return Info.OnAnyChange;
 }
 
-FOnGameplayEffectTagCountChanged& FGameplayTagCountContainer::RegisterGenericGameplayEvent()
-{
-	return OnAnyTagChangeDelegate;
-}
-
-const FGameplayTagContainer& FGameplayTagCountContainer::GetExplicitGameplayTags() const
-{
-	return ExplicitTags;
-}
-
 void FGameplayTagCountContainer::Reset()
 {
 	GameplayTagEventMap.Reset();
@@ -370,7 +325,7 @@ void FGameplayTagCountContainer::Reset()
 	OnAnyTagChangeDelegate.Clear();
 }
 
-void FGameplayTagCountContainer::UpdateTagMap_Internal(const FGameplayTag& Tag, int32 CountDelta)
+bool FGameplayTagCountContainer::UpdateTagMap_Internal(const FGameplayTag& Tag, int32 CountDelta)
 {
 	const bool bTagAlreadyExplicitlyExists = ExplicitTags.HasTag(Tag, EGameplayTagMatchType::Explicit, EGameplayTagMatchType::Explicit);
 
@@ -391,7 +346,7 @@ void FGameplayTagCountContainer::UpdateTagMap_Internal(const FGameplayTag& Tag,
 			{
 				ABILITY_LOG(Warning, TEXT("Attempted to remove tag: %s from tag count container, but it is not explicitly in the container!"), *Tag.ToString());
 			}
-			return;
+			return false;
 		}
 	}
 
@@ -409,6 +364,7 @@ void FGameplayTagCountContainer::UpdateTagMap_Internal(const FGameplayTag& Tag,
 
 	// Check if change delegates are required to fire for the tag or any of its parents based on the count change
 	FGameplayTagContainer TagAndParentsContainer = IGameplayTagsModule::Get().GetGameplayTagsManager().RequestGameplayTagParents(Tag);
+	bool CreatedSignificantChange = false;
 	for (auto CompleteTagIt = TagAndParentsContainer.CreateConstIterator(); CompleteTagIt; ++CompleteTagIt)
 	{
 		const FGameplayTag& CurTag = *CompleteTagIt;
@@ -423,6 +379,7 @@ void FGameplayTagCountContainer::UpdateTagMap_Internal(const FGameplayTag& Tag,
 
 		// If a significant change (new addition or total removal) occurred, trigger related delegates
 		bool SignificantChange = (OldCount == 0 || TagCount == 0);
+		CreatedSignificantChange |= SignificantChange;
 		if (SignificantChange)
 		{
 			OnAnyTagChangeDelegate.Broadcast(CurTag, TagCount);
@@ -438,6 +395,8 @@ void FGameplayTagCountContainer::UpdateTagMap_Internal(const FGameplayTag& Tag,
 			}
 		}
 	}
+
+	return CreatedSignificantChange;
 }
 
 bool FGameplayTagRequirements::RequirementsMet(const FGameplayTagContainer& Container) const
@@ -547,7 +506,6 @@ FGameplayEffectSpecHandle::FGameplayEffectSpecHandle(FGameplayEffectSpec* DataPt
 FGameplayCueParameters::FGameplayCueParameters(const FGameplayEffectSpecForRPC& Spec)
 : NormalizedMagnitude(0.0f)
 , RawMagnitude(0.0f)
-, MatchedTagName(NAME_None)
 , Location(ForceInitToZero)
 , Normal(ForceInitToZero)
 {
@@ -557,7 +515,6 @@ FGameplayCueParameters::FGameplayCueParameters(const FGameplayEffectSpecForRPC&
 FGameplayCueParameters::FGameplayCueParameters(const struct FGameplayEffectContextHandle& InEffectContext)
 : NormalizedMagnitude(0.0f)
 , RawMagnitude(0.0f)
-, MatchedTagName(NAME_None)
 , Location(ForceInitToZero)
 , Normal(ForceInitToZero)
 {
@@ -571,15 +528,12 @@ bool FGameplayCueParameters::NetSerialize(FArchive& Ar, class UPackageMap* Map,
 		REP_NormalizedMagnitude = 0,
 		REP_RawMagnitude,
 		REP_EffectContext,
-		REP_MatchedTagName,
-		REP_OriginalTag,
-		REP_AggregatedSourceTags,
-		REP_AggregatedTargetTags,
 		REP_Location,
 		REP_Normal,
 		REP_Instigator,
 		REP_EffectCauser,
 		REP_SourceObject,
+		REP_PhysMaterial,
 		
 		REP_MAX
 	};
@@ -599,22 +553,6 @@ bool FGameplayCueParameters::NetSerialize(FArchive& Ar, class UPackageMap* Map,
 		{
 			RepBits |= (1 << REP_EffectContext);
 		}
-		if (MatchedTagName != NAME_None)
-		{
-			RepBits |= (1 << REP_MatchedTagName);
-		}
-		if (OriginalTag.IsValid())
-		{
-			RepBits |= (1 << REP_OriginalTag);
-		}
-		if (AggregatedSourceTags.Num() > 0)
-		{
-			RepBits |= (1 << REP_AggregatedSourceTags);
-		}
-		if (AggregatedTargetTags.Num() > 0)
-		{
-			RepBits |= (1 << REP_AggregatedTargetTags);
-		}
 		if (Location.IsNearlyZero() == false)
 		{
 			RepBits |= (1 << REP_Location);
@@ -635,10 +573,18 @@ bool FGameplayCueParameters::NetSerialize(FArchive& Ar, class UPackageMap* Map,
 		{
 			RepBits |= (1 << REP_SourceObject);
 		}
+		if (PhysicalMaterial.IsValid())
+		{
+			RepBits |= (1 << REP_PhysMaterial);
+		}
 	}
 
 	Ar.SerializeBits(&RepBits, REP_MAX);
 
+	// Tag containers serialize empty containers with 1 bit, so no need to serialize this in the RepBits field.
+	AggregatedSourceTags.NetSerialize(Ar, Map, bOutSuccess);
+	AggregatedTargetTags.NetSerialize(Ar, Map, bOutSuccess);
+
 	if (RepBits & (1 << REP_NormalizedMagnitude))
 	{
 		Ar << NormalizedMagnitude;
@@ -651,22 +597,6 @@ bool FGameplayCueParameters::NetSerialize(FArchive& Ar, class UPackageMap* Map,
 	{
 		EffectContext.NetSerialize(Ar, Map, bOutSuccess);
 	}
-	if (RepBits & (1 << REP_MatchedTagName))
-	{
-		Ar << MatchedTagName;
-	}
-	if (RepBits & (1 << REP_OriginalTag))
-	{
-		Ar << OriginalTag;
-	}
-	if (RepBits & (1 << REP_AggregatedSourceTags))
-	{
-		AggregatedSourceTags.NetSerialize(Ar, Map, bOutSuccess);
-	}
-	if (RepBits & (1 << REP_AggregatedTargetTags))
-	{
-		AggregatedTargetTags.NetSerialize(Ar, Map, bOutSuccess);
-	}
 	if (RepBits & (1 << REP_Location))
 	{
 		Location.NetSerialize(Ar, Map, bOutSuccess);
@@ -687,6 +617,10 @@ bool FGameplayCueParameters::NetSerialize(FArchive& Ar, class UPackageMap* Map,
 	{
 		Ar << SourceObject;
 	}
+	if (RepBits & (1 << REP_SourceObject))
+	{
+		Ar << PhysicalMaterial;
+	}
 
 	bOutSuccess = true;
 	return true;
@@ -761,4 +695,63 @@ const UObject* FGameplayCueParameters::GetSourceObject() const
 
 	// Fallback to effect context if the explicit data on gameplaycue parameters is not there.
 	return EffectContext.GetSourceObject();
+}
+
+bool FMinimapReplicationTagCountMap::NetSerialize(FArchive& Ar, class UPackageMap* Map, bool& bOutSuccess)
+{
+	const int32 CountBits = 4;
+	const int32 MaxCount = ((1 << CountBits)-1);
+
+	if (Ar.IsSaving())
+	{
+		int32 Count = TagMap.Num();
+		if (Count > MaxCount)
+		{
+			ABILITY_LOG(Error, TEXT("FMinimapReplicationTagCountMap has too many tags (%d)"), TagMap.Num());
+			Count = MaxCount;
+		}
+
+		Ar.SerializeBits(&Count, CountBits);
+		for(auto& It : TagMap)
+		{
+			FGameplayTag& Tag = It.Key;
+			Tag.NetSerialize(Ar, Map, bOutSuccess);
+			if (--Count <= 0)
+			{
+				break;
+			}
+		}
+	}
+	else
+	{
+		int32 Count = TagMap.Num();
+		Ar.SerializeBits(&Count, CountBits);
+
+		// Reset our local map
+		for(auto& It : TagMap)
+		{
+			It.Value = 0;
+		}
+
+		// See what we have
+		while(Count-- > 0)
+		{
+			FGameplayTag Tag;
+			Tag.NetSerialize(Ar, Map, bOutSuccess);
+			TagMap.FindOrAdd(Tag) = 1;
+		}
+
+		if (Owner)
+		{
+			// Update our tags with owner tags
+			for(auto& It : TagMap)
+			{
+				Owner->SetTagMapCount(It.Key, It.Value);
+			}
+		}
+	}
+
+
+	bOutSuccess = true;
+	return true;
 }
\ No newline at end of file
diff --git a/Engine/Source/Runtime/GameplayAbilities/Public/Abilities/GameplayAbility.h b/Engine/Source/Runtime/GameplayAbilities/Public/Abilities/GameplayAbility.h
--- a/Engine/Source/Runtime/GameplayAbilities/Public/Abilities/GameplayAbility.h
+++ b/Engine/Source/Runtime/GameplayAbilities/Public/Abilities/GameplayAbility.h
@@ -584,7 +584,7 @@ public:
 	UAnimMontage* GetCurrentMontage() const;
 
 	/** Call to set/get the current montage from a montage task. Set to allow hooking up montage events to ability events */
-	void SetCurrentMontage(class UAnimMontage* InCurrentMontage);
+	virtual void SetCurrentMontage(class UAnimMontage* InCurrentMontage);
 
 	/** Returns true if this ability can be canceled */
 	virtual bool CanBeCanceled() const;
diff --git a/Engine/Source/Runtime/GameplayAbilities/Public/AbilitySystemComponent.h b/Engine/Source/Runtime/GameplayAbilities/Public/AbilitySystemComponent.h
--- a/Engine/Source/Runtime/GameplayAbilities/Public/AbilitySystemComponent.h
+++ b/Engine/Source/Runtime/GameplayAbilities/Public/AbilitySystemComponent.h
@@ -167,12 +167,19 @@ class GAMEPLAYABILITIES_API UAbilitySystemComponent : public UGameplayTasksCompo
 	/** Force owning actor to update it's replication, to make sure that gameplay cues get sent down quickly. Override to change how aggressive this is */
 	virtual void ForceReplication();
 
+	void PreReplication(IRepChangedPropertyTracker& ChangedPropertyTracker) override;
+
 	virtual void GetSubobjectsWithStableNamesForNetworking(TArray<UObject*>& Objs) override;
 
 	virtual void PreNetReceive() override;
 	
 	virtual void PostNetReceive() override;
 
+	/** When true, we will not replicate active gameplay effects for this abiltiy system component, so attributes and tags */
+	void SetMinimalReplication(bool bNewMinimalReplication);
+
+	bool bMinimalReplication;
+
 	/** PredictionKeys, see more info in GameplayPrediction.h */
 	UPROPERTY(ReplicatedUsing=OnRep_PredictionKey)
 	FPredictionKey	ReplicatedPredictionKey;
@@ -356,7 +363,10 @@ class GAMEPLAYABILITIES_API UAbilitySystemComponent : public UGameplayTasksCompo
 	 * 
 	 * @return Source tags from the gameplay spec represented by the handle, if possible
 	 */
-	const FGameplayTagContainer* GetGameplayEffectSourceTagsFromHandle(FActiveGameplayEffectHandle Handle) const;
+	const FGameplayTagContainer* GetGameplayEffectSourceTagsFromHandle(FActiveGameplayEffectHandle Handle) const
+	{
+		return ActiveGameplayEffects.GetGameplayEffectSourceTagsFromHandle(Handle);
+	}
 
 	/**
 	 * Get the target tags from the gameplay spec represented by the specified handle, if possible
@@ -365,14 +375,25 @@ class GAMEPLAYABILITIES_API UAbilitySystemComponent : public UGameplayTasksCompo
 	 * 
 	 * @return Target tags from the gameplay spec represented by the handle, if possible
 	 */
-	const FGameplayTagContainer* GetGameplayEffectTargetTagsFromHandle(FActiveGameplayEffectHandle Handle) const;
+	const FGameplayTagContainer* GetGameplayEffectTargetTagsFromHandle(FActiveGameplayEffectHandle Handle) const
+	{
+		return ActiveGameplayEffects.GetGameplayEffectTargetTagsFromHandle(Handle);
+	}
 
 	/**
 	 * Populate the specified capture spec with the data necessary to capture an attribute from the component
 	 * 
 	 * @param OutCaptureSpec	[OUT] Capture spec to populate with captured data
 	 */
-	void CaptureAttributeForGameplayEffect(OUT FGameplayEffectAttributeCaptureSpec& OutCaptureSpec);
+	void CaptureAttributeForGameplayEffect(OUT FGameplayEffectAttributeCaptureSpec& OutCaptureSpec)
+	{
+		// Verify the capture is happening on an attribute the component actually has a set for; if not, can't capture the value
+		const FGameplayAttribute& AttributeToCapture = OutCaptureSpec.BackingDefinition.AttributeToCapture;
+		if (AttributeToCapture.IsValid() && (AttributeToCapture.IsSystemAttribute() || GetAttributeSubobject(AttributeToCapture.GetAttributeSetClass())))
+		{
+			ActiveGameplayEffects.CaptureAttributeForGameplayEffect(OutCaptureSpec);
+		}
+	}
 	
 	// --------------------------------------------
 	// Callbacks / Notifies
@@ -405,13 +426,25 @@ class GAMEPLAYABILITIES_API UAbilitySystemComponent : public UGameplayTasksCompo
 	// --------------------------------------------
 	// Tags
 	// --------------------------------------------
-	virtual bool HasMatchingGameplayTag(FGameplayTag TagToCheck) const override;
+	FORCEINLINE bool HasMatchingGameplayTag(FGameplayTag TagToCheck) const override
+	{
+		return GameplayTagCountContainer.HasMatchingGameplayTag(TagToCheck);
+	}
 
-	virtual bool HasAllMatchingGameplayTags(const FGameplayTagContainer& TagContainer, bool bCountEmptyAsMatch = true) const override;
+	FORCEINLINE bool HasAllMatchingGameplayTags(const FGameplayTagContainer& TagContainer, bool bCountEmptyAsMatch = true) const override
+	{
+		return GameplayTagCountContainer.HasAllMatchingGameplayTags(TagContainer, bCountEmptyAsMatch);
+	}
 
-	virtual bool HasAnyMatchingGameplayTags(const FGameplayTagContainer& TagContainer, bool bCountEmptyAsMatch = true) const override;
+	FORCEINLINE bool HasAnyMatchingGameplayTags(const FGameplayTagContainer& TagContainer, bool bCountEmptyAsMatch = true) const override
+	{
+		return GameplayTagCountContainer.HasAnyMatchingGameplayTags(TagContainer, bCountEmptyAsMatch);
+	}
 
-	virtual void GetOwnedGameplayTags(FGameplayTagContainer& TagContainer) const override;
+	FORCEINLINE void GetOwnedGameplayTags(FGameplayTagContainer& TagContainer) const override
+	{
+		TagContainer.AppendTags(GameplayTagCountContainer.GetExplicitGameplayTags());
+	}
 
 	/** 	 
 	 *  Allows GameCode to add loose gameplaytags which are not backed by a GameplayEffect. 
@@ -421,13 +454,50 @@ class GAMEPLAYABILITIES_API UAbilitySystemComponent : public UGameplayTasksCompo
 	 *	It is up to the calling GameCode to make sure these tags are added on clients/server where necessary
 	 */
 
-	void AddLooseGameplayTag(const FGameplayTag& GameplayTag, int32 Count=1);
+	FORCEINLINE void AddLooseGameplayTag(const FGameplayTag& GameplayTag, int32 Count=1)
+	{
+		UpdateTagMap(GameplayTag, Count);
+	}
+
+	FORCEINLINE void AddLooseGameplayTags(const FGameplayTagContainer& GameplayTags, int32 Count = 1)
+	{
+		UpdateTagMap(GameplayTags, Count);
+	}
+
+	FORCEINLINE void RemoveLooseGameplayTag(const FGameplayTag& GameplayTag, int32 Count = 1)
+	{
+		UpdateTagMap(GameplayTag, -Count);
+	}
+
+	FORCEINLINE void RemoveLooseGameplayTags(const FGameplayTagContainer& GameplayTags, int32 Count = 1)
+	{
+		UpdateTagMap(GameplayTags, -Count);
+	}
+
+	/** 	 
+	 * Minimally replicated tags are replicated tags that come from GEs when in bMinimalReplication mode. 
+	 * (The GEs do not replicate, but the tags they grant do replicate via these functions)
+	 */
+
+	FORCEINLINE void AddMinimalReplicationGameplayTag(const FGameplayTag& GameplayTag)
+	{
+		MinimalReplicationTags.AddTag(GameplayTag);
+	}
 
-	void AddLooseGameplayTags(const FGameplayTagContainer& GameplayTag, int32 Count = 1);
+	FORCEINLINE void AddMinimalReplicationGameplayTags(const FGameplayTagContainer& GameplayTags)
+	{
+		MinimalReplicationTags.AddTags(GameplayTags);
+	}
 
-	void RemoveLooseGameplayTag(const FGameplayTag& GameplayTag, int32 Count = 1);
+	FORCEINLINE void RemoveMinimalReplicationGameplayTag(const FGameplayTag& GameplayTag)
+	{
+		MinimalReplicationTags.RemoveTag(GameplayTag);
+	}
 
-	void RemoveLooseGameplayTags(const FGameplayTagContainer& GameplayTag, int32 Count = 1);
+	FORCEINLINE void RemoveMinimalReplicationGameplayTags(const FGameplayTagContainer& GameplayTags)
+	{
+		MinimalReplicationTags.RemoveTags(GameplayTags);
+	}
 	
 	/** Allow events to be registered for specific gameplay tags being added or removed */
 	FOnGameplayEffectTagCountChanged& RegisterGameplayTagEvent(FGameplayTag Tag, EGameplayTagEventType::Type EventType=EGameplayTagEventType::NewOrRemoved);
@@ -480,15 +550,27 @@ class GAMEPLAYABILITIES_API UAbilitySystemComponent : public UGameplayTasksCompo
 	FActiveGameplayEffectHandle ApplyGameplayEffectToSelf(const UGameplayEffect *GameplayEffect, float Level, const FGameplayEffectContextHandle& EffectContext, FPredictionKey PredictionKey = FPredictionKey());
 
 	// Returns the number of gameplay effects that are currently active on this ability system component
-	int32 GetNumActiveGameplayEffects() const;
+	int32 GetNumActiveGameplayEffects() const
+	{
+		return ActiveGameplayEffects.GetNumGameplayEffects();
+	}
 
 	// Makes a copy of all the active effects on this ability component
-	void GetAllActiveGameplayEffectSpecs(TArray<FGameplayEffectSpec>& OutSpecCopies);
+	void GetAllActiveGameplayEffectSpecs(TArray<FGameplayEffectSpec>& OutSpecCopies)
+	{
+		ActiveGameplayEffects.GetAllActiveGameplayEffectSpecs(OutSpecCopies);
+	}
 
-	void SetBaseAttributeValueFromReplication(float NewValue, FGameplayAttribute Attribute);
+	void SetBaseAttributeValueFromReplication(float NewValue, FGameplayAttribute Attribute)
+	{
+		ActiveGameplayEffects.SetBaseAttributeValueFromReplication(Attribute, NewValue);
+	}
 
 	/** Tests if all modifiers in this GameplayEffect will leave the attribute > 0.f */
-	bool CanApplyAttributeModifiers(const UGameplayEffect *GameplayEffect, float Level, const FGameplayEffectContextHandle& EffectContext);
+	bool CanApplyAttributeModifiers(const UGameplayEffect *GameplayEffect, float Level, const FGameplayEffectContextHandle& EffectContext)
+	{
+		return ActiveGameplayEffects.CanApplyAttributeModifiers(GameplayEffect, Level, EffectContext);
+	}
 
 	// Generic 'Get expected magnitude (list) if I was to apply this outgoing or incoming'
 
@@ -509,9 +591,15 @@ class GAMEPLAYABILITIES_API UAbilitySystemComponent : public UGameplayTasksCompo
 	TArray<FActiveGameplayEffectHandle> GetActiveEffects(const FGameplayEffectQuery& Query) const;
 
 	/** This will give the world time that all effects matching this query will be finished. If multiple effects match, it returns the one that returns last.*/
-	float GetActiveEffectsEndTime(const FGameplayEffectQuery& Query) const;
+	float GetActiveEffectsEndTime(const FGameplayEffectQuery& Query) const
+	{
+		return ActiveGameplayEffects.GetActiveEffectsEndTime(Query);
+	}
 
-	void ModifyActiveEffectStartTime(FActiveGameplayEffectHandle Handle, float StartTimeDiff);
+	void ModifyActiveEffectStartTime(FActiveGameplayEffectHandle Handle, float StartTimeDiff)
+	{
+		ActiveGameplayEffects.ModifyActiveEffectStartTime(Handle, StartTimeDiff);
+	}
 
 	/** Removes all active effects that contain any of the tags in Tags */
 	UFUNCTION(BlueprintCallable, Category = GameplayEffects)
@@ -538,8 +626,13 @@ class GAMEPLAYABILITIES_API UAbilitySystemComponent : public UGameplayTasksCompo
 	
 	void PrintAllGameplayEffects() const;
 
+	bool CachedIsNetSimulated;
+
 	/** Returns true of this component has authority */
-	bool IsOwnerActorAuthoritative() const;
+	bool IsOwnerActorAuthoritative() const
+	{
+		return !CachedIsNetSimulated;
+	}
 
 	// ----------------------------------------------------------------------------------------------------------------
 	//
@@ -554,24 +647,64 @@ class GAMEPLAYABILITIES_API UAbilitySystemComponent : public UGameplayTasksCompo
 	UFUNCTION(NetMulticast, unreliable)
 	void NetMulticast_InvokeGameplayCueExecuted(const FGameplayTag GameplayCueTag, FPredictionKey PredictionKey, FGameplayEffectContextHandle EffectContext);
 
+	UFUNCTION(NetMulticast, unreliable)
+	void NetMulticast_InvokeGameplayCuesExecuted(const FGameplayTagContainer GameplayCueTags, FPredictionKey PredictionKey, FGameplayEffectContextHandle EffectContext);
+
 	UFUNCTION(NetMulticast, unreliable)
 	void NetMulticast_InvokeGameplayCueExecuted_WithParams(const FGameplayTag GameplayCueTag, FPredictionKey PredictionKey, FGameplayCueParameters GameplayCueParameters);
 
+	UFUNCTION(NetMulticast, unreliable)
+	void NetMulticast_InvokeGameplayCuesExecuted_WithParams(const FGameplayTagContainer GameplayCueTags, FPredictionKey PredictionKey, FGameplayCueParameters GameplayCueParameters);
+
 	UFUNCTION(NetMulticast, unreliable)
 	void NetMulticast_InvokeGameplayCueAdded(const FGameplayTag GameplayCueTag, FPredictionKey PredictionKey, FGameplayEffectContextHandle EffectContext);
 	
 	UFUNCTION(NetMulticast, unreliable)
 	void NetMulticast_InvokeGameplayCueAddedAndWhileActive_FromSpec(const FGameplayEffectSpecForRPC& Spec, FPredictionKey PredictionKey);
 
+	UFUNCTION(NetMulticast, unreliable)
+	void NetMulticast_InvokeGameplayCueAddedAndWhileActive_WithParams(const FGameplayTag GameplayCueTag, FPredictionKey PredictionKey, FGameplayCueParameters GameplayCueParameters);
+
+	UFUNCTION(NetMulticast, unreliable)
+	void NetMulticast_InvokeGameplayCuesAddedAndWhileActive_WithParams(const FGameplayTagContainer GameplayCueTags, FPredictionKey PredictionKey, FGameplayCueParameters GameplayCueParameters);
+
 	// GameplayCues can also come on their own. These take an optional effect context to pass through hit result, etc
 	void ExecuteGameplayCue(const FGameplayTag GameplayCueTag, FGameplayEffectContextHandle EffectContext = FGameplayEffectContextHandle());
 
 	// This version allows the caller to set an explicit FGameplayCueParmeters.
 	void ExecuteGameplayCue(const FGameplayTag GameplayCueTag, const FGameplayCueParameters& GameplayCueParameters);
 
-	void AddGameplayCue(const FGameplayTag GameplayCueTag, FGameplayEffectContextHandle EffectContext = FGameplayEffectContextHandle());
+	// -------------------------
+
+	void AddGameplayCue(const FGameplayTag GameplayCueTag, FGameplayEffectContextHandle EffectContext = FGameplayEffectContextHandle())
+	{
+		AddGameplayCue_Internal(GameplayCueTag, EffectContext, ActiveGameplayCues);
+	}
+
+	/** Add gameplaycue for minimal replication mode. Should only be called in paths that would replicate gameplaycues in other ways (through GE for example) if not in minimal replication mode */
+	void AddGameplayCue_MinimalReplication(const FGameplayTag GameplayCueTag, FGameplayEffectContextHandle EffectContext = FGameplayEffectContextHandle())
+	{
+		AddGameplayCue_Internal(GameplayCueTag, EffectContext, MinimalReplicationGameplayCues);
+	}
+
+	void AddGameplayCue_Internal(const FGameplayTag GameplayCueTag, const FGameplayEffectContextHandle& EffectContext, FActiveGameplayCueContainer& GameplayCueContainer );
+
+	// -------------------------
 	
-	void RemoveGameplayCue(const FGameplayTag GameplayCueTag);
+	void RemoveGameplayCue(const FGameplayTag GameplayCueTag)
+	{
+		RemoveGameplayCue_Internal(GameplayCueTag, ActiveGameplayCues);
+	}
+
+	/** Remove gameplaycue for minimal replication mode. Should only be called in paths that would replicate gameplaycues in other ways (through GE for example) if not in minimal replication mode */
+	void RemoveGameplayCue_MinimalReplication(const FGameplayTag GameplayCueTag)
+	{
+		RemoveGameplayCue_Internal(GameplayCueTag, MinimalReplicationGameplayCues);
+	}
+
+	void RemoveGameplayCue_Internal(const FGameplayTag GameplayCueTag, FActiveGameplayCueContainer& GameplayCueContainer);
+
+	// -------------------------
 
 	/** Removes any GameplayCue added on its own, i.e. not as part of a GameplayEffect. */
 	void RemoveAllGameplayCues();
@@ -584,8 +717,10 @@ class GAMEPLAYABILITIES_API UAbilitySystemComponent : public UGameplayTasksCompo
 
 	/** Allows polling to see if a GameplayCue is active. We expect most GameplayCue handling to be event based, but some cases we may need to check if a GamepalyCue is active (Animation Blueprint for example) */
 	UFUNCTION(BlueprintCallable, Category="GameplayCue", meta=(GameplayTagFilter="GameplayCue"))
-	bool IsGameplayCueActive(const FGameplayTag GameplayCueTag) const;
-
+	bool IsGameplayCueActive(const FGameplayTag GameplayCueTag) const
+	{
+		return HasMatchingGameplayTag(GameplayCueTag);
+	}
 
 	// ----------------------------------------------------------------------------------------------------------------
 
@@ -1161,6 +1296,29 @@ public:
 
 	// ---------------------------------------------------------------------
 
+	FORCEINLINE void SetTagMapCount(const FGameplayTag& Tag, int32 NewCount)
+	{
+		GameplayTagCountContainer.SetTagCount(Tag, NewCount);
+	}
+	
+	FORCEINLINE void UpdateTagMap(const FGameplayTag& BaseTag, int32 CountDelta)
+	{
+		if (GameplayTagCountContainer.UpdateTagCount(BaseTag, CountDelta))
+		{
+			OnTagUpdated(BaseTag, CountDelta > 0);
+		}
+	}
+	
+	FORCEINLINE void UpdateTagMap(const FGameplayTagContainer& Container, int32 CountDelta)
+	{
+		for (auto TagIt = Container.CreateConstIterator(); TagIt; ++TagIt)
+		{
+			const FGameplayTag& Tag = *TagIt;
+			UpdateTagMap(Tag, CountDelta);
+		}
+	}	
+
+
 #if ENABLE_VISUAL_LOG
 	void ClearDebugInstantEffects();
 #endif // ENABLE_VISUAL_LOG
@@ -1191,12 +1349,16 @@ protected:
 	// --------------------------------------------
 	
 	// Contains all of the gameplay effects that are currently active on this component
-	UPROPERTY(ReplicatedUsing=OnRep_GameplayEffects)
+	UPROPERTY(Replicated)
 	FActiveGameplayEffectsContainer	ActiveGameplayEffects;
 
-	UPROPERTY(ReplicatedUsing=OnRep_GameplayEffects)
+	UPROPERTY(Replicated)
 	FActiveGameplayCueContainer	ActiveGameplayCues;
 
+	/** Replicated gameplaycues when in minimal replication mode. These are cues that would come normally come from ActiveGameplayEffects */
+	UPROPERTY(Replicated)
+	FActiveGameplayCueContainer	MinimalReplicationGameplayCues;
+
 	/** Abilities with these tags are not able to be activated */
 	FGameplayTagCountContainer BlockedAbilityTags;
 
@@ -1204,9 +1366,6 @@ protected:
 	UPROPERTY(Transient, Replicated)
 	TArray<uint8> BlockedAbilityBindings;
 
-	UFUNCTION()
-	void OnRep_GameplayEffects();
-
 	void DebugCyclicAggregatorBroadcasts(struct FAggregator* Aggregator);
 
 	// ---------------------------------------------
@@ -1214,9 +1373,8 @@ protected:
 	// Acceleration map for all gameplay tags (OwnedGameplayTags from GEs and explicit GameplayCueTags)
 	FGameplayTagCountContainer GameplayTagCountContainer;
 
-	void UpdateTagMap(const FGameplayTag& BaseTag, int32 CountDelta);
-	
-	void UpdateTagMap(const FGameplayTagContainer& Container, int32 CountDelta);
+	UPROPERTY(Replicated)
+	FMinimapReplicationTagCountMap MinimalReplicationTags;
 
 	void ResetTagMap();
 
diff --git a/Engine/Source/Runtime/GameplayAbilities/Public/AbilitySystemGlobals.h b/Engine/Source/Runtime/GameplayAbilities/Public/AbilitySystemGlobals.h
--- a/Engine/Source/Runtime/GameplayAbilities/Public/AbilitySystemGlobals.h
+++ b/Engine/Source/Runtime/GameplayAbilities/Public/AbilitySystemGlobals.h
@@ -171,6 +171,7 @@ class GAMEPLAYABILITIES_API UAbilitySystemGlobals : public UObject
 
 	// GameplayCue Parameters
 	virtual void InitGameplayCueParameters(FGameplayCueParameters& CueParameters, const FGameplayEffectSpecForRPC &Spec);
+	virtual void InitGameplayCueParameters_GESpec(FGameplayCueParameters& CueParameters, const FGameplayEffectSpec &Spec);
 	virtual void InitGameplayCueParameters(FGameplayCueParameters& CueParameters, const FGameplayEffectContextHandle& EffectContext);
 
 	// Trigger async loading of the gameplay cue object libraries. By default, the manager will do this on creation,
@@ -248,6 +249,8 @@ private:
 	void OnTableReimported(UObject* InObject);
 #endif
 
+	void HandlePostWorldCreate(UWorld* NewWorld);
+
 #if WITH_EDITORONLY_DATA
 	bool RegisteredReimportCallback;
 #endif
diff --git a/Engine/Source/Runtime/GameplayAbilities/Public/AbilitySystemStats.h b/Engine/Source/Runtime/GameplayAbilities/Public/AbilitySystemStats.h
--- a/Engine/Source/Runtime/GameplayAbilities/Public/AbilitySystemStats.h
+++ b/Engine/Source/Runtime/GameplayAbilities/Public/AbilitySystemStats.h
@@ -32,4 +32,5 @@ DECLARE_CYCLE_STAT_EXTERN(TEXT("GameplayCueNotify Static"), STAT_HandleGameplayC
 DECLARE_CYCLE_STAT_EXTERN(TEXT("GameplayCueNotify Actor"), STAT_HandleGameplayCueNotifyActor, STATGROUP_AbilitySystem, );
 DECLARE_CYCLE_STAT_EXTERN(TEXT("ApplyGameplayEffectToTarget"), STAT_ApplyGameplayEffectToTarget, STATGROUP_AbilitySystem, );
 DECLARE_CYCLE_STAT_EXTERN(TEXT("ActiveGameplayEffect Added"), STAT_OnActiveGameplayEffectAdded, STATGROUP_AbilitySystem, );
-DECLARE_CYCLE_STAT_EXTERN(TEXT("ActiveGameplayEffect Removed"), STAT_OnActiveGameplayEffectRemoved, STATGROUP_AbilitySystem, );
\ No newline at end of file
+DECLARE_CYCLE_STAT_EXTERN(TEXT("ActiveGameplayEffect Removed"), STAT_OnActiveGameplayEffectRemoved, STATGROUP_AbilitySystem, );
+DECLARE_CYCLE_STAT_EXTERN(TEXT("GameplayCueInterface HandleGameplayCue"), STAT_GameplayCueInterface_HandleGameplayCue, STATGROUP_AbilitySystem, );
diff --git a/Engine/Source/Runtime/GameplayAbilities/Public/GameplayCueInterface.h b/Engine/Source/Runtime/GameplayAbilities/Public/GameplayCueInterface.h
--- a/Engine/Source/Runtime/GameplayAbilities/Public/GameplayCueInterface.h
+++ b/Engine/Source/Runtime/GameplayAbilities/Public/GameplayCueInterface.h
@@ -39,7 +39,9 @@ class GAMEPLAYABILITIES_API IGameplayCueInterface
 	UFUNCTION(BlueprintCallable, BlueprintCosmetic, Category="Ability|GameplayCue")
 	virtual void ForwardGameplayCueToParent();
 
-	static void DispatchBlueprintCustomHandler(AActor* Actor, UFunction* Func, EGameplayCueEvent::Type EventType, FGameplayCueParameters Parameters);	
+	static void DispatchBlueprintCustomHandler(AActor* Actor, UFunction* Func, EGameplayCueEvent::Type EventType, FGameplayCueParameters Parameters);
+
+	static void ClearTagToFunctionMap();
 
 	IGameplayCueInterface() : bForwardToParent(false) {}
 
@@ -97,6 +99,9 @@ struct FActiveGameplayCueContainer : public FFastArraySerializer
 	UPROPERTY()
 	class UAbilitySystemComponent*	Owner;
 
+	/** Should this container only rpelicate in minimal replication mode */
+	bool bMinimalReplication;
+
 	void AddCue(const FGameplayTag& Tag, const FPredictionKey& PredictionKey);
 	void RemoveCue(const FGameplayTag& Tag);
 
@@ -108,10 +113,7 @@ struct FActiveGameplayCueContainer : public FFastArraySerializer
 	/** Does explicit check for gameplay cue tag */
 	bool HasCue(const FGameplayTag& Tag) const;
 
-	bool NetDeltaSerialize(FNetDeltaSerializeInfo & DeltaParms)
-	{
-		return FastArrayDeltaSerialize<FActiveGameplayCue>(GameplayCues, DeltaParms, *this);
-	}
+	bool NetDeltaSerialize(FNetDeltaSerializeInfo & DeltaParms);
 
 private:
 
diff --git a/Engine/Source/Runtime/GameplayAbilities/Public/GameplayCueManager.h b/Engine/Source/Runtime/GameplayAbilities/Public/GameplayCueManager.h
--- a/Engine/Source/Runtime/GameplayAbilities/Public/GameplayCueManager.h
+++ b/Engine/Source/Runtime/GameplayAbilities/Public/GameplayCueManager.h
@@ -60,10 +60,12 @@ class GAMEPLAYABILITIES_API UGameplayCueManager : public UDataAsset
 	// Wrappers to handle replicating executed cues
 	// -------------------------------------------------------------
 
-	virtual void InvokeGameplayCueExecuted_FromSpec(UAbilitySystemComponent* OwningComponent, const FGameplayEffectSpecForRPC Spec, FPredictionKey PredictionKey);
+	virtual void InvokeGameplayCueExecuted_FromSpec(UAbilitySystemComponent* OwningComponent, const FGameplayEffectSpec& Spec, FPredictionKey PredictionKey);
 	virtual void InvokeGameplayCueExecuted(UAbilitySystemComponent* OwningComponent, const FGameplayTag GameplayCueTag, FPredictionKey PredictionKey, FGameplayEffectContextHandle EffectContext);
 	virtual void InvokeGameplayCueExecuted_WithParams(UAbilitySystemComponent* OwningComponent, const FGameplayTag GameplayCueTag, FPredictionKey PredictionKey, FGameplayCueParameters GameplayCueParameters);
 
+	virtual void InvokeGameplayCueAddedAndWhileActive_FromSpec(UAbilitySystemComponent* OwningComponent, const FGameplayEffectSpec& Spec, FPredictionKey PredictionKey);
+
 	/** Start or stop a gameplay cue send context. Used by FScopedGameplayCueSendContext above, when all contexts are removed the cues are flushed */
 	void StartGameplayCueSendContext();
 	void EndGameplayCueSendContext();
diff --git a/Engine/Source/Runtime/GameplayAbilities/Public/GameplayCueNotify_Actor.h b/Engine/Source/Runtime/GameplayAbilities/Public/GameplayCueNotify_Actor.h
--- a/Engine/Source/Runtime/GameplayAbilities/Public/GameplayCueNotify_Actor.h
+++ b/Engine/Source/Runtime/GameplayAbilities/Public/GameplayCueNotify_Actor.h
@@ -105,8 +105,6 @@ class GAMEPLAYABILITIES_API AGameplayCueNotify_Actor : public AActor
 
 	FGCNotifyActorKey NotifyKey;
 	
-	int32 StackCount;
-
 protected:
 	FTimerHandle FinishTimerHandle;
 
diff --git a/Engine/Source/Runtime/GameplayAbilities/Public/GameplayCueSet.h b/Engine/Source/Runtime/GameplayAbilities/Public/GameplayCueSet.h
--- a/Engine/Source/Runtime/GameplayAbilities/Public/GameplayCueSet.h
+++ b/Engine/Source/Runtime/GameplayAbilities/Public/GameplayCueSet.h
@@ -80,6 +80,6 @@ class GAMEPLAYABILITIES_API UGameplayCueSet : public UDataAsset
 	static FGameplayTag	BaseGameplayCueTag();
 
 protected:
-	virtual bool HandleGameplayCueNotify_Internal(AActor* TargetActor, int32 DataIdx, EGameplayCueEvent::Type EventType, const FGameplayCueParameters& Parameters);
-	virtual void BuildAccelerationMap_Internal();	
+	virtual bool HandleGameplayCueNotify_Internal(AActor* TargetActor, int32 DataIdx, EGameplayCueEvent::Type EventType, FGameplayCueParameters& Parameters);
+	virtual void BuildAccelerationMap_Internal();
 };
\ No newline at end of file
diff --git a/Engine/Source/Runtime/GameplayAbilities/Public/GameplayCue_Types.h b/Engine/Source/Runtime/GameplayAbilities/Public/GameplayCue_Types.h
--- a/Engine/Source/Runtime/GameplayAbilities/Public/GameplayCue_Types.h
+++ b/Engine/Source/Runtime/GameplayAbilities/Public/GameplayCue_Types.h
@@ -28,9 +28,8 @@ struct FGameplayCuePendingExecute
 	, OwningComponent(NULL)
 	{
 	}
-
-	UPROPERTY()
-	FGameplayTag GameplayCueTag;
+	
+	TArray<FGameplayTag, TInlineAllocator<1> > GameplayCueTags;
 	
 	/** Prediction key that spawned this cue */
 	UPROPERTY()
