Moved GetAbility call to cpp to address compile issues

diff --git a/Engine/Source/Runtime/GameplayAbilities/Private/GameplayEffectTypes.cpp b/Engine/Source/Runtime/GameplayAbilities/Private/GameplayEffectTypes.cpp
--- a/Engine/Source/Runtime/GameplayAbilities/Private/GameplayEffectTypes.cpp
+++ b/Engine/Source/Runtime/GameplayAbilities/Private/GameplayEffectTypes.cpp
@@ -78,6 +78,11 @@ void FGameplayEffectContext::SetAbility(const UGameplayAbility* InGameplayAbilit
 	}
 }
 
+const UGameplayAbility* FGameplayEffectContext::GetAbility() const
+{
+	return Ability.GetDefaultObject();
+}
+
 void FGameplayEffectContext::AddActors(const TArray<TWeakObjectPtr<AActor>>& InActors, bool bReset)
 {
 	if (bReset && Actors.Num())
diff --git a/Engine/Source/Runtime/GameplayAbilities/Public/GameplayEffectTypes.h b/Engine/Source/Runtime/GameplayAbilities/Public/GameplayEffectTypes.h
--- a/Engine/Source/Runtime/GameplayAbilities/Public/GameplayEffectTypes.h
+++ b/Engine/Source/Runtime/GameplayAbilities/Public/GameplayEffectTypes.h
@@ -306,10 +306,7 @@ struct GAMEPLAYABILITIES_API FGameplayEffectContext
 	}
 
 	/** returns the CDO of the ability used to instigate this context */
-	const UGameplayAbility* GetAbility() const
-	{
-		return Ability.GetDefaultObject();
-	}
+	const UGameplayAbility* GetAbility() const;
 
 	int32 GetAbilityLevel() const
 	{
