Copying //UE4/Dev-Framework to //UE4/Dev-Main (Source: //UE4/Dev-Framework @ 3182037)

diff --git a/Engine/Source/Runtime/Engine/Public/BoneContainer.h b/Engine/Source/Runtime/Engine/Public/BoneContainer.h
--- a/Engine/Source/Runtime/Engine/Public/BoneContainer.h
+++ b/Engine/Source/Runtime/Engine/Public/BoneContainer.h
@@ -322,8 +322,14 @@ struct FBoneReference
 	/** Cached bone index for run time - right now bone index of skeleton **/
 	int32 BoneIndex;
 
+	/** Change this to Bitfield if we have more than one bool 
+	 * This specifies whether or not this indices is mesh or skeleton
+	 */
+	bool bUseSkeletonIndex;
+
 	FBoneReference()
 		: BoneIndex(INDEX_NONE)
+		, bUseSkeletonIndex(false)
 	{
 	}
 
@@ -335,14 +341,49 @@ struct FBoneReference
 	/** Initialize Bone Reference, return TRUE if success, otherwise, return false **/
 	ENGINE_API bool Initialize(const FBoneContainer& RequiredBones);
 
-	// @fixme laurent - only used by blendspace 'PerBoneBlend'. Fix this to support SkeletalMesh pose.
+	// only used by blendspace 'PerBoneBlend'. This is skeleton indices since the input is skeleton
+	// @note, if you use this function, it won't work with GetCompactPoseIndex, GetMeshPoseIndex;
+	// it triggers ensure in those functions
 	ENGINE_API bool Initialize(const USkeleton* Skeleton);
 
 	/** return true if valid. Otherwise return false **/
 	ENGINE_API bool IsValid(const FBoneContainer& RequiredBones) const;
 
-	FMeshPoseBoneIndex GetMeshPoseIndex() const { return FMeshPoseBoneIndex(BoneIndex); }
-	FCompactPoseBoneIndex GetCompactPoseIndex(const FBoneContainer& RequiredBones) const { return RequiredBones.MakeCompactPoseIndex(GetMeshPoseIndex()); }
+	FMeshPoseBoneIndex GetMeshPoseIndex(const FBoneContainer& RequiredBones) const
+	{ 
+		// accessing array with invalid index would cause crash, so we have to check here
+		if (BoneIndex != INDEX_NONE)
+		{
+			if (bUseSkeletonIndex)
+			{
+				return FMeshPoseBoneIndex(RequiredBones.GetSkeletonToPoseBoneIndexArray()[BoneIndex]);
+			}
+			else
+			{
+				return FMeshPoseBoneIndex(BoneIndex);
+			}
+		}
+
+		return FMeshPoseBoneIndex(INDEX_NONE);
+	}
+
+	FCompactPoseBoneIndex GetCompactPoseIndex(const FBoneContainer& RequiredBones) const 
+	{ 
+		// accessing array with invalid index would cause crash, so we have to check here
+		if (BoneIndex != INDEX_NONE)
+		{
+			if (bUseSkeletonIndex)
+			{
+				return RequiredBones.GetCompactPoseIndexFromSkeletonIndex(BoneIndex);
+			}
+			else
+			{
+				return RequiredBones.MakeCompactPoseIndex(GetMeshPoseIndex(RequiredBones));
+			}
+		}
+		
+		return FCompactPoseBoneIndex(INDEX_NONE);
+	}
 
 	// need this because of BoneReference being used in CurveMetaData and that is in SmartName
 	friend FArchive& operator<<(FArchive& Ar, FBoneReference& B)
diff --git a/Engine/Source/Runtime/Engine/Public/EngineGlobals.h b/Engine/Source/Runtime/Engine/Public/EngineGlobals.h
--- a/Engine/Source/Runtime/Engine/Public/EngineGlobals.h
+++ b/Engine/Source/Runtime/Engine/Public/EngineGlobals.h
@@ -12,6 +12,3 @@ extern ENGINE_API bool						GDisallowNetworkTravel;
 
 /** The GPU time taken to render the last frame. Same metric as FPlatformTime::Cycles(). */
 extern ENGINE_API uint32					GGPUFrameTime;
-
-/** A reference to the ability system so we do not need to get the module every time we access the tag manager. */
-extern ENGINE_API class UGameplayTagsManager* GGameplayTagsManager;
\ No newline at end of file
diff --git a/Engine/Source/Runtime/Engine/Public/OpusAudioInfo.h b/Engine/Source/Runtime/Engine/Public/OpusAudioInfo.h
--- a/Engine/Source/Runtime/Engine/Public/OpusAudioInfo.h
+++ b/Engine/Source/Runtime/Engine/Public/OpusAudioInfo.h
@@ -15,86 +15,21 @@ struct FOpusDecoderWrapper;
 /**
 * Helper class to parse opus data
 */
-class FOpusAudioInfo : public ICompressedAudioInfo
+class FOpusAudioInfo : public IStreamedCompressedInfo
 {
 public:
-	ENGINE_API FOpusAudioInfo(void);
-	ENGINE_API virtual ~FOpusAudioInfo(void);
+	ENGINE_API FOpusAudioInfo();
+	ENGINE_API virtual ~FOpusAudioInfo();
 
-	// ICompressedAudioInfo Interface
-	ENGINE_API virtual bool ReadCompressedInfo(const uint8* InSrcBufferData, uint32 InSrcBufferDataSize, struct FSoundQualityInfo* QualityInfo) override;
-	ENGINE_API virtual bool ReadCompressedData(uint8* Destination, bool bLooping, uint32 BufferSize) override;
-	ENGINE_API virtual void SeekToTime(const float SeekTime) override {};
-	ENGINE_API virtual void ExpandFile(uint8* DstBuffer, struct FSoundQualityInfo* QualityInfo) override;
-	ENGINE_API virtual void EnableHalfRate(bool HalfRate) override {};
-	virtual uint32 GetSourceBufferSize() const override { return SrcBufferDataSize;}
-	virtual bool UsesVorbisChannelOrdering() const override { return false; }
-	virtual int GetStreamBufferSize() const override { return  MONO_PCM_BUFFER_SIZE; }
-
-	virtual bool SupportsStreaming() const override {return true;}
-	virtual bool StreamCompressedInfo(USoundWave* Wave, struct FSoundQualityInfo* QualityInfo) override;
-	virtual bool StreamCompressedData(uint8* Destination, bool bLooping, uint32 BufferSize) override;
-	virtual int32 GetCurrentChunkIndex() const override {return CurrentChunkIndex;}
-	virtual int32 GetCurrentChunkOffset() const override {return SrcBufferOffset;}
-	// End of ICompressedAudioInfo Interface
+	//~ Begin IStreamedCompressedInfo Interface
+	virtual bool ParseHeader(const uint8* InSrcBufferData, uint32 InSrcBufferDataSize, FSoundQualityInfo* QualityInfo) override;
+	virtual uint32 GetFrameSize() override;
+	virtual uint32 GetMaxFrameSizeSamples() const override;
+	virtual bool CreateDecoder() override;
+	virtual int32 Decode(const uint8* FrameData, uint16 FrameSize, int16* OutPCMData, int32 SampleSize) override;
+	//~ End IStreamedCompressedInfo Interface
 
 protected:
-	/** Emulate read from memory functionality */
-	size_t			Read(void *ptr, uint32 size);
-
-	/**
-	 * Decompresses a frame of Opus data to PCM buffer
-	 *
-	 * @param FrameSize Size of the frame in bytes
-	 * @return The amount of samples that were decompressed (< 0 indicates error)
-	 */
-	int32 DecompressToPCMBuffer(uint16 FrameSize);
-
-	/**
-	 * Adds to the count of samples that have currently been decoded
-	 *
-	 * @param NewSamples	How many samples have been decoded
-	 * @return How many samples were actually part of the true sample count
-	 */
-	uint32 IncrementCurrentSampleCount(uint32 NewSamples);
-
-	/**
-	 * Writes data from decoded PCM buffer, taking into account whether some PCM has been written before
-	 *
-	 * @param Destination	Where to place the decoded sound
-	 * @param BufferSize	Size of the destination buffer in bytes
-	 * @return				How many bytes were written
-	 */
-	uint32	WriteFromDecodedPCM(uint8* Destination, uint32 BufferSize);
-
-	/**
-	 * Zeroes the contents of a buffer
-	 *
-	 * @param Destination	Buffer to zero
-	 * @param BufferSize	Size of the destination buffer in bytes
-	 * @return				How many bytes were zeroed
-	 */
-	uint32	ZeroBuffer(uint8* Destination, uint32 BufferSize);
-
+	/** Wrapper around Opus-specific decoding state and APIs */
 	FOpusDecoderWrapper*	OpusDecoderWrapper;
-	const uint8*	SrcBufferData;
-	uint32			SrcBufferDataSize;
-	uint32			SrcBufferOffset;
-	uint32			AudioDataOffset;
-
-	uint16			SampleRate;
-	uint32			TrueSampleCount;
-	uint32			CurrentSampleCount;
-	uint8			NumChannels;
-	uint32			MaxFrameSizeSamples;
-	uint32			SampleStride;
-
-	TArray<uint8>	LastDecodedPCM;
-	uint32			LastPCMByteSize;
-	uint32			LastPCMOffset;
-	bool			bStoringEndOfFile;
-
-	// Streaming specific
-	USoundWave*		StreamingSoundWave;
-	int32			CurrentChunkIndex;
 };
diff --git a/Engine/Source/Runtime/Engine/Public/SceneManagement.h b/Engine/Source/Runtime/Engine/Public/SceneManagement.h
--- a/Engine/Source/Runtime/Engine/Public/SceneManagement.h
+++ b/Engine/Source/Runtime/Engine/Public/SceneManagement.h
@@ -2292,7 +2292,9 @@ extern ENGINE_API void ApplyViewModeOverrides(
 extern ENGINE_API void DrawUVs(FViewport* InViewport, FCanvas* InCanvas, int32 InTextYPos, const int32 LODLevel, int32 UVChannel, TArray<FVector2D> SelectedEdgeTexCoords, class FStaticMeshRenderData* StaticMeshRenderData, class FStaticLODModel* SkeletalMeshRenderData );
 
 /**
- * Computes the screen size of a given sphere bounds in the given view
+ * Computes the screen size of a given sphere bounds in the given view.
+ * The screen size is the projected diameter of the bounding sphere of the model.
+ * i.e. 0.5 means half the screen's maximum dimension.
  * @param Origin - Origin of the bounds in world space
  * @param SphereRadius - Radius of the sphere to use to calculate screen coverage
  * @param View - The view to calculate the display factor for
@@ -2300,6 +2302,37 @@ extern ENGINE_API void DrawUVs(FViewport* InViewport, FCanvas* InCanvas, int32 I
  */
 float ENGINE_API ComputeBoundsScreenSize(const FVector4& Origin, const float SphereRadius, const FSceneView& View);
 
+/**
+ * Computes the screen size of a given sphere bounds in the given view.
+ * The screen size is the projected diameter of the bounding sphere of the model. 
+ * i.e. 0.5 means half the screen's maximum dimension.
+ * @param BoundsOrigin - Origin of the bounds in world space
+ * @param SphereRadius - Radius of the sphere to use to calculate screen coverage
+ * @param ViewOrigin - The origin of the view to calculate the display factor for
+ * @param ProjMatrix - The projection matrix used to scale screen size bounds
+ * @return float - The screen size calculated
+ */
+float ENGINE_API ComputeBoundsScreenSize(const FVector4& BoundsOrigin, const float SphereRadius, const FVector4& ViewOrigin, const FMatrix& ProjMatrix);
+
+/**
+ * Computes the screen radius squared of a given sphere bounds in the given view. This is used at
+ * runtime instead of ComputeBoundsScreenSize to avoid a square root.
+ * @param Origin - Origin of the bounds in world space
+ * @param SphereRadius - Radius of the sphere to use to calculate screen coverage
+ * @param View - The view to calculate the display factor for
+ * @return float - The screen size calculated
+ */
+float ENGINE_API ComputeBoundsScreenRadiusSquared(const FVector4& Origin, const float SphereRadius, const FSceneView& View);
+
+/**
+ * Computes the draw distance of a given sphere bounds in the given view with the specified screen size.
+ * @param ScreenSize - The screen size (as computed by ComputeBoundsScreenSize)
+ * @param SphereRadius - Radius of the sphere to use to calculate screen coverage
+ * @param ProjMatrix - The projection matrix used to scale screen size bounds
+ * @return float - The draw distance calculated
+ */
+float ENGINE_API ComputeBoundsDrawDistance(const float ScreenSize, const float SphereRadius, const FMatrix& ProjMatrix);
+
 /**
  * Computes the LOD level for the given static meshes render data in the given view.
  * @param RenderData - Render data for the mesh
diff --git a/Engine/Source/Runtime/GameplayAbilities/GameplayAbilities.Build.cs b/Engine/Source/Runtime/GameplayAbilities/GameplayAbilities.Build.cs
--- a/Engine/Source/Runtime/GameplayAbilities/GameplayAbilities.Build.cs
+++ b/Engine/Source/Runtime/GameplayAbilities/GameplayAbilities.Build.cs
@@ -49,7 +49,6 @@ namespace UnrealBuildTool.Rules
 			if (UEBuildConfiguration.bBuildEditor == true)
 			{
 				PrivateDependencyModuleNames.Add("UnrealEd");
-                PrivateDependencyModuleNames.Add("GameplayTagsEditor");
                 PrivateDependencyModuleNames.Add("Slate");
 				PrivateDependencyModuleNames.Add("SequenceRecorder");
 			}
diff --git a/Engine/Source/Runtime/GameplayAbilities/Private/Abilities/GameplayAbility.cpp b/Engine/Source/Runtime/GameplayAbilities/Private/Abilities/GameplayAbility.cpp
--- a/Engine/Source/Runtime/GameplayAbilities/Private/Abilities/GameplayAbility.cpp
+++ b/Engine/Source/Runtime/GameplayAbilities/Private/Abilities/GameplayAbility.cpp
@@ -3,7 +3,6 @@
 #include "AbilitySystemPrivatePCH.h"
 #include "Abilities/GameplayAbility.h"
 #include "GameplayEffect.h"
-#include "GameplayTagsModule.h"
 #include "AbilitySystemComponent.h"
 #include "AbilitySystemBlueprintLibrary.h"
 #include "AbilityTask.h"
@@ -183,12 +182,12 @@ bool UGameplayAbility::DoesAbilitySatisfyTagRequirements(const UAbilitySystemCom
 
 		AbilitySystemComponent.GetOwnedGameplayTags(AbilitySystemComponentTags);
 
-		if (AbilitySystemComponentTags.MatchesAny(ActivationBlockedTags, false))
+		if (AbilitySystemComponentTags.HasAny(ActivationBlockedTags))
 		{
 			bBlocked = true;
 		}
 
-		if (!AbilitySystemComponentTags.MatchesAll(ActivationRequiredTags, true))
+		if (!AbilitySystemComponentTags.HasAll(ActivationRequiredTags))
 		{
 			bMissing = true;
 		}
@@ -198,12 +197,12 @@ bool UGameplayAbility::DoesAbilitySatisfyTagRequirements(const UAbilitySystemCom
 	{
 		if (SourceBlockedTags.Num() || SourceRequiredTags.Num())
 		{
-			if (SourceTags->MatchesAny(SourceBlockedTags, false))
+			if (SourceTags->HasAny(SourceBlockedTags))
 			{
 				bBlocked = true;
 			}
 
-			if (!SourceTags->MatchesAll(SourceRequiredTags, true))
+			if (!SourceTags->HasAll(SourceRequiredTags))
 			{
 				bMissing = true;
 			}
@@ -214,12 +213,12 @@ bool UGameplayAbility::DoesAbilitySatisfyTagRequirements(const UAbilitySystemCom
 	{
 		if (TargetBlockedTags.Num() || TargetRequiredTags.Num())
 		{
-			if (TargetTags->MatchesAny(TargetBlockedTags, false))
+			if (TargetTags->HasAny(TargetBlockedTags))
 			{
 				bBlocked = true;
 			}
 
-			if (!TargetTags->MatchesAll(TargetRequiredTags, true))
+			if (!TargetTags->HasAll(TargetRequiredTags))
 			{
 				bMissing = true;
 			}
@@ -712,7 +711,7 @@ UGameplayEffect* UGameplayAbility::GetCooldownGameplayEffect() const
 	}
 	else
 	{
-		return CooldownGameplayEffect;
+		return nullptr;
 	}
 }
 
@@ -724,7 +723,7 @@ UGameplayEffect* UGameplayAbility::GetCostGameplayEffect() const
 	}
 	else
 	{
-		return CostGameplayEffect;
+		return nullptr;
 	}
 }
 
@@ -909,16 +908,6 @@ FGameplayEffectSpecHandle UGameplayAbility::MakeOutgoingGameplayEffectSpec(TSubc
 	return MakeOutgoingGameplayEffectSpec(CurrentSpecHandle, CurrentActorInfo, CurrentActivationInfo, GameplayEffectClass, Level);
 }
 
-FGameplayEffectSpecHandle UGameplayAbility::GetOutgoingGameplayEffectSpec(const UGameplayEffect* GameplayEffect, float Level) const
-{
-	if ( GameplayEffect )
-	{
-		return MakeOutgoingGameplayEffectSpec(GameplayEffect->GetClass(), Level);
-	}
-	
-	return FGameplayEffectSpecHandle(nullptr);
-}
-
 int32 AbilitySystemShowMakeOutgoingGameplayEffectSpecs = 0;
 static FAutoConsoleVariableRef CVarAbilitySystemShowMakeOutgoingGameplayEffectSpecs(TEXT("AbilitySystem.ShowClientMakeOutgoingSpecs"), AbilitySystemShowMakeOutgoingGameplayEffectSpecs, TEXT("Displays all GameplayEffect specs created on non authority clients"), ECVF_Default );
 
@@ -1399,11 +1388,6 @@ FGameplayEffectContextHandle UGameplayAbility::MakeEffectContext(const FGameplay
 	return Context;
 }
 
-FGameplayEffectContextHandle UGameplayAbility::GetEffectContext(const FGameplayAbilitySpecHandle Handle, const FGameplayAbilityActorInfo *ActorInfo) const
-{
-	return MakeEffectContext(Handle, ActorInfo);
-}
-
 bool UGameplayAbility::IsTriggered() const
 {
 	// Assume that if there is triggered data, then we are triggered. 
@@ -1497,17 +1481,6 @@ FActiveGameplayEffectHandle UGameplayAbility::BP_ApplyGameplayEffectToOwner(TSub
 	return FActiveGameplayEffectHandle();
 }
 
-FActiveGameplayEffectHandle UGameplayAbility::K2_ApplyGameplayEffectToOwner(const UGameplayEffect* GameplayEffect, int32 GameplayEffectLevel, int32 Stacks)
-{
-	if ( GameplayEffect )
-	{
-		return BP_ApplyGameplayEffectToOwner(GameplayEffect->GetClass(), GameplayEffectLevel, Stacks);
-	}
-
-	ABILITY_LOG(Error, TEXT("K2_ApplyGameplayEffectToOwner called on ability %s with no GameplayEffect."), *GetName());
-	return FActiveGameplayEffectHandle();
-}
-
 FActiveGameplayEffectHandle UGameplayAbility::ApplyGameplayEffectToOwner(const FGameplayAbilitySpecHandle Handle, const FGameplayAbilityActorInfo* ActorInfo, const FGameplayAbilityActivationInfo ActivationInfo, const UGameplayEffect* GameplayEffect, float GameplayEffectLevel, int32 Stacks) const
 {
 	if (GameplayEffect && (HasAuthorityOrPredictionKey(ActorInfo, &ActivationInfo)))
@@ -1549,11 +1522,6 @@ TArray<FActiveGameplayEffectHandle> UGameplayAbility::BP_ApplyGameplayEffectToTa
 	return ApplyGameplayEffectToTarget(CurrentSpecHandle, CurrentActorInfo, CurrentActivationInfo, Target, GameplayEffectClass, GameplayEffectLevel, Stacks);
 }
 
-TArray<FActiveGameplayEffectHandle> UGameplayAbility::K2_ApplyGameplayEffectToTarget(FGameplayAbilityTargetDataHandle Target, const UGameplayEffect* GameplayEffect, int32 GameplayEffectLevel, int32 Stacks)
-{
-	return BP_ApplyGameplayEffectToTarget(Target, GameplayEffect->GetClass(), GameplayEffectLevel, Stacks);
-}
-
 TArray<FActiveGameplayEffectHandle> UGameplayAbility::ApplyGameplayEffectToTarget(const FGameplayAbilitySpecHandle Handle, const FGameplayAbilityActorInfo* ActorInfo, const FGameplayAbilityActivationInfo ActivationInfo, const FGameplayAbilityTargetDataHandle& Target, TSubclassOf<UGameplayEffect> GameplayEffectClass, float GameplayEffectLevel, int32 Stacks) const
 {
 	SCOPE_CYCLE_COUNTER(STAT_ApplyGameplayEffectToTarget);
@@ -1644,39 +1612,6 @@ void UGameplayAbility::BP_RemoveGameplayEffectFromOwnerWithGrantedTags(FGameplay
 	CurrentActorInfo->AbilitySystemComponent->RemoveActiveEffects(Query, StacksToRemove);
 }
 
-void UGameplayAbility::ConvertDeprecatedGameplayEffectReferencesToBlueprintReferences(UGameplayEffect* OldGE, TSubclassOf<UGameplayEffect> NewGEClass)
-{
-	bool bChangedSomething = false;
-	if ( CooldownGameplayEffect && CooldownGameplayEffect == OldGE )
-	{
-		if ( !CooldownGameplayEffectClass )
-		{
-			CooldownGameplayEffectClass = NewGEClass;
-		}
-
-		CooldownGameplayEffect = nullptr;
-
-		bChangedSomething = true;
-	}
-
-	if ( CostGameplayEffect && CostGameplayEffect == OldGE )
-	{
-		if ( !CostGameplayEffectClass )
-		{
-			CostGameplayEffectClass = NewGEClass;
-		}
-
-		CostGameplayEffect = nullptr;
-
-		bChangedSomething = true;
-	}
-
-	if ( bChangedSomething )
-	{
-		MarkPackageDirty();
-	}
-}
-
 float UGameplayAbility::GetCooldownTimeRemaining() const
 {
 	return IsInstantiated() ? GetCooldownTimeRemaining(CurrentActorInfo) : 0.f;
diff --git a/Engine/Source/Runtime/GameplayAbilities/Private/Abilities/GameplayAbility_Montage.cpp b/Engine/Source/Runtime/GameplayAbilities/Private/Abilities/GameplayAbility_Montage.cpp
--- a/Engine/Source/Runtime/GameplayAbilities/Private/Abilities/GameplayAbility_Montage.cpp
+++ b/Engine/Source/Runtime/GameplayAbilities/Private/Abilities/GameplayAbility_Montage.cpp
@@ -80,27 +80,3 @@ void UGameplayAbility_Montage::GetGameplayEffectsWhileAnimating(TArray<const UGa
 		}
 	}
 }
-
-void UGameplayAbility_Montage::ConvertDeprecatedGameplayEffectReferencesToBlueprintReferences(UGameplayEffect* OldGE, TSubclassOf<UGameplayEffect> NewGEClass)
-{
-	Super::ConvertDeprecatedGameplayEffectReferencesToBlueprintReferences(OldGE, NewGEClass);
-
-	bool bChangedSomething = false;
-
-	for ( int32 GEIdx = GameplayEffectsWhileAnimating.Num() - 1; GEIdx >= 0; --GEIdx )
-	{
-		const UGameplayEffect* Effect = GameplayEffectsWhileAnimating[GEIdx];
-		if ( Effect && Effect == OldGE )
-		{
-			GameplayEffectClassesWhileAnimating.Add(NewGEClass);
-			GameplayEffectsWhileAnimating.RemoveAt(GEIdx);
-
-			bChangedSomething = true;
-		}
-	}
-
-	if (bChangedSomething)
-	{
-		MarkPackageDirty();
-	}
-}
\ No newline at end of file
diff --git a/Engine/Source/Runtime/GameplayAbilities/Private/Abilities/Tasks/AbilityTask_WaitAbilityActivate.cpp b/Engine/Source/Runtime/GameplayAbilities/Private/Abilities/Tasks/AbilityTask_WaitAbilityActivate.cpp
--- a/Engine/Source/Runtime/GameplayAbilities/Private/Abilities/Tasks/AbilityTask_WaitAbilityActivate.cpp
+++ b/Engine/Source/Runtime/GameplayAbilities/Private/Abilities/Tasks/AbilityTask_WaitAbilityActivate.cpp
@@ -49,8 +49,8 @@ void UAbilityTask_WaitAbilityActivate::OnAbilityActivate(UGameplayAbility* Activ
 
 	if (TagRequirements.IsEmpty())
 	{
-		if ((WithTag.IsValid() && !ActivatedAbility->AbilityTags.HasTag(WithTag, EGameplayTagMatchType::IncludeParentTags, EGameplayTagMatchType::Explicit)) ||
-			(WithoutTag.IsValid() && ActivatedAbility->AbilityTags.HasTag(WithoutTag, EGameplayTagMatchType::IncludeParentTags, EGameplayTagMatchType::Explicit)))
+		if ((WithTag.IsValid() && !ActivatedAbility->AbilityTags.HasTag(WithTag)) ||
+			(WithoutTag.IsValid() && ActivatedAbility->AbilityTags.HasTag(WithoutTag)))
 		{
 			// Failed tag check
 			return;
diff --git a/Engine/Source/Runtime/GameplayAbilities/Private/Abilities/Tasks/AbilityTask_WaitAbilityCommit.cpp b/Engine/Source/Runtime/GameplayAbilities/Private/Abilities/Tasks/AbilityTask_WaitAbilityCommit.cpp
--- a/Engine/Source/Runtime/GameplayAbilities/Private/Abilities/Tasks/AbilityTask_WaitAbilityCommit.cpp
+++ b/Engine/Source/Runtime/GameplayAbilities/Private/Abilities/Tasks/AbilityTask_WaitAbilityCommit.cpp
@@ -43,8 +43,8 @@ void UAbilityTask_WaitAbilityCommit::OnDestroy(bool AbilityEnded)
 
 void UAbilityTask_WaitAbilityCommit::OnAbilityCommit(UGameplayAbility *ActivatedAbility)
 {
-	if ( (WithTag.IsValid() && !ActivatedAbility->AbilityTags.HasTag(WithTag, EGameplayTagMatchType::IncludeParentTags, EGameplayTagMatchType::Explicit)) ||
-		 (WithoutTag.IsValid() && ActivatedAbility->AbilityTags.HasTag(WithoutTag, EGameplayTagMatchType::IncludeParentTags, EGameplayTagMatchType::Explicit)))
+	if ( (WithTag.IsValid() && !ActivatedAbility->AbilityTags.HasTag(WithTag)) ||
+		 (WithoutTag.IsValid() && ActivatedAbility->AbilityTags.HasTag(WithoutTag)))
 	{
 		// Failed tag check
 		return;
diff --git a/Engine/Source/Runtime/GameplayAbilities/Private/Abilities/Tasks/AbilityTask_WaitAttributeChange.cpp b/Engine/Source/Runtime/GameplayAbilities/Private/Abilities/Tasks/AbilityTask_WaitAttributeChange.cpp
--- a/Engine/Source/Runtime/GameplayAbilities/Private/Abilities/Tasks/AbilityTask_WaitAttributeChange.cpp
+++ b/Engine/Source/Runtime/GameplayAbilities/Private/Abilities/Tasks/AbilityTask_WaitAttributeChange.cpp
@@ -60,8 +60,8 @@ void UAbilityTask_WaitAttributeChange::OnAttributeChange(float NewValue, const F
 	}
 	else
 	{
-		if ((WithTag.IsValid() && !Data->EffectSpec.CapturedSourceTags.GetAggregatedTags()->HasTag(WithTag, EGameplayTagMatchType::IncludeParentTags, EGameplayTagMatchType::Explicit)) ||
-			(WithoutTag.IsValid() && Data->EffectSpec.CapturedSourceTags.GetAggregatedTags()->HasTag(WithoutTag, EGameplayTagMatchType::IncludeParentTags, EGameplayTagMatchType::Explicit)))
+		if ((WithTag.IsValid() && !Data->EffectSpec.CapturedSourceTags.GetAggregatedTags()->HasTag(WithTag)) ||
+			(WithoutTag.IsValid() && Data->EffectSpec.CapturedSourceTags.GetAggregatedTags()->HasTag(WithoutTag)))
 		{
 			// Failed tag check
 			return;
diff --git a/Engine/Source/Runtime/GameplayAbilities/Private/AbilitySystemComponent.cpp b/Engine/Source/Runtime/GameplayAbilities/Private/AbilitySystemComponent.cpp
--- a/Engine/Source/Runtime/GameplayAbilities/Private/AbilitySystemComponent.cpp
+++ b/Engine/Source/Runtime/GameplayAbilities/Private/AbilitySystemComponent.cpp
@@ -267,21 +267,6 @@ FGameplayEffectSpecHandle UAbilitySystemComponent::MakeOutgoingSpec(TSubclassOf<
 	return FGameplayEffectSpecHandle(nullptr);
 }
 
-FGameplayEffectSpecHandle UAbilitySystemComponent::GetOutgoingSpec(const UGameplayEffect* GameplayEffect, float Level, FGameplayEffectContextHandle Contex) const
-{
-	if (GameplayEffect)
-	{
-		return MakeOutgoingSpec(GameplayEffect->GetClass(), Level, MakeEffectContext());
-	}
-
-	return FGameplayEffectSpecHandle(nullptr);
-}
-
-FGameplayEffectSpecHandle UAbilitySystemComponent::GetOutgoingSpec(const UGameplayEffect* GameplayEffect, float Level) const
-{
-	return GetOutgoingSpec(GameplayEffect, Level, MakeEffectContext());
-}
-
 FGameplayEffectContextHandle UAbilitySystemComponent::MakeEffectContext() const
 {
 	FGameplayEffectContextHandle Context = FGameplayEffectContextHandle(UAbilitySystemGlobals::Get().AllocGameplayEffectContext());
@@ -292,11 +277,6 @@ FGameplayEffectContextHandle UAbilitySystemComponent::MakeEffectContext() const
 	return Context;
 }
 
-FGameplayEffectContextHandle UAbilitySystemComponent::GetEffectContext() const
-{
-	return MakeEffectContext();
-}
-
 int32 UAbilitySystemComponent::GetGameplayEffectCount(TSubclassOf<UGameplayEffect> SourceGameplayEffect, UAbilitySystemComponent* OptionalInstigatorFilterComponent, bool bEnforceOnGoingCheck)
 {
 	int32 Count = 0;
@@ -331,13 +311,6 @@ int32 UAbilitySystemComponent::GetGameplayEffectCount(TSubclassOf<UGameplayEffec
 	return Count;
 }
 
-int32 UAbilitySystemComponent::GetAggregatedStackCount(const FActiveGameplayEffectQuery& Query)
-{
-	PRAGMA_DISABLE_DEPRECATION_WARNINGS
-	return ActiveGameplayEffects.GetActiveEffectCount(Query);
-	PRAGMA_ENABLE_DEPRECATION_WARNINGS
-}
-
 int32 UAbilitySystemComponent::GetAggregatedStackCount(const FGameplayEffectQuery& Query)
 {
 	return ActiveGameplayEffects.GetActiveEffectCount(Query);
@@ -361,11 +334,6 @@ FActiveGameplayEffectHandle UAbilitySystemComponent::BP_ApplyGameplayEffectToTar
 	return ApplyGameplayEffectToTarget(GameplayEffect, Target, Level, Context);	
 }
 
-FActiveGameplayEffectHandle UAbilitySystemComponent::K2_ApplyGameplayEffectToTarget(UGameplayEffect *GameplayEffect, UAbilitySystemComponent *Target, float Level, FGameplayEffectContextHandle Context)
-{
-	return ApplyGameplayEffectToTarget(GameplayEffect, Target, Level, Context);
-}
-
 /** This is a helper function used in automated testing, I'm not sure how useful it will be to gamecode or blueprints */
 FActiveGameplayEffectHandle UAbilitySystemComponent::ApplyGameplayEffectToTarget(UGameplayEffect *GameplayEffect, UAbilitySystemComponent *Target, float Level, FGameplayEffectContextHandle Context, FPredictionKey PredictionKey)
 {
@@ -396,16 +364,6 @@ FActiveGameplayEffectHandle UAbilitySystemComponent::BP_ApplyGameplayEffectToSel
 	return FActiveGameplayEffectHandle();
 }
 
-FActiveGameplayEffectHandle UAbilitySystemComponent::K2_ApplyGameplayEffectToSelf(const UGameplayEffect *GameplayEffect, float Level, FGameplayEffectContextHandle EffectContext)
-{
-	if ( GameplayEffect )
-	{
-		return BP_ApplyGameplayEffectToSelf(GameplayEffect->GetClass(), Level, EffectContext);
-	}
-
-	return FActiveGameplayEffectHandle();
-}
-
 /** This is a helper function - it seems like this will be useful as a blueprint interface at the least, but Level parameter may need to be expanded */
 FActiveGameplayEffectHandle UAbilitySystemComponent::ApplyGameplayEffectToSelf(const UGameplayEffect *GameplayEffect, float Level, const FGameplayEffectContextHandle& EffectContext, FPredictionKey PredictionKey)
 {
@@ -1221,14 +1179,6 @@ void UAbilitySystemComponent::NetMulticast_InvokeGameplayCuesAddedAndWhileActive
 
 // ----------------------------------------------------------------------------------------
 
-
-// #deprecated, use FGameplayEffectQuery version
-TArray<float> UAbilitySystemComponent::GetActiveEffectsTimeRemaining(const FActiveGameplayEffectQuery Query) const
-{
-	PRAGMA_DISABLE_DEPRECATION_WARNINGS
-	return ActiveGameplayEffects.GetActiveEffectsTimeRemaining(Query);
-	PRAGMA_ENABLE_DEPRECATION_WARNINGS
-}
 TArray<float> UAbilitySystemComponent::GetActiveEffectsTimeRemaining(const FGameplayEffectQuery& Query) const
 {
 	return ActiveGameplayEffects.GetActiveEffectsTimeRemaining(Query);
@@ -1239,25 +1189,11 @@ TArray<TPair<float,float>> UAbilitySystemComponent::GetActiveEffectsTimeRemainin
 	return ActiveGameplayEffects.GetActiveEffectsTimeRemainingAndDuration(Query);
 }
 
-// #deprecated, use FGameplayEffectQuery version
-TArray<float> UAbilitySystemComponent::GetActiveEffectsDuration(const FActiveGameplayEffectQuery Query) const
-{
-	PRAGMA_DISABLE_DEPRECATION_WARNINGS
-	return ActiveGameplayEffects.GetActiveEffectsDuration(Query);
-	PRAGMA_ENABLE_DEPRECATION_WARNINGS
-}
 TArray<float> UAbilitySystemComponent::GetActiveEffectsDuration(const FGameplayEffectQuery& Query) const
 {
 	return ActiveGameplayEffects.GetActiveEffectsDuration(Query);
 }
 
-// #deprecated, use FGameplayEffectQuery version
-TArray<FActiveGameplayEffectHandle> UAbilitySystemComponent::GetActiveEffects(const FActiveGameplayEffectQuery Query) const
-{
-	PRAGMA_DISABLE_DEPRECATION_WARNINGS
-	return ActiveGameplayEffects.GetActiveEffects(Query);
-	PRAGMA_ENABLE_DEPRECATION_WARNINGS
-}
 TArray<FActiveGameplayEffectHandle> UAbilitySystemComponent::GetActiveEffects(const FGameplayEffectQuery& Query) const
 {
 	return ActiveGameplayEffects.GetActiveEffects(Query);
@@ -1299,18 +1235,6 @@ int32 UAbilitySystemComponent::RemoveActiveEffectsWithGrantedTags(const FGamepla
 	return 0;
 }
 
-// #deprecated, use FGameplayEffectQuery version
-int32 UAbilitySystemComponent::RemoveActiveEffects(const FActiveGameplayEffectQuery Query, int32 StacksToRemove)
-{
-	if (IsOwnerActorAuthoritative())
-	{
-		PRAGMA_DISABLE_DEPRECATION_WARNINGS
-		return ActiveGameplayEffects.RemoveActiveEffects(Query, StacksToRemove);
-		PRAGMA_ENABLE_DEPRECATION_WARNINGS
-	}
-
-	return 0;
-}
 int32 UAbilitySystemComponent::RemoveActiveEffects(const FGameplayEffectQuery& Query, int32 StacksToRemove)
 {
 	if (IsOwnerActorAuthoritative())
@@ -2105,7 +2029,7 @@ void UAbilitySystemComponent::Debug_Internal(FAbilitySystemComponentDebugInfo& I
 				StatusText = TEXT(" (InputBlocked)");
 				AbilityTextColor = FColor::Red;
 			}
-			else if (AbilitySpec.Ability->AbilityTags.MatchesAny(BlockedAbilityTags.GetExplicitGameplayTags(), false))
+			else if (AbilitySpec.Ability->AbilityTags.HasAny(BlockedAbilityTags.GetExplicitGameplayTags()))
 			{
 				StatusText = TEXT(" (TagBlocked)");
 				AbilityTextColor = FColor::Red;
diff --git a/Engine/Source/Runtime/GameplayAbilities/Private/AbilitySystemComponent_Abilities.cpp b/Engine/Source/Runtime/GameplayAbilities/Private/AbilitySystemComponent_Abilities.cpp
--- a/Engine/Source/Runtime/GameplayAbilities/Private/AbilitySystemComponent_Abilities.cpp
+++ b/Engine/Source/Runtime/GameplayAbilities/Private/AbilitySystemComponent_Abilities.cpp
@@ -693,8 +693,8 @@ void UAbilitySystemComponent::CancelAbilities(const FGameplayTagContainer* WithT
 			continue;
 		}
 
-		bool WithTagPass = (!WithTags || Spec.Ability->AbilityTags.MatchesAny(*WithTags, false));
-		bool WithoutTagPass = (!WithoutTags || !Spec.Ability->AbilityTags.MatchesAny(*WithoutTags, false));
+		bool WithTagPass = (!WithTags || Spec.Ability->AbilityTags.HasAny(*WithTags));
+		bool WithoutTagPass = (!WithoutTags || !Spec.Ability->AbilityTags.HasAny(*WithoutTags));
 
 		if (WithTagPass && WithoutTagPass)
 		{
@@ -793,7 +793,7 @@ void UAbilitySystemComponent::ApplyAbilityBlockAndCancelTags(const FGameplayTagC
 bool UAbilitySystemComponent::AreAbilityTagsBlocked(const FGameplayTagContainer& Tags) const
 {
 	// Expand the passed in tags to get parents, not the blocked tags
-	return Tags.MatchesAny(BlockedAbilityTags.GetExplicitGameplayTags(), false);
+	return Tags.HasAny(BlockedAbilityTags.GetExplicitGameplayTags());
 }
 
 void UAbilitySystemComponent::BlockAbilitiesWithTags(const FGameplayTagContainer& Tags)
@@ -866,9 +866,14 @@ void UAbilitySystemComponent::OnRep_ActivateAbilities()
 
 void UAbilitySystemComponent::GetActivatableGameplayAbilitySpecsByAllMatchingTags(const FGameplayTagContainer& GameplayTagContainer, TArray < struct FGameplayAbilitySpec* >& MatchingGameplayAbilities, bool bOnlyAbilitiesThatSatisfyTagRequirements) const
 {
+	if (!GameplayTagContainer.IsValid())
+	{
+		return;
+	}
+
 	for (const FGameplayAbilitySpec& Spec : ActivatableAbilities.Items)
 	{		
-		if (Spec.Ability && Spec.Ability->AbilityTags.MatchesAll(GameplayTagContainer, false))
+		if (Spec.Ability && Spec.Ability->AbilityTags.HasAll(GameplayTagContainer))
 		{
 			// Consider abilities that are blocked by tags currently if we're supposed to (default behavior).  
 			// That way, we can use the blocking to find an appropriate ability based on tags when we have more than 
@@ -2337,51 +2342,58 @@ void UAbilitySystemComponent::OnRep_ReplicatedAnimMontage()
 				AnimInstance->Montage_SetPlayRate(LocalAnimMontageInfo.AnimMontage, RepAnimMontageInfo.PlayRate);
 			}
 
-			int32 RepSectionID = LocalAnimMontageInfo.AnimMontage->GetSectionIndexFromPosition(RepAnimMontageInfo.Position);
-			int32 RepNextSectionID = int32(RepAnimMontageInfo.NextSectionID) - 1;
-		
-			// And NextSectionID for the replicated SectionID.
-			if( RepSectionID != INDEX_NONE )
-			{
-				int32 NextSectionID = AnimInstance->Montage_GetNextSectionID(LocalAnimMontageInfo.AnimMontage, RepSectionID);
-
-				// If NextSectionID is different than the replicated one, then set it.
-				if( NextSectionID != RepNextSectionID )
-				{
-					AnimInstance->Montage_SetNextSection(LocalAnimMontageInfo.AnimMontage->GetSectionName(RepSectionID), LocalAnimMontageInfo.AnimMontage->GetSectionName(RepNextSectionID), LocalAnimMontageInfo.AnimMontage);
-				}
+			// Compressed Flags
+			const bool bIsStopped = AnimInstance->Montage_GetIsStopped(LocalAnimMontageInfo.AnimMontage);
+			const bool bReplicatedIsStopped = bool(RepAnimMontageInfo.IsStopped);
 
-				// Make sure we haven't received that update too late and the client hasn't already jumped to another section. 
-				int32 CurrentSectionID = LocalAnimMontageInfo.AnimMontage->GetSectionIndexFromPosition(AnimInstance->Montage_GetPosition(LocalAnimMontageInfo.AnimMontage));
-				if ((CurrentSectionID != RepSectionID) && (CurrentSectionID != RepNextSectionID))
+			// Process stopping first, so we don't change sections and cause blending to pop.
+			if (bReplicatedIsStopped)
+			{
+				if (!bIsStopped)
 				{
-					// Client is in a wrong section, teleport him into the begining of the right section
-					const float SectionStartTime = LocalAnimMontageInfo.AnimMontage->GetAnimCompositeSection(RepSectionID).GetTime();
-					AnimInstance->Montage_SetPosition(LocalAnimMontageInfo.AnimMontage, SectionStartTime);
+					CurrentMontageStop(RepAnimMontageInfo.BlendTime);
 				}
 			}
-
-			// Update Position. If error is too great, jump to replicated position.
-			float CurrentPosition = AnimInstance->Montage_GetPosition(LocalAnimMontageInfo.AnimMontage);
-			int32 CurrentSectionID = LocalAnimMontageInfo.AnimMontage->GetSectionIndexFromPosition(CurrentPosition);
-			// Only check threshold if we are located in the same section. Different sections require a bit more work as we could be jumping around the timeline.
-			if ((CurrentSectionID == RepSectionID) && (FMath::Abs(CurrentPosition - RepAnimMontageInfo.Position) > MONTAGE_REP_POS_ERR_THRESH) && RepAnimMontageInfo.IsStopped == 0)
+			else
 			{
-				// fast forward to server position and trigger notifies
-				if (FAnimMontageInstance* MontageInstance = AnimInstance->GetActiveInstanceForMontage(RepAnimMontageInfo.AnimMontage))
+				int32 RepSectionID = LocalAnimMontageInfo.AnimMontage->GetSectionIndexFromPosition(RepAnimMontageInfo.Position);
+				int32 RepNextSectionID = int32(RepAnimMontageInfo.NextSectionID) - 1;
+		
+				// And NextSectionID for the replicated SectionID.
+				if( RepSectionID != INDEX_NONE )
 				{
-					MontageInstance->HandleEvents(CurrentPosition, RepAnimMontageInfo.Position, nullptr);
-					AnimInstance->TriggerAnimNotifies(0.f);
+					int32 NextSectionID = AnimInstance->Montage_GetNextSectionID(LocalAnimMontageInfo.AnimMontage, RepSectionID);
+
+					// If NextSectionID is different than the replicated one, then set it.
+					if( NextSectionID != RepNextSectionID )
+					{
+						AnimInstance->Montage_SetNextSection(LocalAnimMontageInfo.AnimMontage->GetSectionName(RepSectionID), LocalAnimMontageInfo.AnimMontage->GetSectionName(RepNextSectionID), LocalAnimMontageInfo.AnimMontage);
+					}
+
+					// Make sure we haven't received that update too late and the client hasn't already jumped to another section. 
+					int32 CurrentSectionID = LocalAnimMontageInfo.AnimMontage->GetSectionIndexFromPosition(AnimInstance->Montage_GetPosition(LocalAnimMontageInfo.AnimMontage));
+					if ((CurrentSectionID != RepSectionID) && (CurrentSectionID != RepNextSectionID))
+					{
+						// Client is in a wrong section, teleport him into the begining of the right section
+						const float SectionStartTime = LocalAnimMontageInfo.AnimMontage->GetAnimCompositeSection(RepSectionID).GetTime();
+						AnimInstance->Montage_SetPosition(LocalAnimMontageInfo.AnimMontage, SectionStartTime);
+					}
 				}
-				AnimInstance->Montage_SetPosition(LocalAnimMontageInfo.AnimMontage, RepAnimMontageInfo.Position);
-			}
 
-			// Compressed Flags
-			bool bIsStopped = AnimInstance->Montage_GetIsStopped(LocalAnimMontageInfo.AnimMontage);
-			bool ReplicatedIsStopped = bool(RepAnimMontageInfo.IsStopped);
-			if( ReplicatedIsStopped && !bIsStopped )
-			{
-				CurrentMontageStop(RepAnimMontageInfo.BlendTime);
+				// Update Position. If error is too great, jump to replicated position.
+				float CurrentPosition = AnimInstance->Montage_GetPosition(LocalAnimMontageInfo.AnimMontage);
+				int32 CurrentSectionID = LocalAnimMontageInfo.AnimMontage->GetSectionIndexFromPosition(CurrentPosition);
+				// Only check threshold if we are located in the same section. Different sections require a bit more work as we could be jumping around the timeline.
+				if ((CurrentSectionID == RepSectionID) && (FMath::Abs(CurrentPosition - RepAnimMontageInfo.Position) > MONTAGE_REP_POS_ERR_THRESH) && RepAnimMontageInfo.IsStopped == 0)
+				{
+					// fast forward to server position and trigger notifies
+					if (FAnimMontageInstance* MontageInstance = AnimInstance->GetActiveInstanceForMontage(RepAnimMontageInfo.AnimMontage))
+					{
+						MontageInstance->HandleEvents(CurrentPosition, RepAnimMontageInfo.Position, nullptr);
+						AnimInstance->TriggerAnimNotifies(0.f);
+					}
+					AnimInstance->Montage_SetPosition(LocalAnimMontageInfo.AnimMontage, RepAnimMontageInfo.Position);
+				}
 			}
 		}
 	}
diff --git a/Engine/Source/Runtime/GameplayAbilities/Private/AbilitySystemGlobals.cpp b/Engine/Source/Runtime/GameplayAbilities/Private/AbilitySystemGlobals.cpp
--- a/Engine/Source/Runtime/GameplayAbilities/Private/AbilitySystemGlobals.cpp
+++ b/Engine/Source/Runtime/GameplayAbilities/Private/AbilitySystemGlobals.cpp
@@ -5,7 +5,6 @@
 #include "AbilitySystemComponent.h"
 #include "AbilitySystemInterface.h"
 #include "Abilities/GameplayAbilityTypes.h"
-#include "GameplayTagsModule.h"
 #include "GameplayCueInterface.h"
 #include "GameplayCueManager.h"
 #include "GameplayTagResponseTable.h"
@@ -88,8 +87,7 @@ bool UAbilitySystemGlobals::DeriveGameplayCueTagFromAssetName(FString AssetName,
 				AssetName = FString(TEXT("GameplayCue.")) + AssetName;
 			}
 
-			IGameplayTagsModule& GameplayTagsModule = IGameplayTagsModule::Get();
-			GameplayCueTag = GameplayTagsModule.GetGameplayTagsManager().RequestGameplayTag(FName(*AssetName), false);
+			GameplayCueTag = UGameplayTagsManager::Get().RequestGameplayTag(FName(*AssetName), false);
 		}
 		GameplayCueName = GameplayCueTag.GetTagName();
 	}
@@ -189,8 +187,7 @@ UFunction* UAbilitySystemGlobals::GetGameplayCueFunction(const FGameplayTag& Chi
 	// and will be regenerated all the time in the editor. Just doing a single pass at startup won't be enough,
 	// we'll need a mechanism for registering classes when they are loaded on demand.
 	
-	IGameplayTagsModule& GameplayTagsModule = IGameplayTagsModule::Get();
-	FGameplayTagContainer TagAndParentsContainer = GameplayTagsModule.GetGameplayTagsManager().RequestGameplayTagParents(ChildTag);
+	FGameplayTagContainer TagAndParentsContainer = ChildTag.GetGameplayTagParents();
 
 	for (auto InnerTagIt = TagAndParentsContainer.CreateConstIterator(); InnerTagIt; ++InnerTagIt)
 	{
diff --git a/Engine/Source/Runtime/GameplayAbilities/Private/AbilitySystemTestAttributeSet.cpp b/Engine/Source/Runtime/GameplayAbilities/Private/AbilitySystemTestAttributeSet.cpp
--- a/Engine/Source/Runtime/GameplayAbilities/Private/AbilitySystemTestAttributeSet.cpp
+++ b/Engine/Source/Runtime/GameplayAbilities/Private/AbilitySystemTestAttributeSet.cpp
@@ -2,7 +2,6 @@
 
 #include "AbilitySystemPrivatePCH.h"
 #include "AbilitySystemTestAttributeSet.h"
-#include "GameplayTagsModule.h"
 #include "GameplayEffectExtension.h"
 #include "AbilitySystemComponent.h"
 
@@ -45,7 +44,7 @@ bool UAbilitySystemTestAttributeSet::PreGameplayEffectExecute(struct FGameplayEf
 			{
 				// Dodge!
 				Data.EvaluatedData.Magnitude = 0.f;
-				Data.EvaluatedData.Tags.AddTag(IGameplayTagsModule::RequestGameplayTag(FName(TEXT("Dodged"))));
+				Data.EvaluatedData.Tags.AddTag(FGameplayTag::RequestGameplayTag(FName(TEXT("Dodged"))));
 
 				// How dodge is handled will be game dependant. There are a few options I think of:
 				// -We still apply 0 damage, but tag it as Dodged. The GameplayCue system could pick up on this and play a visual effect. The combat log could pick up in and print it out too.
@@ -64,17 +63,17 @@ bool UAbilitySystemTestAttributeSet::PreGameplayEffectExecute(struct FGameplayEf
 				{
 					// Crit!
 					Data.EvaluatedData.Magnitude *= SourceAttributes->CritMultiplier;
-					Data.EvaluatedData.Tags.AddTag(IGameplayTagsModule::RequestGameplayTag(FName(TEXT("Damage.Crit"))));
+					Data.EvaluatedData.Tags.AddTag(FGameplayTag::RequestGameplayTag(FName(TEXT("Damage.Crit"))));
 				}
 			}
 
 			// Now apply armor reduction
-			if (Data.EvaluatedData.Tags.HasTag(IGameplayTagsModule::RequestGameplayTag(FName(TEXT("Damage.Physical"))), EGameplayTagMatchType::IncludeParentTags, EGameplayTagMatchType::Explicit))
+			if (Data.EvaluatedData.Tags.HasTag(FGameplayTag::RequestGameplayTag(FName(TEXT("Damage.Physical")))))
 			{
 				// This is a trivial/naive implementation of armor. It assumes the rmorDamageReduction is an actual % to reduce physics damage by.
 				// Real games would probably use armor rating as an attribute and calculate a % here based on the damage source's level, etc.
 				Data.EvaluatedData.Magnitude *= (1.f - ArmorDamageReduction);
-				Data.EvaluatedData.Tags.AddTag(IGameplayTagsModule::RequestGameplayTag(FName(TEXT("Damage.Mitigated.Armor"))));
+				Data.EvaluatedData.Tags.AddTag(FGameplayTag::RequestGameplayTag(FName(TEXT("Damage.Mitigated.Armor"))));
 			}
 		}
 
@@ -97,7 +96,7 @@ void UAbilitySystemTestAttributeSet::PostGameplayEffectExecute(const struct FGam
 	if (DamageProperty == ModifiedProperty)
 	{
 		// Anytime Damage is applied with 'Damage.Fire' tag, there is a chance to apply a burning DOT
-		if (Data.EffectSpec.CapturedSourceTags.GetAggregatedTags()->HasTag( IGameplayTagsModule::RequestGameplayTag(FName(TEXT("FireDamage"))), EGameplayTagMatchType::IncludeParentTags, EGameplayTagMatchType::Explicit))
+		if (Data.EffectSpec.CapturedSourceTags.GetAggregatedTags()->HasTag( FGameplayTag::RequestGameplayTag(FName(TEXT("FireDamage")))))
 		{
 			// Logic to rand() a burning DOT, if successful, apply DOT GameplayEffect to the target
 		}
diff --git a/Engine/Source/Runtime/GameplayAbilities/Private/GameplayAbilityTypes.cpp b/Engine/Source/Runtime/GameplayAbilities/Private/GameplayAbilityTypes.cpp
--- a/Engine/Source/Runtime/GameplayAbilities/Private/GameplayAbilityTypes.cpp
+++ b/Engine/Source/Runtime/GameplayAbilities/Private/GameplayAbilityTypes.cpp
@@ -3,7 +3,6 @@
 #include "AbilitySystemPrivatePCH.h"
 #include "AbilitySystemComponent.h"
 #include "GameplayTagContainer.h"
-#include "GameplayTagsModule.h"
 #include "GameplayEffectTypes.h"
 #include "GameplayAbility.h"
 
diff --git a/Engine/Source/Runtime/GameplayAbilities/Private/GameplayCueInterface.cpp b/Engine/Source/Runtime/GameplayAbilities/Private/GameplayCueInterface.cpp
--- a/Engine/Source/Runtime/GameplayAbilities/Private/GameplayCueInterface.cpp
+++ b/Engine/Source/Runtime/GameplayAbilities/Private/GameplayCueInterface.cpp
@@ -3,7 +3,6 @@
 #include "AbilitySystemPrivatePCH.h"
 #include "AbilitySystemComponent.h"
 #include "GameplayCueInterface.h"
-#include "GameplayTagsModule.h"
 #include "GameplayCueSet.h"
 
 
@@ -57,8 +56,7 @@ void IGameplayCueInterface::HandleGameplayCue(AActor *Self, FGameplayTag Gamepla
 
 	// Look up a custom function for this gameplay tag. 
 	UClass* Class = Self->GetClass();
-	IGameplayTagsModule& GameplayTagsModule = IGameplayTagsModule::Get();
-	FGameplayTagContainer TagAndParentsContainer = GameplayTagsModule.GetGameplayTagsManager().RequestGameplayTagParents(GameplayCueTag);
+	FGameplayTagContainer TagAndParentsContainer = GameplayCueTag.GetGameplayTagParents();
 
 	Parameters.OriginalTag = GameplayCueTag;
 
diff --git a/Engine/Source/Runtime/GameplayAbilities/Private/GameplayCueManager.cpp b/Engine/Source/Runtime/GameplayAbilities/Private/GameplayCueManager.cpp
--- a/Engine/Source/Runtime/GameplayAbilities/Private/GameplayCueManager.cpp
+++ b/Engine/Source/Runtime/GameplayAbilities/Private/GameplayCueManager.cpp
@@ -5,7 +5,6 @@
 #include "GameplayCueInterface.h"
 #include "GameplayCueManager.h"
 #include "GameplayCueSet.h"
-#include "GameplayTagsModule.h"
 #include "GameplayCueNotify_Static.h"
 #include "AbilitySystemComponent.h"
 #include "Net/DataReplication.h"
@@ -458,7 +457,7 @@ static void SearchDynamicClassCues(const FName PropertyName, const TArray<FStrin
 {
 	// Iterate over all Dynamic Classes (nativized Blueprints). Search for ones with GameplayCueName tag.
 
-	IGameplayTagsModule& GameplayTagsModule = IGameplayTagsModule::Get();
+	UGameplayTagsManager& Manager = UGameplayTagsManager::Get();
 	TMap<FName, FDynamicClassStaticData>& DynamicClassMap = GetDynamicClassMap();
 	for (auto PairIter : DynamicClassMap)
 	{
@@ -479,7 +478,7 @@ static void SearchDynamicClassCues(const FName PropertyName, const TArray<FStrin
 
 			ABILITY_LOG(Log, TEXT("GameplayCueManager Found a Dynamic Class: %s / %s"), *FoundGameplayTag->ToString(), *ClassPath);
 
-			FGameplayTag  GameplayCueTag = GameplayTagsModule.GetGameplayTagsManager().RequestGameplayTag(*FoundGameplayTag, false);
+			FGameplayTag GameplayCueTag = Manager.RequestGameplayTag(*FoundGameplayTag, false);
 			if (GameplayCueTag.IsValid())
 			{
 				FStringAssetReference StringRef(ClassPath); // TODO: is there any translation needed?
@@ -588,7 +587,7 @@ static FAutoConsoleVariable CVarGameplyCueAddToGlobalSetDebug(TEXT("GameplayCue.
 
 void UGameplayCueManager::BuildCuesToAddToGlobalSet(const TArray<FAssetData>& AssetDataList, FName TagPropertyName, TArray<FGameplayCueReferencePair>& OutCuesToAdd, TArray<FStringAssetReference>& OutAssetsToLoad, FShouldLoadGCNotifyDelegate ShouldLoad)
 {
-	IGameplayTagsModule& GameplayTagsModule = IGameplayTagsModule::Get();
+	UGameplayTagsManager& Manager = UGameplayTagsManager::Get();
 
 	OutAssetsToLoad.Reserve(OutAssetsToLoad.Num() + AssetDataList.Num());
 
@@ -624,7 +623,7 @@ void UGameplayCueManager::BuildCuesToAddToGlobalSet(const TArray<FAssetData>& As
 
 			ABILITY_LOG(Log, TEXT("GameplayCueManager Found: %s / %s"), *FoundGameplayTag.ToString(), **GeneratedClassTag);
 
-			FGameplayTag  GameplayCueTag = GameplayTagsModule.GetGameplayTagsManager().RequestGameplayTag(FoundGameplayTag, false);
+			FGameplayTag  GameplayCueTag = Manager.RequestGameplayTag(FoundGameplayTag, false);
 			if (GameplayCueTag.IsValid())
 			{
 				// Add a new NotifyData entry to our flat list for this one
diff --git a/Engine/Source/Runtime/GameplayAbilities/Private/GameplayCueSet.cpp b/Engine/Source/Runtime/GameplayAbilities/Private/GameplayCueSet.cpp
--- a/Engine/Source/Runtime/GameplayAbilities/Private/GameplayCueSet.cpp
+++ b/Engine/Source/Runtime/GameplayAbilities/Private/GameplayCueSet.cpp
@@ -136,7 +136,7 @@ void UGameplayCueSet::Empty()
 
 void UGameplayCueSet::PrintCues() const
 {
-	FGameplayTagContainer AllGameplayCueTags = IGameplayTagsModule::Get().GetGameplayTagsManager().RequestGameplayTagChildren(BaseGameplayCueTag());
+	FGameplayTagContainer AllGameplayCueTags = UGameplayTagsManager::Get().RequestGameplayTagChildren(BaseGameplayCueTag());
 
 	for (FGameplayTag ThisGameplayCueTag : AllGameplayCueTags)
 	{
@@ -254,7 +254,7 @@ void UGameplayCueSet::BuildAccelerationMap_Internal()
 		GameplayCueDataMap.FindOrAdd(GameplayCueData[idx].GameplayCueTag) = idx;
 	}
 
-	FGameplayTagContainer AllGameplayCueTags = IGameplayTagsModule::Get().GetGameplayTagsManager().RequestGameplayTagChildren(BaseGameplayCueTag());
+	FGameplayTagContainer AllGameplayCueTags = UGameplayTagsManager::Get().RequestGameplayTagChildren(BaseGameplayCueTag());
 
 
 	// Create entries for children.
@@ -299,5 +299,5 @@ void UGameplayCueSet::BuildAccelerationMap_Internal()
 FGameplayTag UGameplayCueSet::BaseGameplayCueTag()
 {
 	// Note we should not cache this off as a static variable, since for new projects the GameplayCue tag will not be found until one is created.
-	return IGameplayTagsModule::Get().GetGameplayTagsManager().RequestGameplayTag(TEXT("GameplayCue"), false);
+	return FGameplayTag::RequestGameplayTag(TEXT("GameplayCue"), false);
 }
\ No newline at end of file
diff --git a/Engine/Source/Runtime/GameplayAbilities/Private/GameplayCueTranslator.cpp b/Engine/Source/Runtime/GameplayAbilities/Private/GameplayCueTranslator.cpp
--- a/Engine/Source/Runtime/GameplayAbilities/Private/GameplayCueTranslator.cpp
+++ b/Engine/Source/Runtime/GameplayAbilities/Private/GameplayCueTranslator.cpp
@@ -186,7 +186,7 @@ void FGameplayCueTranslationManager::BuildTagTranslationTable()
 	SCOPE_LOG_TIME_IN_SECONDS(*FString::Printf(TEXT("FGameplayCueTranslatorManager::BuildTagTranslationTables")), nullptr)
 #endif
 
-	TagManager = &IGameplayTagsModule::Get().GetGameplayTagsManager();
+	TagManager = &UGameplayTagsManager::Get();
 	check(TagManager);
 	
 	FGameplayTagContainer AllGameplayCueTags = TagManager->RequestGameplayTagChildren(UGameplayCueSet::BaseGameplayCueTag());
diff --git a/Engine/Source/Runtime/GameplayAbilities/Private/GameplayEffect.cpp b/Engine/Source/Runtime/GameplayAbilities/Private/GameplayEffect.cpp
--- a/Engine/Source/Runtime/GameplayAbilities/Private/GameplayEffect.cpp
+++ b/Engine/Source/Runtime/GameplayAbilities/Private/GameplayEffect.cpp
@@ -4,7 +4,6 @@
 #include "GameplayEffect.h"
 #include "GameplayEffectExtension.h"
 #include "GameplayEffectTypes.h"
-#include "GameplayTagsModule.h"
 #include "AbilitySystemComponent.h"
 #include "AbilitySystemInterface.h"
 #include "GameplayModMagnitudeCalculation.h"
@@ -61,8 +60,6 @@ void UGameplayEffect::GetOwnedGameplayTags(FGameplayTagContainer& TagContainer)
 
 void UGameplayEffect::GetTargetEffects(TArray<const UGameplayEffect*, TInlineAllocator<4> >& OutEffects) const
 {
-	OutEffects.Append(TargetEffects);
-
 	for ( TSubclassOf<UGameplayEffect> EffectClass : TargetEffectClasses )
 	{
 		if ( EffectClass )
@@ -95,15 +92,6 @@ void UGameplayEffect::PostLoad()
 	// We need to update when we first load to override values coming in from the superclass
 	// We also copy the tags from the old tag containers into the inheritable tag containers
 
-	InheritableGameplayEffectTags.Added.AppendTags(GameplayEffectTags);
-	GameplayEffectTags.RemoveAllTags();
-
-	InheritableOwnedTagsContainer.Added.AppendTags(OwnedTagsContainer);
-	OwnedTagsContainer.RemoveAllTags();
-
-	RemoveGameplayEffectsWithTags.Added.AppendTags(ClearTagsContainer);
-	ClearTagsContainer.RemoveAllTags();
-
 	UpdateInheritedTagProperties();
 
 	for (FGameplayAbilitySpecDef& Def : GrantedAbilities)
@@ -252,12 +240,12 @@ bool FAttributeBasedFloat::operator==(const FAttributeBasedFloat& Other) const
 		return false;
 	}
 	if (SourceTagFilter.Num() != Other.SourceTagFilter.Num() ||
-		!SourceTagFilter.MatchesAll(Other.SourceTagFilter, true))
+		!SourceTagFilter.HasAll(Other.SourceTagFilter))
 	{
 		return false;
 	}
 	if (TargetTagFilter.Num() != Other.TargetTagFilter.Num() ||
-		!TargetTagFilter.MatchesAll(Other.TargetTagFilter, true))
+		!TargetTagFilter.HasAll(Other.TargetTagFilter))
 	{
 		return false;
 	}
@@ -1672,7 +1660,7 @@ void FActiveGameplayEffectsContainer::ExecuteActiveEffectsFrom(FGameplayEffectSp
 	// Capture our own tags.
 	// TODO: We should only capture them if we need to. We may have snapshotted target tags (?) (in the case of dots with exotic setups?)
 
-	SpecToUse.CapturedTargetTags.GetActorTags().RemoveAllTags();
+	SpecToUse.CapturedTargetTags.GetActorTags().Reset();
 	Owner->GetOwnedGameplayTags(SpecToUse.CapturedTargetTags.GetActorTags());
 
 	SpecToUse.CalculateModifierMagnitudes();
@@ -2564,7 +2552,7 @@ FActiveGameplayEffect* FActiveGameplayEffectsContainer::ApplyGameplayEffectSpec(
 	UAbilitySystemGlobals::Get().GlobalPreGameplayEffectSpecApply(AppliedEffectSpec, Owner);
 
 	// Make sure our target's tags are collected, so we can properly filter infinite effects
-	AppliedEffectSpec.CapturedTargetTags.GetActorTags().RemoveAllTags();
+	AppliedEffectSpec.CapturedTargetTags.GetActorTags().Reset();
 	Owner->GetOwnedGameplayTags(AppliedEffectSpec.CapturedTargetTags.GetActorTags());
 
 	// Calc all of our modifier magnitudes now. Some may need to update later based on attributes changing, etc, but those should
@@ -3020,7 +3008,6 @@ void FActiveGameplayEffectsContainer::RemoveActiveGameplayEffectGrantedTagsAndMo
 	}
 
 	// Update gameplaytag count and broadcast delegate if we are at 0
-	IGameplayTagsModule& GameplayTagsModule = IGameplayTagsModule::Get();
 	Owner->UpdateTagMap(Effect.Spec.Def->InheritableOwnedTagsContainer.CombinedTags, -1);
 	Owner->UpdateTagMap(Effect.Spec.DynamicGrantedTags, -1);
 
@@ -3334,7 +3321,7 @@ bool FActiveGameplayEffectsContainer::HasApplicationImmunityToSpec(const FGamepl
 	}
 
 	// Quick map test
-	if (!AggregatedSourceTags->MatchesAny(ApplicationImmunityGameplayTagCountContainer.GetExplicitGameplayTags(), false))
+	if (!AggregatedSourceTags->HasAny(ApplicationImmunityGameplayTagCountContainer.GetExplicitGameplayTags()))
 	{
 		return false;
 	}
@@ -3579,31 +3566,6 @@ bool FActiveGameplayEffectsContainer::CanApplyAttributeModifiers(const UGameplay
 	return true;
 }
 
-// #deprecated
-TArray<float> FActiveGameplayEffectsContainer::GetActiveEffectsTimeRemaining(const FActiveGameplayEffectQuery Query) const
-{
-	SCOPE_CYCLE_COUNTER(STAT_GameplayEffectsGetActiveEffectsTimeRemaining);
-
-	float CurrentTime = GetWorldTime();
-
-	TArray<float>	ReturnList;
-
-	for (const FActiveGameplayEffect& Effect : this)
-	{
-		if (!Query.Matches(Effect))
-		{
-			continue;
-		}
-
-		float Elapsed = CurrentTime - Effect.StartWorldTime;
-		float Duration = Effect.GetDuration();
-
-		ReturnList.Add(Duration - Elapsed);
-	}
-
-	// Note: keep one return location to avoid copy operation.
-	return ReturnList;
-}
 TArray<float> FActiveGameplayEffectsContainer::GetActiveEffectsTimeRemaining(const FGameplayEffectQuery& Query) const
 {
 	SCOPE_CYCLE_COUNTER(STAT_GameplayEffectsGetActiveEffectsTimeRemaining);
@@ -3629,28 +3591,6 @@ TArray<float> FActiveGameplayEffectsContainer::GetActiveEffectsTimeRemaining(con
 	return ReturnList;
 }
 
-
-// #deprecated
-TArray<float> FActiveGameplayEffectsContainer::GetActiveEffectsDuration(const FActiveGameplayEffectQuery Query) const
-{
-	SCOPE_CYCLE_COUNTER(STAT_GameplayEffectsGetActiveEffectsDuration);
-
-	TArray<float>	ReturnList;
-
-	for (const FActiveGameplayEffect& Effect : this)
-	{
-		if (!Query.Matches(Effect))
-		{
-			continue;
-		}
-
-		ReturnList.Add(Effect.GetDuration());
-	}
-
-	// Note: keep one return location to avoid copy operation.
-	return ReturnList;
-}
-
 TArray<float> FActiveGameplayEffectsContainer::GetActiveEffectsDuration(const FGameplayEffectQuery& Query) const
 {
 	SCOPE_CYCLE_COUNTER(STAT_GameplayEffectsGetActiveEffectsDuration);
@@ -3696,26 +3636,6 @@ TArray<TPair<float,float>> FActiveGameplayEffectsContainer::GetActiveEffectsTime
 	return ReturnList;
 }
 
-// #deprecated
-TArray<FActiveGameplayEffectHandle> FActiveGameplayEffectsContainer::GetActiveEffects(const FActiveGameplayEffectQuery Query) const
-{
-	SCOPE_CYCLE_COUNTER(STAT_GameplayEffectsGetActiveEffects);
-
-	TArray<FActiveGameplayEffectHandle> ReturnList;
-
-	for (const FActiveGameplayEffect& Effect : this)
-	{
-		if (!Query.Matches(Effect))
-		{
-			continue;
-		}
-
-		ReturnList.Add(Effect.Handle);
-	}
-
-	return ReturnList;
-}
-
 TArray<FActiveGameplayEffectHandle> FActiveGameplayEffectsContainer::GetActiveEffects(const FGameplayEffectQuery& Query) const
 {
 	SCOPE_CYCLE_COUNTER(STAT_GameplayEffectsGetActiveEffects);
@@ -3809,26 +3729,6 @@ void FActiveGameplayEffectsContainer::ModifyActiveEffectStartTime(FActiveGamepla
 	}
 }
 
-// #deprecated, use FGameplayEffectQuery version
-int32 FActiveGameplayEffectsContainer::RemoveActiveEffects(const FActiveGameplayEffectQuery Query, int32 StacksToRemove)
-{
-	// Force a lock because the removals could cause other removals earlier in the array, so iterating backwards is not safe all by itself
-	GAMEPLAYEFFECT_SCOPE_LOCK();
-	int32 NumRemoved = 0;
-
-	// Manually iterating through in reverse because this is a removal operation
-	for (int32 idx=GetNumGameplayEffects()-1; idx >= 0; --idx)
-	{
-		const FActiveGameplayEffect& Effect = *GetActiveGameplayEffect(idx);
-		if (Effect.IsPendingRemove==false && Query.Matches(Effect))
-		{
-			InternalRemoveActiveGameplayEffect(idx, StacksToRemove, true);		
-			++NumRemoved;
-		}
-	}
-	return NumRemoved;
-}
-
 int32 FActiveGameplayEffectsContainer::RemoveActiveEffects(const FGameplayEffectQuery& Query, int32 StacksToRemove)
 {
 	// Force a lock because the removals could cause other removals earlier in the array, so iterating backwards is not safe all by itself
@@ -3848,24 +3748,6 @@ int32 FActiveGameplayEffectsContainer::RemoveActiveEffects(const FGameplayEffect
 	return NumRemoved;
 }
 
-int32 FActiveGameplayEffectsContainer::GetActiveEffectCount(const FActiveGameplayEffectQuery Query, bool bEnforceOnGoingCheck) const
-{
-	int32 Count = 0;
-
-	for (const FActiveGameplayEffect& Effect : this)
-	{
-		if (!Effect.bIsInhibited || !bEnforceOnGoingCheck)
-		{
-			if (Query.Matches(Effect))
-			{
-				Count += Effect.Spec.StackCount;
-			}
-		}
-	}
-
-	return Count;
-}
-
 int32 FActiveGameplayEffectsContainer::GetActiveEffectCount(const FGameplayEffectQuery& Query, bool bEnforceOnGoingCheck) const
 {
 	int32 Count = 0;
@@ -4259,7 +4141,7 @@ bool FGameplayEffectQuery::Matches(const FGameplayEffectSpec& Spec) const
 		// static to avoid memory allocations every time we do a query
 		check(IsInGameThread());
 		static FGameplayTagContainer TargetTags;
-		TargetTags.RemoveAllTags();
+		TargetTags.Reset();
 		if (Spec.Def->InheritableGameplayEffectTags.CombinedTags.Num() > 0)
 		{
 			TargetTags.AppendTags(Spec.Def->InheritableGameplayEffectTags.CombinedTags);
@@ -4285,7 +4167,7 @@ bool FGameplayEffectQuery::Matches(const FGameplayEffectSpec& Spec) const
 		// static to avoid memory allocations every time we do a query
 		check(IsInGameThread());
 		static FGameplayTagContainer GETags;
-		GETags.RemoveAllTags();
+		GETags.Reset();
 		if (Spec.Def->InheritableGameplayEffectTags.CombinedTags.Num() > 0)
 		{
 			GETags.AppendTags(Spec.Def->InheritableGameplayEffectTags.CombinedTags);
@@ -4448,100 +4330,6 @@ FGameplayEffectQuery FGameplayEffectQuery::MakeQuery_MatchNoSourceTags(const FGa
 	return OutQuery;
 }
 
-
-bool FActiveGameplayEffectQuery::Matches(const FActiveGameplayEffect& Effect) const
-{
-	// Anything in the ignore handle list is an immediate non-match
-	if (IgnoreHandles.Contains(Effect.Handle))
-	{
-		return false;
-	}
-
-	if (CustomMatch.IsBound())
-	{
-		return CustomMatch.Execute(Effect);
-	}
-
- 	// if we are looking for owning tags check them on the Granted Tags and Owned Tags Container
- 	if (OwningTagContainer)
- 	{
- 		if (!Effect.Spec.Def->InheritableOwnedTagsContainer.CombinedTags.MatchesAny(*OwningTagContainer, true) &&
- 			!Effect.Spec.DynamicGrantedTags.MatchesAny(*OwningTagContainer, false))
- 		{
- 			// if the GameplayEffect didn't match check the spec for tags that were added when this effect was created
- 			if (!Effect.Spec.CapturedSourceTags.GetSpecTags().MatchesAny(*OwningTagContainer, false))
- 			{
- 				return false;
- 			}
- 		}
- 	}	
- 	
- 	// if we are just looking for Tags on the Effect then look at the Gameplay Effect Tags
- 	if (EffectTagContainer)
- 	{
- 		if (!Effect.Spec.Def->InheritableGameplayEffectTags.CombinedTags.MatchesAny(*EffectTagContainer, true) &&
- 			!Effect.Spec.DynamicAssetTags.MatchesAny(*EffectTagContainer, false))
- 		{
- 			// this doesn't match our Tags so bail
- 			return false;
- 		}
- 	}
-
-	// if we are just looking for Tags on the Effect then look at the Gameplay Effect Tags
-	if (EffectTagContainer_Rejection)
-	{
-		if (Effect.Spec.Def->InheritableGameplayEffectTags.CombinedTags.MatchesAny(*EffectTagContainer_Rejection, true) ||
-			Effect.Spec.DynamicAssetTags.MatchesAny(*EffectTagContainer_Rejection, false))
-		{
-			// this matches our Rejection Tags so bail
-			return false;
-		}
-	}
-
-	// if we are looking for ModifyingAttribute go over each of the Spec Modifiers and check the Attributes
-	if (ModifyingAttribute.IsValid())
-	{
-		bool FailedModifyingAttributeCheck = true;
-
-		for(int32 ModIdx = 0; ModIdx < Effect.Spec.Modifiers.Num(); ++ModIdx)
-		{
-			const FGameplayModifierInfo& ModDef = Effect.Spec.Def->Modifiers[ModIdx];
-			const FModifierSpec& ModSpec = Effect.Spec.Modifiers[ModIdx];
-
-			if (ModDef.Attribute == ModifyingAttribute)
-			{
-				FailedModifyingAttributeCheck = false;
-				break;
-			}
-		}
-		if (FailedModifyingAttributeCheck)
-		{
-			return false;
-		}
-	}
-
-	// check source object
-	if (EffectSource)
-	{
-		if (Effect.Spec.GetEffectContext().GetSourceObject() != EffectSource)
-		{
-			return false;
-		}
-	}
-
-	// check definition
-	if (EffectDef)
-	{
-		if (Effect.Spec.Def != EffectDef)
-		{
-			return false;
-		}
-	}
-
-	// passed all the checks
-	return true;
-}
-
 bool FGameplayModifierInfo::operator==(const FGameplayModifierInfo& Other) const
 {
 	if (Attribute != Other.Attribute)
@@ -4559,20 +4347,20 @@ bool FGameplayModifierInfo::operator==(const FGameplayModifierInfo& Other) const
 		return false;
 	}
 
-	if (SourceTags.RequireTags.Num() != Other.SourceTags.RequireTags.Num() || !SourceTags.RequireTags.MatchesAll(Other.SourceTags.RequireTags, true))
+	if (SourceTags.RequireTags.Num() != Other.SourceTags.RequireTags.Num() || !SourceTags.RequireTags.HasAll(Other.SourceTags.RequireTags))
 	{
 		return false;
 	}
-	if (SourceTags.IgnoreTags.Num() != Other.SourceTags.IgnoreTags.Num() || !SourceTags.IgnoreTags.MatchesAll(Other.SourceTags.IgnoreTags, true))
+	if (SourceTags.IgnoreTags.Num() != Other.SourceTags.IgnoreTags.Num() || !SourceTags.IgnoreTags.HasAll(Other.SourceTags.IgnoreTags))
 	{
 		return false;
 	}
 
-	if (TargetTags.RequireTags.Num() != Other.TargetTags.RequireTags.Num() || !TargetTags.RequireTags.MatchesAll(Other.TargetTags.RequireTags, true))
+	if (TargetTags.RequireTags.Num() != Other.TargetTags.RequireTags.Num() || !TargetTags.RequireTags.HasAll(Other.TargetTags.RequireTags))
 	{
 		return false;
 	}
-	if (TargetTags.IgnoreTags.Num() != Other.TargetTags.IgnoreTags.Num() || !TargetTags.IgnoreTags.MatchesAll(Other.TargetTags.IgnoreTags, true))
+	if (TargetTags.IgnoreTags.Num() != Other.TargetTags.IgnoreTags.Num() || !TargetTags.IgnoreTags.HasAll(Other.TargetTags.IgnoreTags))
 	{
 		return false;
 	}
@@ -4588,14 +4376,14 @@ bool FGameplayModifierInfo::operator!=(const FGameplayModifierInfo& Other) const
 void FInheritedTagContainer::UpdateInheritedTagProperties(const FInheritedTagContainer* Parent)
 {
 	// Make sure we've got a fresh start
-	CombinedTags.RemoveAllTags();
+	CombinedTags.Reset();
 
 	// Re-add the Parent's tags except the one's we have removed
 	if (Parent)
 	{
 		for (auto Itr = Parent->CombinedTags.CreateConstIterator(); Itr; ++Itr)
 		{
-			if (!Removed.HasTag(*Itr, EGameplayTagMatchType::Explicit, EGameplayTagMatchType::IncludeParentTags))
+			if (!Itr->MatchesAny(Removed))
 			{
 				CombinedTags.AddTag(*Itr);
 			}
@@ -4607,7 +4395,7 @@ void FInheritedTagContainer::UpdateInheritedTagProperties(const FInheritedTagCon
 	{
 		// Remove trumps add for explicit matches but not for parent tags.
 		// This lets us remove all inherited tags starting with Foo but still add Foo.Bar
-		if (!Removed.HasTag(*Itr, EGameplayTagMatchType::Explicit, EGameplayTagMatchType::Explicit))
+		if (!Removed.HasTagExact(*Itr))
 		{
 			CombinedTags.AddTag(*Itr);
 		}
@@ -4618,8 +4406,8 @@ void FInheritedTagContainer::PostInitProperties()
 {
 	// we shouldn't inherit the added and removed tags from our parents
 	// make sure that these fields are clear
-	Added.RemoveAllTags();
-	Removed.RemoveAllTags();
+	Added.Reset();
+	Removed.Reset();
 }
 
 void FInheritedTagContainer::AddTag(const FGameplayTag& TagToAdd)
diff --git a/Engine/Source/Runtime/GameplayAbilities/Private/GameplayEffectAggregator.cpp b/Engine/Source/Runtime/GameplayAbilities/Private/GameplayEffectAggregator.cpp
--- a/Engine/Source/Runtime/GameplayAbilities/Private/GameplayEffectAggregator.cpp
+++ b/Engine/Source/Runtime/GameplayAbilities/Private/GameplayEffectAggregator.cpp
@@ -34,13 +34,13 @@ bool FAggregatorMod::Qualifies(const FAggregatorEvaluateParameters& Parameters)
 		if (!bSourceFilterMet)
 		{
 			const FGameplayTagContainer* SourceTags = HandleComponent->GetGameplayEffectSourceTagsFromHandle(ActiveHandle);
-			bSourceFilterMet = (SourceTags && SourceTags->MatchesAll(Parameters.AppliedSourceTagFilter, false));
+			bSourceFilterMet = (SourceTags && SourceTags->HasAll(Parameters.AppliedSourceTagFilter));
 		}
 
 		if (!bTargetFilterMet)
 		{
 			const FGameplayTagContainer* TargetTags = HandleComponent->GetGameplayEffectTargetTagsFromHandle(ActiveHandle);
-			bTargetFilterMet = (TargetTags && TargetTags->MatchesAll(Parameters.AppliedTargetTagFilter, false));
+			bTargetFilterMet = (TargetTags && TargetTags->HasAll(Parameters.AppliedTargetTagFilter));
 		}
 	}
 
diff --git a/Engine/Source/Runtime/GameplayAbilities/Private/GameplayEffectExtension_LifestealTest.cpp b/Engine/Source/Runtime/GameplayAbilities/Private/GameplayEffectExtension_LifestealTest.cpp
--- a/Engine/Source/Runtime/GameplayAbilities/Private/GameplayEffectExtension_LifestealTest.cpp
+++ b/Engine/Source/Runtime/GameplayAbilities/Private/GameplayEffectExtension_LifestealTest.cpp
@@ -3,7 +3,6 @@
 #include "AbilitySystemPrivatePCH.h"
 #include "GameplayEffectExtension_LifestealTest.h"
 #include "AbilitySystemTestAttributeSet.h"
-#include "GameplayTagsModule.h"
 #include "AbilitySystemComponent.h"
 
 UGameplayEffectExtension_LifestealTest::UGameplayEffectExtension_LifestealTest(const FObjectInitializer& ObjectInitializer)
@@ -20,8 +19,6 @@ void UGameplayEffectExtension_LifestealTest::PreGameplayEffectExecute(const FGam
 
 void UGameplayEffectExtension_LifestealTest::PostGameplayEffectExecute(const FGameplayModifierEvaluatedData &SelfData, const FGameplayEffectModCallbackData &Data) const
 {
-	IGameplayTagsModule& GameplayTagsModule = IGameplayTagsModule::Get();
-
 	float DamageDone = Data.EvaluatedData.Magnitude;
 	float LifestealPCT = SelfData.Magnitude;
 
diff --git a/Engine/Source/Runtime/GameplayAbilities/Private/GameplayEffectTypes.cpp b/Engine/Source/Runtime/GameplayAbilities/Private/GameplayEffectTypes.cpp
--- a/Engine/Source/Runtime/GameplayAbilities/Private/GameplayEffectTypes.cpp
+++ b/Engine/Source/Runtime/GameplayAbilities/Private/GameplayEffectTypes.cpp
@@ -2,7 +2,6 @@
 
 #include "AbilitySystemPrivatePCH.h"
 #include "AbilitySystemInterface.h"
-#include "GameplayTagsModule.h"
 #include "GameplayEffectTypes.h"
 #include "AbilitySystemComponent.h"
 #include "AbilitySystemGlobals.h"
@@ -363,7 +362,7 @@ void FGameplayTagCountContainer::Notify_StackCountChange(const FGameplayTag& Tag
 	// The purpose of this function is to let anyone listening on the EGameplayTagEventType::AnyCountChange event know that the 
 	// stack count of a GE that was backing this GE has changed. We do not update our internal map/count with this info, since that
 	// map only counts the number of GE/sources that are giving that tag.
-	FGameplayTagContainer TagAndParentsContainer = IGameplayTagsModule::Get().GetGameplayTagsManager().RequestGameplayTagParents(Tag);
+	FGameplayTagContainer TagAndParentsContainer = Tag.GetGameplayTagParents();
 	for (auto CompleteTagIt = TagAndParentsContainer.CreateConstIterator(); CompleteTagIt; ++CompleteTagIt)
 	{
 		const FGameplayTag& CurTag = *CompleteTagIt;
@@ -399,7 +398,7 @@ void FGameplayTagCountContainer::Reset()
 
 bool FGameplayTagCountContainer::UpdateTagMap_Internal(const FGameplayTag& Tag, int32 CountDelta)
 {
-	const bool bTagAlreadyExplicitlyExists = ExplicitTags.HasTag(Tag, EGameplayTagMatchType::Explicit, EGameplayTagMatchType::Explicit);
+	const bool bTagAlreadyExplicitlyExists = ExplicitTags.HasTagExact(Tag);
 
 	// Need special case handling to maintain the explicit tag list correctly, adding the tag to the list if it didn't previously exist and a
 	// positive delta comes in, and removing it from the list if it did exist and a negative delta comes in.
@@ -414,7 +413,7 @@ bool FGameplayTagCountContainer::UpdateTagMap_Internal(const FGameplayTag& Tag,
 		else
 		{
 			// only warn about tags that are in the container but will not be removed because they aren't explicitly in the container
-			if (ExplicitTags.HasTag(Tag, EGameplayTagMatchType::IncludeParentTags, EGameplayTagMatchType::Explicit))
+			if (ExplicitTags.HasTag(Tag))
 			{
 				ABILITY_LOG(Warning, TEXT("Attempted to remove tag: %s from tag count container, but it is not explicitly in the container!"), *Tag.ToString());
 			}
@@ -435,7 +434,7 @@ bool FGameplayTagCountContainer::UpdateTagMap_Internal(const FGameplayTag& Tag,
 	}
 
 	// Check if change delegates are required to fire for the tag or any of its parents based on the count change
-	FGameplayTagContainer TagAndParentsContainer = IGameplayTagsModule::Get().GetGameplayTagsManager().RequestGameplayTagParents(Tag);
+	FGameplayTagContainer TagAndParentsContainer = Tag.GetGameplayTagParents();
 	bool CreatedSignificantChange = false;
 	for (auto CompleteTagIt = TagAndParentsContainer.CreateConstIterator(); CompleteTagIt; ++CompleteTagIt)
 	{
@@ -478,8 +477,8 @@ bool FGameplayTagCountContainer::UpdateTagMap_Internal(const FGameplayTag& Tag,
 
 bool FGameplayTagRequirements::RequirementsMet(const FGameplayTagContainer& Container) const
 {
-	bool HasRequired = Container.MatchesAll(RequireTags, true);
-	bool HasIgnored = Container.MatchesAny(IgnoreTags, false);
+	bool HasRequired = Container.HasAll(RequireTags);
+	bool HasIgnored = Container.HasAny(IgnoreTags);
 
 	return HasRequired && !HasIgnored;
 }
@@ -538,7 +537,7 @@ const FGameplayTagContainer* FTagContainerAggregator::GetAggregatedTags() const
 	if (CacheIsValid == false)
 	{
 		CacheIsValid = true;
-		CachedAggregator.RemoveAllTags(CapturedActorTags.Num() + CapturedSpecTags.Num() + ScopedTags.Num());
+		CachedAggregator.Reset(CapturedActorTags.Num() + CapturedSpecTags.Num() + ScopedTags.Num());
 		CachedAggregator.AppendTags(CapturedActorTags);
 		CachedAggregator.AppendTags(CapturedSpecTags);
 		CachedAggregator.AppendTags(ScopedTags);
diff --git a/Engine/Source/Runtime/GameplayAbilities/Private/GameplayTagResponseTable.cpp b/Engine/Source/Runtime/GameplayAbilities/Private/GameplayTagResponseTable.cpp
--- a/Engine/Source/Runtime/GameplayAbilities/Private/GameplayTagResponseTable.cpp
+++ b/Engine/Source/Runtime/GameplayAbilities/Private/GameplayTagResponseTable.cpp
@@ -15,9 +15,12 @@ UGameplayTagReponseTable::UGameplayTagReponseTable(const FObjectInitializer& Obj
 	: Super(ObjectInitializer)
 {
 	// Make an empty tag query. We will actually fill the tag out prior to evaluating the query with ::MakeQuery
-	FGameplayTagContainer Container;
-	Container.AddTagFast(FGameplayTag());
-	Query = FGameplayEffectQuery::MakeQuery_MatchAllOwningTags(Container);
+
+	Query.OwningTagQuery = FGameplayTagQuery::BuildQuery(
+		FGameplayTagQueryExpression()
+		.AllTagsMatch()
+		.AddTag(FGameplayTag())
+	);
 }
 
 void UGameplayTagReponseTable::PostLoad()
diff --git a/Engine/Source/Runtime/GameplayAbilities/Public/Abilities/GameplayAbility.h b/Engine/Source/Runtime/GameplayAbilities/Public/Abilities/GameplayAbility.h
--- a/Engine/Source/Runtime/GameplayAbilities/Public/Abilities/GameplayAbility.h
+++ b/Engine/Source/Runtime/GameplayAbilities/Public/Abilities/GameplayAbility.h
@@ -207,9 +207,6 @@ public:
 	/** Returns an effect context, given a specified actor info */
 	virtual FGameplayEffectContextHandle MakeEffectContext(const FGameplayAbilitySpecHandle Handle, const FGameplayAbilityActorInfo *ActorInfo) const;
 
-	DEPRECATED(4.12, "GetEffectContext is deprecated, use MakeEffectContext")
-	virtual FGameplayEffectContextHandle GetEffectContext(const FGameplayAbilitySpecHandle Handle, const FGameplayAbilityActorInfo *ActorInfo) const;
-
 	virtual UWorld* GetWorld() const override
 	{
 		if (!IsInstantiated())
@@ -458,9 +455,6 @@ protected:
 	UFUNCTION(BlueprintCallable, Category = Ability, DisplayName="ApplyGameplayEffectToOwner")
 	FActiveGameplayEffectHandle BP_ApplyGameplayEffectToOwner(TSubclassOf<UGameplayEffect> GameplayEffectClass, int32 GameplayEffectLevel = 1, int32 Stacks = 1);
 
-	UFUNCTION(BlueprintCallable, Category = Ability, DisplayName="ApplyGameplayEffectToOwner", meta=(DeprecatedFunction, DeprecationMessage = "Use new ApplyGameplayEffectToOwner"))
-	FActiveGameplayEffectHandle K2_ApplyGameplayEffectToOwner(const UGameplayEffect* GameplayEffect, int32 GameplayEffectLevel = 1, int32 Stacks = 1);
-
 	/** Non blueprintcallable, safe to call on CDO/NonInstance abilities */
 	FActiveGameplayEffectHandle ApplyGameplayEffectToOwner(const FGameplayAbilitySpecHandle Handle, const FGameplayAbilityActorInfo* ActorInfo, const FGameplayAbilityActivationInfo ActivationInfo, const UGameplayEffect* GameplayEffect, float GameplayEffectLevel, int32 Stacks = 1) const;
 
@@ -476,9 +470,6 @@ protected:
 	UFUNCTION(BlueprintCallable, Category = Ability, DisplayName = "ApplyGameplayEffectToTarget")
 	TArray<FActiveGameplayEffectHandle> BP_ApplyGameplayEffectToTarget(FGameplayAbilityTargetDataHandle TargetData, TSubclassOf<UGameplayEffect> GameplayEffectClass, int32 GameplayEffectLevel = 1, int32 Stacks = 1);
 
-	UFUNCTION(BlueprintCallable, Category = Ability, DisplayName = "ApplyGameplayEffectToTarget", meta=(DeprecatedFunction, DeprecationMessage = "Use new ApplyGameplayEffectToTarget"))
-	TArray<FActiveGameplayEffectHandle> K2_ApplyGameplayEffectToTarget(FGameplayAbilityTargetDataHandle TargetData, const UGameplayEffect* GameplayEffect, int32 GameplayEffectLevel = 1, int32 Stacks = 1);
-
 	/** Non blueprintcallable, safe to call on CDO/NonInstance abilities */
 	TArray<FActiveGameplayEffectHandle> ApplyGameplayEffectToTarget(const FGameplayAbilitySpecHandle Handle, const FGameplayAbilityActorInfo* ActorInfo, const FGameplayAbilityActivationInfo ActivationInfo, const FGameplayAbilityTargetDataHandle& Target, TSubclassOf<UGameplayEffect> GameplayEffectClass, float GameplayEffectLevel, int32 Stacks = 1) const;
 
@@ -575,10 +566,6 @@ public:
 	UFUNCTION(BlueprintCallable, Category=Ability)
 	FGameplayEffectSpecHandle MakeOutgoingGameplayEffectSpec(TSubclassOf<UGameplayEffect> GameplayEffectClass, float Level=1.f) const;
 
-	/** Convenience method for abilities to get outgoing gameplay effect specs (for example, to pass on to projectiles to apply to whoever they hit) */
-	UFUNCTION(BlueprintCallable, Category=Ability, meta=(DeprecatedFunction, DeprecationMessage = "Use new MakeOutgoingGameplayEffectSpec"))
-	FGameplayEffectSpecHandle GetOutgoingGameplayEffectSpec(const UGameplayEffect* GameplayEffect, float Level=1.f) const;
-
 	FGameplayEffectSpecHandle MakeOutgoingGameplayEffectSpec(const FGameplayAbilitySpecHandle Handle, const FGameplayAbilityActorInfo* ActorInfo, const FGameplayAbilityActivationInfo ActivationInfo, TSubclassOf<UGameplayEffect> GameplayEffectClass, float Level = 1.f) const;
 
 	/** Add the Ability's tags to the given GameplayEffectSpec. This is likely to be overridden per project. */
@@ -683,10 +670,6 @@ protected:
 	UPROPERTY(EditDefaultsOnly, Category = Costs)
 	TSubclassOf<class UGameplayEffect> CostGameplayEffectClass;
 
-	/** Deprecated. Use CostGameplayEffectClass instead */
-	UPROPERTY(VisibleDefaultsOnly, Category=Deprecated)
-	class UGameplayEffect* CostGameplayEffect;
-
 	/** Triggers to determine if this ability should execute in response to an event */
 	UPROPERTY(EditDefaultsOnly, Category = Triggers)
 	TArray<FAbilityTriggerData> AbilityTriggers;
@@ -700,10 +683,6 @@ protected:
 	/** Deprecated? This GameplayEffect represents the cooldown. It will be applied when the ability is committed and the ability cannot be used again until it is expired. */
 	UPROPERTY(EditDefaultsOnly, Category = Cooldowns)
 	TSubclassOf<class UGameplayEffect> CooldownGameplayEffectClass;
-
-	/** Deprecated. Use CooldownGameplayEffectClass instead */
-	UPROPERTY(VisibleDefaultsOnly, Category=Deprecated)
-	class UGameplayEffect* CooldownGameplayEffect;
 	
 	// ----------------------------------------------------------------------------------------------------------------
 	//
@@ -824,11 +803,6 @@ public:
 	/** Retrieves the SourceObject associated with this ability. Callable on non instanced */
 	UObject* GetSourceObject(FGameplayAbilitySpecHandle Handle, const FGameplayAbilityActorInfo* ActorInfo) const;
 
-public:
-
-	// Temporary conversion code
-	virtual void ConvertDeprecatedGameplayEffectReferencesToBlueprintReferences(UGameplayEffect* OldGE, TSubclassOf<UGameplayEffect> NewGEClass);
-
 protected:
 
 	/** 
diff --git a/Engine/Source/Runtime/GameplayAbilities/Public/Abilities/GameplayAbility_Montage.h b/Engine/Source/Runtime/GameplayAbilities/Public/Abilities/GameplayAbility_Montage.h
--- a/Engine/Source/Runtime/GameplayAbilities/Public/Abilities/GameplayAbility_Montage.h
+++ b/Engine/Source/Runtime/GameplayAbilities/Public/Abilities/GameplayAbility_Montage.h
@@ -37,6 +37,4 @@ public:
 	void OnMontageEnded(UAnimMontage* Montage, bool bInterrupted, TWeakObjectPtr<UAbilitySystemComponent> AbilitySystemComponent, TArray<struct FActiveGameplayEffectHandle>	AppliedEffects);
 
 	void GetGameplayEffectsWhileAnimating(TArray<const UGameplayEffect*>& OutEffects) const;
-
-	virtual void ConvertDeprecatedGameplayEffectReferencesToBlueprintReferences(UGameplayEffect* OldGE, TSubclassOf<UGameplayEffect> NewGEClass) override;
 };
\ No newline at end of file
diff --git a/Engine/Source/Runtime/GameplayAbilities/Public/AbilitySystemComponent.h b/Engine/Source/Runtime/GameplayAbilities/Public/AbilitySystemComponent.h
--- a/Engine/Source/Runtime/GameplayAbilities/Public/AbilitySystemComponent.h
+++ b/Engine/Source/Runtime/GameplayAbilities/Public/AbilitySystemComponent.h
@@ -306,22 +306,10 @@ class GAMEPLAYABILITIES_API UAbilitySystemComponent : public UGameplayTasksCompo
 	UFUNCTION(BlueprintCallable, Category = GameplayEffects)
 	FGameplayEffectSpecHandle MakeOutgoingSpec(TSubclassOf<UGameplayEffect> GameplayEffectClass, float Level, FGameplayEffectContextHandle Context) const;
 
-	/** Get an outgoing GameplayEffectSpec that is ready to be applied to other things. */
-	UFUNCTION(BlueprintCallable, Category = GameplayEffects, meta=(DeprecatedFunction, DeprecationMessage = "Use new MakeOutgoingSpec"))
-	FGameplayEffectSpecHandle GetOutgoingSpec(const UGameplayEffect* GameplayEffect, float Level = 1.f) const;
-
-	/** Overloaded version that allows Context to be specified */
-	FGameplayEffectSpecHandle GetOutgoingSpec(const UGameplayEffect* GameplayEffect, float Level, FGameplayEffectContextHandle Context) const;
-
 	/** Create an EffectContext for the owner of this AbilitySystemComponent */
 	UFUNCTION(BlueprintCallable, Category = GameplayEffects)
 	FGameplayEffectContextHandle MakeEffectContext() const;
 
-	/** Deprecated. Use MakeEffectContext instead. */
-	DEPRECATED(4.12, "GetEffectContext is deprecated, use MakeEffectContext")
-	UFUNCTION(BlueprintCallable, Category = GameplayEffects, meta=(DeprecatedFunction, DeprecationMessage = "Use new MakeEffectContext"))
-	FGameplayEffectContextHandle GetEffectContext() const;
-
 	/**
 	 * Get the count of the specified source effect on the ability system component. For non-stacking effects, this is the sum of all active instances.
 	 * For stacking effects, this is the sum of all valid stack counts. If an instigator is specified, only effects from that instigator are counted.
@@ -335,8 +323,6 @@ class GAMEPLAYABILITIES_API UAbilitySystemComponent : public UGameplayTasksCompo
 	int32 GetGameplayEffectCount(TSubclassOf<UGameplayEffect> SourceGameplayEffect, UAbilitySystemComponent* OptionalInstigatorFilterComponent, bool bEnforceOnGoingCheck = true);
 
 	/** Returns the sum of StackCount of all gameplay effects that pass query */
-	DEPRECATED(4.9, "FActiveGameplayEffectQuery is deprecated, use version that takes FGameplayEffectQuery")
-	int32 GetAggregatedStackCount(const FActiveGameplayEffectQuery& Query);
 	int32 GetAggregatedStackCount(const FGameplayEffectQuery& Query);
 
 	/** This only exists so it can be hooked up to a multicast delegate */
@@ -463,14 +449,14 @@ class GAMEPLAYABILITIES_API UAbilitySystemComponent : public UGameplayTasksCompo
 		return GameplayTagCountContainer.HasMatchingGameplayTag(TagToCheck);
 	}
 
-	FORCEINLINE bool HasAllMatchingGameplayTags(const FGameplayTagContainer& TagContainer, bool bCountEmptyAsMatch = true) const override
+	FORCEINLINE bool HasAllMatchingGameplayTags(const FGameplayTagContainer& TagContainer) const override
 	{
-		return GameplayTagCountContainer.HasAllMatchingGameplayTags(TagContainer, bCountEmptyAsMatch);
+		return GameplayTagCountContainer.HasAllMatchingGameplayTags(TagContainer);
 	}
 
-	FORCEINLINE bool HasAnyMatchingGameplayTags(const FGameplayTagContainer& TagContainer, bool bCountEmptyAsMatch = true) const override
+	FORCEINLINE bool HasAnyMatchingGameplayTags(const FGameplayTagContainer& TagContainer) const override
 	{
-		return GameplayTagCountContainer.HasAnyMatchingGameplayTags(TagContainer, bCountEmptyAsMatch);
+		return GameplayTagCountContainer.HasAnyMatchingGameplayTags(TagContainer);
 	}
 
 	FORCEINLINE void GetOwnedGameplayTags(FGameplayTagContainer& TagContainer) const override
@@ -582,16 +568,10 @@ class GAMEPLAYABILITIES_API UAbilitySystemComponent : public UGameplayTasksCompo
 	UFUNCTION(BlueprintCallable, Category = GameplayEffects, meta=(DisplayName = "ApplyGameplayEffectToTarget"))
 	FActiveGameplayEffectHandle BP_ApplyGameplayEffectToTarget(TSubclassOf<UGameplayEffect> GameplayEffectClass, UAbilitySystemComponent *Target, float Level, FGameplayEffectContextHandle Context);
 
-	UFUNCTION(BlueprintCallable, Category = GameplayEffects, meta=(DisplayName = "ApplyGameplayEffectToTarget"), meta=(DeprecatedFunction, DeprecationMessage = "Use new ApplyGameplayEffectToSelf"))
-	FActiveGameplayEffectHandle K2_ApplyGameplayEffectToTarget(UGameplayEffect *GameplayEffect, UAbilitySystemComponent *Target, float Level, FGameplayEffectContextHandle Context);
-
 	FActiveGameplayEffectHandle ApplyGameplayEffectToTarget(UGameplayEffect *GameplayEffect, UAbilitySystemComponent *Target, float Level = UGameplayEffect::INVALID_LEVEL, FGameplayEffectContextHandle Context = FGameplayEffectContextHandle(), FPredictionKey PredictionKey = FPredictionKey());
 
 	UFUNCTION(BlueprintCallable, Category = GameplayEffects, meta=(DisplayName = "ApplyGameplayEffectToSelf"))
 	FActiveGameplayEffectHandle BP_ApplyGameplayEffectToSelf(TSubclassOf<UGameplayEffect> GameplayEffectClass, float Level, FGameplayEffectContextHandle EffectContext);
-
-	UFUNCTION(BlueprintCallable, Category = GameplayEffects, meta=(DisplayName = "ApplyGameplayEffectToSelf"), meta=(DeprecatedFunction, DeprecationMessage = "Use new ApplyGameplayEffectToSelf"))
-	FActiveGameplayEffectHandle K2_ApplyGameplayEffectToSelf(const UGameplayEffect *GameplayEffect, float Level, FGameplayEffectContextHandle EffectContext);
 	
 	FActiveGameplayEffectHandle ApplyGameplayEffectToSelf(const UGameplayEffect *GameplayEffect, float Level, const FGameplayEffectContextHandle& EffectContext, FPredictionKey PredictionKey = FPredictionKey());
 
@@ -629,18 +609,12 @@ class GAMEPLAYABILITIES_API UAbilitySystemComponent : public UGameplayTasksCompo
 	//		-Get duration of CD
 	//		-Get magnitude + duration of a movespeed buff
 
-	DEPRECATED(4.9, "FActiveGameplayEffectQuery is deprecated, use version that takes FGameplayEffectQuery")
-	TArray<float> GetActiveEffectsTimeRemaining(const FActiveGameplayEffectQuery Query) const;
 	TArray<float> GetActiveEffectsTimeRemaining(const FGameplayEffectQuery& Query) const;
 
-	DEPRECATED(4.9, "FActiveGameplayEffectQuery is deprecated, use version that takes FGameplayEffectQuery")
-	TArray<float> GetActiveEffectsDuration(const FActiveGameplayEffectQuery Query) const;
 	TArray<float> GetActiveEffectsDuration(const FGameplayEffectQuery& Query) const;
 
 	TArray<TPair<float,float>> GetActiveEffectsTimeRemainingAndDuration(const FGameplayEffectQuery& Query) const;
 
-	DEPRECATED(4.9, "FActiveGameplayEffectQuery is deprecated, use version that takes FGameplayEffectQuery")
-	TArray<FActiveGameplayEffectHandle> GetActiveEffects(const FActiveGameplayEffectQuery Query) const;
 	TArray<FActiveGameplayEffectHandle> GetActiveEffects(const FGameplayEffectQuery& Query) const;
 
 	/** This will give the world time that all effects matching this query will be finished. If multiple effects match, it returns the one that returns last.*/
@@ -676,8 +650,6 @@ class GAMEPLAYABILITIES_API UAbilitySystemComponent : public UGameplayTasksCompo
 	int32 RemoveActiveEffectsWithGrantedTags(FGameplayTagContainer Tags);
 
 	/** Removes all active effects that match given query. StacksToRemove=-1 will remove all stacks. */
-	DEPRECATED(4.9, "FActiveGameplayEffectQuery is deprecated, use version that takes FGameplayEffectQuery")
-	int32 RemoveActiveEffects(const FActiveGameplayEffectQuery Query, int32 StacksToRemove=-1);
 	int32 RemoveActiveEffects(const FGameplayEffectQuery& Query, int32 StacksToRemove = -1);
 
 	void OnRestackGameplayEffects();	
diff --git a/Engine/Source/Runtime/GameplayAbilities/Public/AbilitySystemGlobals.h b/Engine/Source/Runtime/GameplayAbilities/Public/AbilitySystemGlobals.h
--- a/Engine/Source/Runtime/GameplayAbilities/Public/AbilitySystemGlobals.h
+++ b/Engine/Source/Runtime/GameplayAbilities/Public/AbilitySystemGlobals.h
@@ -2,7 +2,7 @@
 
 #pragma once
 
-#include "GameplayTagsModule.h"
+#include "GameplayTags.h"
 #include "GameplayAbilitiesModule.h"
 #include "GameplayEffectTypes.h"
 #include "AbilitySystemGlobals.generated.h"
@@ -184,27 +184,27 @@ class GAMEPLAYABILITIES_API UAbilitySystemGlobals : public UObject
 	{
 		if (ActivateFailCooldownName != NAME_None)
 		{
-			ActivateFailCooldownTag = IGameplayTagsModule::RequestGameplayTag(ActivateFailCooldownName);
+			ActivateFailCooldownTag = FGameplayTag::RequestGameplayTag(ActivateFailCooldownName);
 		}
 
 		if (ActivateFailCostName != NAME_None)
 		{
-			ActivateFailCostTag = IGameplayTagsModule::RequestGameplayTag(ActivateFailCostName);
+			ActivateFailCostTag = FGameplayTag::RequestGameplayTag(ActivateFailCostName);
 		}
 
 		if (ActivateFailTagsBlockedName != NAME_None)
 		{
-			ActivateFailTagsBlockedTag = IGameplayTagsModule::RequestGameplayTag(ActivateFailTagsBlockedName);
+			ActivateFailTagsBlockedTag = FGameplayTag::RequestGameplayTag(ActivateFailTagsBlockedName);
 		}
 
 		if (ActivateFailTagsMissingName != NAME_None)
 		{
-			ActivateFailTagsMissingTag = IGameplayTagsModule::RequestGameplayTag(ActivateFailTagsMissingName);
+			ActivateFailTagsMissingTag = FGameplayTag::RequestGameplayTag(ActivateFailTagsMissingName);
 		}
 
 		if (ActivateFailNetworkingName != NAME_None)
 		{
-			ActivateFailNetworkingTag = IGameplayTagsModule::RequestGameplayTag(ActivateFailNetworkingName);
+			ActivateFailNetworkingTag = FGameplayTag::RequestGameplayTag(ActivateFailNetworkingName);
 		}
 	}
 
diff --git a/Engine/Source/Runtime/GameplayAbilities/Public/GameplayEffect.h b/Engine/Source/Runtime/GameplayAbilities/Public/GameplayEffect.h
--- a/Engine/Source/Runtime/GameplayAbilities/Public/GameplayEffect.h
+++ b/Engine/Source/Runtime/GameplayAbilities/Public/GameplayEffect.h
@@ -1523,18 +1523,12 @@ struct GAMEPLAYABILITIES_API FActiveGameplayEffectsContainer : public FFastArray
 
 	bool CanApplyAttributeModifiers(const UGameplayEffect *GameplayEffect, float Level, const FGameplayEffectContextHandle& EffectContext);
 	
-	DEPRECATED(4.9, "FActiveGameplayEffectQuery is deprecated, use FGameplayEffectQuery instead")
-	TArray<float> GetActiveEffectsTimeRemaining(const FActiveGameplayEffectQuery Query) const;
 	TArray<float> GetActiveEffectsTimeRemaining(const FGameplayEffectQuery& Query) const;
 
-	DEPRECATED(4.9, "FActiveGameplayEffectQuery is deprecated, use FGameplayEffectQuery instead")
-	TArray<float> GetActiveEffectsDuration(const FActiveGameplayEffectQuery Query) const;
 	TArray<float> GetActiveEffectsDuration(const FGameplayEffectQuery& Query) const;
 
 	TArray<TPair<float,float>> GetActiveEffectsTimeRemainingAndDuration(const FGameplayEffectQuery& Query) const;
 
-	DEPRECATED(4.9, "FActiveGameplayEffectQuery is deprecated, use FGameplayEffectQuery instead")
-	TArray<FActiveGameplayEffectHandle> GetActiveEffects(const FActiveGameplayEffectQuery Query) const;
 	TArray<FActiveGameplayEffectHandle> GetActiveEffects(const FGameplayEffectQuery& Query) const;
 
 	float GetActiveEffectsEndTime(const FGameplayEffectQuery& Query) const;
@@ -1545,8 +1539,6 @@ struct GAMEPLAYABILITIES_API FActiveGameplayEffectsContainer : public FFastArray
 
 	void ModifyActiveEffectStartTime(FActiveGameplayEffectHandle Handle, float StartTimeDiff);
 
-	DEPRECATED(4.9, "FActiveGameplayEffectQuery is deprecated, use FGameplayEffectQuery instead")
-	int32 RemoveActiveEffects(const FActiveGameplayEffectQuery Query, int32 StacksToRemove);
 	int32 RemoveActiveEffects(const FGameplayEffectQuery& Query, int32 StacksToRemove);
 
 	/**
@@ -1554,8 +1546,6 @@ struct GAMEPLAYABILITIES_API FActiveGameplayEffectsContainer : public FFastArray
 	 * 
 	 * @return Count of the effects matching the specified query
 	 */
-	DEPRECATED(4.9, "FActiveGameplayEffectQuery is deprecated, use FGameplayEffectQuery instead")
-	int32 GetActiveEffectCount(const FActiveGameplayEffectQuery Query, bool bEnforceOnGoingCheck = true) const;
 	int32 GetActiveEffectCount(const FGameplayEffectQuery& Query, bool bEnforceOnGoingCheck = true) const;
 
 	float GetServerWorldTime() const;
@@ -1837,10 +1827,6 @@ public:
 	UPROPERTY(EditDefaultsOnly, BlueprintReadOnly, Category = GameplayEffect,meta = (DisplayName = "Linked Gameplay Effects"))
 	TArray<TSubclassOf<UGameplayEffect>> TargetEffectClasses;
 
-	/** Deprecated. Use TargetEffectClasses instead */
-	UPROPERTY(VisibleDefaultsOnly, Category = Deprecated)
-	TArray<UGameplayEffect*> TargetEffects;
-
 	/** Effects to apply when a stacking effect "overflows" its stack count through another attempted application. Added whether the overflow application succeeds or not. */
 	UPROPERTY(EditDefaultsOnly, BlueprintReadOnly, Category=Overflow)
 	TArray<TSubclassOf<UGameplayEffect>> OverflowEffects;
@@ -1901,18 +1887,10 @@ public:
 	/** The GameplayEffect's Tags: tags the the GE *has* and DOES NOT give to the actor. */
 	UPROPERTY(EditDefaultsOnly, BlueprintReadOnly, Category = Tags, meta = (DisplayName = "GameplayEffectAssetTag"))
 	FInheritedTagContainer InheritableGameplayEffectTags;
-
-	/** The GameplayEffect's Tags: tags the the GE *has* and DOES NOT give to the actor. */
-	UPROPERTY(EditDefaultsOnly, BlueprintReadOnly, Category = Deprecated, meta=(DisplayName="GameplayEffectAssetTag"))
-	FGameplayTagContainer GameplayEffectTags;
 	
 	/** "These tags are applied to the actor I am applied to" */
 	UPROPERTY(EditDefaultsOnly, BlueprintReadOnly, Category = Tags, meta=(DisplayName="GrantedTags"))
 	FInheritedTagContainer InheritableOwnedTagsContainer;
-
-	/** "These tags are applied to the actor I am applied to" */
-	UPROPERTY(EditDefaultsOnly, BlueprintReadOnly, Category = Deprecated)
-	FGameplayTagContainer OwnedTagsContainer;
 	
 	/** Once Applied, these tags requirements are used to determined if the GameplayEffect is "on" or "off". A GameplayEffect can be off and do nothing, but still applied. */
 	UPROPERTY(EditDefaultsOnly, BlueprintReadOnly, Category = Tags)
@@ -1926,10 +1904,6 @@ public:
 	UPROPERTY(EditDefaultsOnly, BlueprintReadOnly, Category = Tags)
 	FInheritedTagContainer RemoveGameplayEffectsWithTags;
 
-	/** Deprecated. Use RemoveGameplayEffectsWithTags instead */
-	UPROPERTY(EditDefaultsOnly, BlueprintReadOnly, Category=Deprecated)
-	FGameplayTagContainer ClearTagsContainer;
-
 	/** Grants the owner immunity from these source tags.  */
 	UPROPERTY(EditDefaultsOnly, BlueprintReadOnly, Category = Immunity, meta = (DisplayName = "GrantedApplicationImmunityTags"))
 	FGameplayTagRequirements GrantedApplicationImmunityTags;
diff --git a/Engine/Source/Runtime/GameplayAbilities/Public/GameplayEffectTypes.h b/Engine/Source/Runtime/GameplayAbilities/Public/GameplayEffectTypes.h
--- a/Engine/Source/Runtime/GameplayAbilities/Public/GameplayEffectTypes.h
+++ b/Engine/Source/Runtime/GameplayAbilities/Public/GameplayEffectTypes.h
@@ -999,17 +999,16 @@ struct GAMEPLAYABILITIES_API FGameplayTagCountContainer
 	/**
 	 * Check if the count container has gameplay tags that matches against all of the specified tags (expands to include parents of asset tags)
 	 * 
-	 * @param TagContainer			Tag container to check for a match
-	 * @param bCountEmptyAsMatch	If true, the parameter tag container will count as matching, even if it's empty
+	 * @param TagContainer			Tag container to check for a match. If empty will return true
 	 * 
 	 * @return True if the count container matches all of the gameplay tags
 	 */
-	FORCEINLINE bool HasAllMatchingGameplayTags(const FGameplayTagContainer& TagContainer, bool bCountEmptyAsMatch = true) const
+	FORCEINLINE bool HasAllMatchingGameplayTags(const FGameplayTagContainer& TagContainer) const
 	{
 		// if the TagContainer count is 0 return bCountEmptyAsMatch;
 		if (TagContainer.Num() == 0)
 		{
-			return bCountEmptyAsMatch;
+			return true;
 		}
 
 		bool AllMatch = true;
@@ -1027,14 +1026,18 @@ struct GAMEPLAYABILITIES_API FGameplayTagCountContainer
 	/**
 	 * Check if the count container has gameplay tags that matches against any of the specified tags (expands to include parents of asset tags)
 	 * 
-	 * @param TagContainer			Tag container to check for a match
-	 * @param bCountEmptyAsMatch	If true, the parameter tag container will count as matching, even if it's empty
+	 * @param TagContainer			Tag container to check for a match. If empty will return false
 	 * 
 	 * @return True if the count container matches any of the gameplay tags
 	 */
-	FORCEINLINE bool HasAnyMatchingGameplayTags(const FGameplayTagContainer& TagContainer, bool bCountEmptyAsMatch = true) const
+	FORCEINLINE bool HasAnyMatchingGameplayTags(const FGameplayTagContainer& TagContainer) const
 	{
-		bool AnyMatch = (bCountEmptyAsMatch && TagContainer.Num() == 0);
+		if (TagContainer.Num() == 0)
+		{
+			return false;
+		}
+
+		bool AnyMatch = false;
 		for (const FGameplayTag& Tag : TagContainer)
 		{
 			if (GameplayTagCountMap.FindRef(Tag) > 0)
